[MaterialDesigner]
Move plugin

diff --git a/Engine/Plugins/Experimental/DynamicMaterial/Config/BaseDynamicMaterial.ini b/Engine/Plugins/Experimental/DynamicMaterial/Config/BaseDynamicMaterial.ini
new file mode 100644
--- /dev/null
+++ b/Engine/Plugins/Experimental/DynamicMaterial/Config/BaseDynamicMaterial.ini
@@ -0,0 +1,18 @@
+[CoreRedirects]
+
++ClassRedirects=(OldName="/Script/DynamicMaterial.DMMaterialValueTexture2D",NewName="/Script/DynamicMaterial.DMMaterialValueTexture")
++ClassRedirects=(OldName="/Script/DynamicMaterial.DMMaterialValueTextureCube",NewName="/Script/DynamicMaterial.DMMaterialValueTexture")
++ClassRedirects=(OldName="/Script/DynamicMaterial.DMMaterialValueVolumeTexture",NewName="/Script/DynamicMaterial.DMMaterialValueTexture")
++ClassRedirects=(OldName="/Script/DynamicMaterial.DynamicMaterialModelBase",NewName="/Script/DynamicMaterial.DynamicMaterialModel")
++ClassRedirects=(OldName="/Script/DynamicMaterialEditor.DynamicMaterialModel",NewName="/Script/DynamicMaterial.DynamicMaterialModel")
+
++EnumRedirects=(OldName="/Script/DynamicMaterial.EDMValueType",NewName="/Script/DynamicMaterial.EDMValueType",ValueChanges=(("VT_Texture_Any","VT_Texture"),("VT_Texture2D_Object","VT_Texture"),("VT_TextureCube_Object","VT_Texture"),("VT_VolumeTexture_Object","VT_Texture")))
+
++PackageRedirects=(OldName="/DynamicMaterial/MaterialFunctions/Effects/MF_DM_Effect_HSV.MF_DM_Effect_HSV",NewName="/DynamicMaterial/MaterialFunctions/Effects/Color/MF_DM_Effect_Base_HSV.MF_DM_Effect_Base_HSV")
++PackageRedirects=(OldName="/DynamicMaterial/MaterialFunctions/Effects/MF_DM_Effect_Base_HSV.MF_DM_Effect_Base_HSV",NewName="/DynamicMaterial/MaterialFunctions/Effects/Color/MF_DM_Effect_Base_HSV.MF_DM_Effect_Base_HSV")
++PackageRedirects=(OldName="/DynamicMaterial/MaterialFunctions/Effects/MF_DM_UVPanner.MF_DM_UVPanner",NewName="/DynamicMaterial/MaterialFunctions/Effects/UV/MF_DM_Effect_UV_Panner.MF_DM_Effect_UV_Panner")
++PackageRedirects=(OldName="/DynamicMaterial/MaterialFunctions/Effects/MF_DM_Effect_UV_Panner.MF_DM_Effect_UV_Panner",NewName="/DynamicMaterial/MaterialFunctions/Effects/UV/MF_DM_Effect_UV_Panner.MF_DM_Effect_UV_Panner")
++PackageRedirects=(OldName="/DynamicMaterial/MaterialFunctions/Effects/MF_DM_Effect_Base_Levels.MF_DM_Effect_Base_Levels",NewName="/DynamicMaterial/MaterialFunctions/Effects/Color/MF_DM_Effect_Base_Levels.MF_DM_Effect_Base_Levels")
++PackageRedirects=(OldName="/DynamicMaterial/MaterialFunctions/Effects/MF_DM_Effect_Mask_Level.MF_DM_Effect_Mask_Level",NewName="/DynamicMaterial/MaterialFunctions/Effects/Alpha/MF_DM_Effect_Mask_Level.MF_DM_Effect_Mask_Level")
++PackageRedirects=(OldName="/DynamicMaterial/MaterialFunctions/Effects/MF_DM_UV_Distort.MF_DM_UV_Distort",NewName="/DynamicMaterial/MaterialFunctions/Effects/UV/MF_DM_UV_Distort.MF_DM_UV_Distort")
++PackageRedirects=(OldName="/DynamicMaterial/MaterialFunctions/Effects/MF_DM_Effect_UV_Pixelate.MF_DM_Effect_UV_Pixelate",NewName="/DynamicMaterial/MaterialFunctions/Effects/UV/MF_DM_Effect_UV_Pixelate.MF_DM_Effect_UV_Pixelate")
diff --git a/Engine/Plugins/Experimental/DynamicMaterial/Source/DynamicMaterial/Private/Components/DMMaterialValue.cpp b/Engine/Plugins/Experimental/DynamicMaterial/Source/DynamicMaterial/Private/Components/DMMaterialValue.cpp
new file mode 100644
--- /dev/null
+++ b/Engine/Plugins/Experimental/DynamicMaterial/Source/DynamicMaterial/Private/Components/DMMaterialValue.cpp
@@ -0,0 +1,600 @@
+// Copyright Epic Games, Inc. All Rights Reserved.
+ 
+#include "Components/DMMaterialValue.h"
+#include "Components/DMMaterialParameter.h"
+#include "DMComponentPath.h"
+#include "Model/DynamicMaterialModel.h"
+#include "Materials/Material.h"
+#include "UObject/Package.h"
+
+#if WITH_EDITOR
+#include "AssetRegistry/AssetRegistryModule.h"
+#include "DMValueDefinition.h"
+#include "DynamicMaterialModule.h"
+#include "Factories/MaterialFactoryNew.h"
+#include "IDetailTreeNode.h"
+#include "IPropertyRowGenerator.h"
+#include "Model/IDMMaterialBuildStateInterface.h"
+#include "Model/IDMMaterialBuildUtilsInterface.h"
+#include "Model/IDynamicMaterialModelEditorOnlyDataInterface.h"
+#include "Modules/ModuleManager.h"
+#include "PropertyEditorModule.h"
+#endif // WITH_EDITOR
+
+#define LOCTEXT_NAMESPACE "DMMaterialValue"
+
+TMap<EDMValueType, TStrongObjectPtr<UClass>> UDMMaterialValue::TypeClasses = {};
+
+const FString UDMMaterialValue::ParameterPathToken = FString(TEXT("Parameter"));
+const TCHAR* UDMMaterialValue::ParameterNamePrefix = TEXT("VALUE_");
+
+#if WITH_EDITOR
+const FName UDMMaterialValue::ValueName = "Value";
+
+UDMMaterialValue* UDMMaterialValue::CreateMaterialValue(UDynamicMaterialModel* InMaterialModel, const FString& InName,
+	EDMValueType InType, bool bInLocal)
+{
+	check(InMaterialModel);
+ 
+	TStrongObjectPtr<UClass>* ValueClassPtr = TypeClasses.Find(InType);
+	check(ValueClassPtr);
+ 
+	UClass* ValueClass = (*ValueClassPtr).Get();
+	check(ValueClass);
+ 
+	UDMMaterialValue* NewValue = NewObject<UDMMaterialValue>(InMaterialModel, ValueClass, NAME_None, RF_Transactional);
+	check(NewValue);
+	NewValue->Type = InType;
+	NewValue->bLocal = bInLocal;
+ 
+	if (InName.IsEmpty() == false)
+	{
+		if (GUndo)
+		{
+			InMaterialModel->Modify();
+		}
+
+		NewValue->Parameter = InMaterialModel->CreateUniqueParameter(*InName);
+		NewValue->Parameter->SetParentComponent(NewValue);
+	}
+ 
+	return NewValue;
+}
+#endif // WITH_EDITOR
+
+UDMMaterialValue::UDMMaterialValue()
+	: UDMMaterialValue(EDMValueType::VT_None)
+{
+	Parameter = nullptr;
+
+#if WITH_EDITORONLY_DATA
+	EditableProperties.Add(ValueName);
+#endif
+}
+ 
+UDynamicMaterialModel* UDMMaterialValue::GetMaterialModel() const
+{
+	return Cast<UDynamicMaterialModel>(GetOuterSafe());
+}
+ 
+#if WITH_EDITOR
+FText UDMMaterialValue::GetTypeName() const
+{
+	return UDMValueDefinitionLibrary::GetValueDefinition(Type).GetDisplayName();
+}
+ 
+FText UDMMaterialValue::GetDescription() const
+{
+	static const FText Template = LOCTEXT("ValueDescriptionTemplate", "{0} ({1})");
+ 
+	return FText::Format(
+		Template,
+		Parameter ? FText::FromName(Parameter->GetParameterName()) : FText::AsNumber(FindIndex() + 1),
+		GetTypeName()
+	);
+}
+
+void UDMMaterialValue::PostCDOContruct()
+{
+	Super::PostCDOContruct();
+
+	if (Type != EDMValueType::VT_None)
+	{
+		TypeClasses.Emplace(Type, TStrongObjectPtr<UClass>(GetClass()));
+	}
+}
+
+void UDMMaterialValue::PostLoad()
+{
+	Super::PostLoad();
+
+	if (Parameter)
+	{
+		if (GUndo)
+		{
+			Parameter->Modify();
+		}
+
+		Parameter->SetParentComponent(this);
+	}
+}
+
+void UDMMaterialValue::PostEditImport()
+{
+	Super::PostEditImport();
+
+	if (Parameter)
+	{
+		if (GUndo)
+		{
+			Parameter->Modify();
+		}
+
+		Parameter->SetParentComponent(this);
+	}
+}
+
+void UDMMaterialValue::BeginDestroy()
+{
+	Super::BeginDestroy();
+	PreviewMaterial = nullptr;
+}
+#endif // WITH_EDITOR
+
+FName UDMMaterialValue::GetMaterialParameterName() const
+{
+	return Parameter ? FName(ParameterNamePrefix + Parameter->GetParameterName().ToString()) : GetFName();
+}
+
+UDMMaterialComponent* UDMMaterialValue::GetSubComponentByPath(FDMComponentPath& InPath, const FDMComponentPathSegment& InPathSegment) const
+{
+	if (InPathSegment.GetToken() == ParameterPathToken)
+	{
+		return Parameter;
+	}
+
+	return Super::GetSubComponentByPath(InPath, InPathSegment);
+}
+ 
+#if WITH_EDITOR
+bool UDMMaterialValue::SetParameterName(FName InBaseName)
+{
+	if (Parameter && Parameter->GetParameterName() == InBaseName)
+	{
+		return false;
+	}
+
+	if (!IsComponentValid())
+	{
+		return false;
+	}
+ 
+	UDynamicMaterialModel* MaterialModel = GetMaterialModel();
+	check(MaterialModel);
+
+	if (GUndo && IsValid(Parameter))
+	{
+		Parameter->Modify();
+		MaterialModel->Modify();
+	}
+ 
+	if (InBaseName.IsNone())
+	{
+		if (Parameter)
+		{
+			Parameter->SetParentComponent(nullptr);
+			MaterialModel->FreeParameter(Parameter);
+			Parameter = nullptr;
+		}
+ 
+		return true;
+	}
+ 
+	if (Parameter)
+	{
+		Parameter->RenameParameter(InBaseName);
+	}
+	else
+	{
+		Parameter = MaterialModel->CreateUniqueParameter(InBaseName);
+		Parameter->SetParentComponent(this);
+	}
+ 
+	return true;
+}
+
+void UDMMaterialValue::OnComponentRemoved()
+{
+	if (Parameter)
+	{
+		if (GUndo)
+		{
+			Parameter->Modify();
+		}
+
+		Parameter->SetComponentState(EDMComponentLifetimeState::Removed);
+	}
+ 
+	Super::OnComponentRemoved();
+}
+
+int32 UDMMaterialValue::FindIndex() const
+{
+	UDynamicMaterialModel* MaterialModel = GetMaterialModel();
+	check(MaterialModel);
+ 
+	const TArray<UDMMaterialValue*>& Values = MaterialModel->GetValues();
+ 
+	for (int32 ValueIdx = 0; ValueIdx < Values.Num(); ++ValueIdx)
+	{
+		if (Values[ValueIdx] == this)
+		{
+			return ValueIdx;
+		}
+	}
+ 
+	check(false);
+	return INDEX_NONE;
+}
+ 
+int32 UDMMaterialValue::FindIndexSafe() const
+{
+	UDynamicMaterialModel* MaterialModel = GetMaterialModel();
+	check(MaterialModel);
+ 
+	const TArray<UDMMaterialValue*>& Values = MaterialModel->GetValues();
+ 
+	for (int32 ValueIdx = 0; ValueIdx < Values.Num(); ++ValueIdx)
+	{
+		if (Values[ValueIdx] == this)
+		{
+			return ValueIdx;
+		}
+	}
+ 
+	return INDEX_NONE;
+}
+ 
+UMaterial* UDMMaterialValue::GetPreviewMaterial()
+{
+	if (!PreviewMaterial)
+	{
+		CreatePreviewMaterial();
+ 
+		if (PreviewMaterial)
+		{
+			MarkComponentDirty();
+		}
+	}
+ 
+	return PreviewMaterial;
+}
+#endif // WITH_EDITOR
+ 
+void UDMMaterialValue::SetMIDParameter(UMaterialInstanceDynamic* InMID) const
+{
+	checkNoEntry();
+}
+ 
+#if WITH_EDITOR
+void UDMMaterialValue::DoClean()
+{
+	if (!IsComponentValid())
+	{
+		return;
+	}
+
+	UpdatePreviewMaterial();
+
+	Super::DoClean();
+}
+
+void UDMMaterialValue::PostEditorDuplicate(UDynamicMaterialModel* InMaterialModel, UDMMaterialComponent* InParent)
+{
+	if (GetOuter() != InMaterialModel)
+	{
+		FName OldParameterName = NAME_None;
+
+		if (UDMMaterialParameter* Param = GetParameter())
+		{
+			// Reset this to null as it holds a copy of the parameter from the copied-from object.
+			// This will not be in the model's parameter list and will share the same name as the old parameter.
+			// Just null the reference and create a new parameter.
+			if (InMaterialModel->ConditionalFreeParameter(Param))
+			{
+				OldParameterName = Param->GetParameterName();
+				Parameter = nullptr;
+			}
+		}
+
+		Super::PostEditorDuplicate(InMaterialModel, InParent);
+
+		Rename(nullptr, InMaterialModel, UE::DynamicMaterial::RenameFlags);
+
+		if (OldParameterName != NAME_None)
+		{
+			SetParameterName(OldParameterName);
+		}
+	}
+	else
+	{
+		Super::PostEditorDuplicate(InMaterialModel, InParent);
+	}
+
+	PreviewMaterial = nullptr;
+}
+
+bool UDMMaterialValue::Modify(bool bInAlwaysMarkDirty)
+{
+	const bool bSaved = Super::Modify(bInAlwaysMarkDirty);
+
+	if (Parameter)
+	{
+		Parameter->Modify(bInAlwaysMarkDirty);
+	}
+
+	return bSaved;
+}
+
+void UDMMaterialValue::PostEditUndo()
+{
+	Super::PostEditUndo();
+
+	if (!IsComponentValid())
+	{
+		return;
+	}
+ 
+	MarkComponentDirty();
+ 
+	OnValueUpdated(/*bForceStructureUpdate*/ true); // Just in case - Undos are not meant to be quick and easy.
+}
+ 
+void UDMMaterialValue::PostEditChangeProperty(FPropertyChangedEvent& PropertyChangedEvent)
+{
+	Super::PostEditChangeProperty(PropertyChangedEvent);
+
+	if (!IsComponentValid())
+	{
+		return;
+	}
+
+	FName MemberPropertyName = PropertyChangedEvent.GetMemberPropertyName();
+
+	if (MemberPropertyName == NAME_None)
+	{
+		return;
+	}
+
+	for (const FName& EditableProperty : EditableProperties)
+	{
+		if (EditableProperty == MemberPropertyName)
+		{
+			OnValueUpdated(/* bForceStructureUpdate */ EditableProperty != ValueName);
+			return;
+		}
+	}
+}
+ 
+void UDMMaterialValue::EnsureDetailObjects()
+{
+	if (!IsComponentValid())
+	{
+		return;
+	}
+
+	if (PropertyRowGenerator.IsValid() && DetailTreeNode.IsValid() && PropertyHandle.IsValid())
+	{
+		return;
+	}
+ 
+	PropertyHandle.Reset();
+	DetailTreeNode.Reset();
+	PropertyRowGenerator.Reset();
+ 
+	FPropertyEditorModule& PropertyEditor = FModuleManager::Get().LoadModuleChecked<FPropertyEditorModule>(TEXT("PropertyEditor"));
+ 
+	FPropertyRowGeneratorArgs RowGeneratorArgs;
+	PropertyRowGenerator = PropertyEditor.CreatePropertyRowGenerator(RowGeneratorArgs);
+	PropertyRowGenerator->SetObjects({this});
+ 
+	for (const TSharedRef<IDetailTreeNode>& CategoryNode : PropertyRowGenerator->GetRootTreeNodes())
+	{
+		if (CategoryNode->GetNodeName() != TEXT("Material Designer"))
+		{
+			continue;
+		}
+ 
+		TArray<TSharedRef<IDetailTreeNode>> ChildNodes;
+		CategoryNode->GetChildren(ChildNodes);
+ 
+		for (const TSharedRef<IDetailTreeNode>& ChildNode : ChildNodes)
+		{
+			if (ChildNode->GetNodeType() != EDetailNodeType::Item)
+			{
+				continue;
+			}
+ 
+			if (ChildNode->GetNodeName() != ValueName)
+			{
+				continue;
+			}
+ 
+			DetailTreeNode = ChildNode;
+			PropertyHandle = ChildNode->CreatePropertyHandle();
+			return;
+		}
+	}
+}
+
+TSharedPtr<IDetailTreeNode> UDMMaterialValue::GetDetailTreeNode()
+{
+	EnsureDetailObjects();
+	return DetailTreeNode;
+}
+ 
+TSharedPtr<IPropertyHandle> UDMMaterialValue::GetPropertyHandle()
+{
+	EnsureDetailObjects();
+	return PropertyHandle;
+}
+#endif // WITH_EDITOR
+
+UDMMaterialValue::UDMMaterialValue(EDMValueType InType)
+	: Type(InType)
+	, bLocal(false)
+	, Parameter(nullptr)
+#if WITH_EDITORONLY_DATA
+	, PreviewMaterial(nullptr)
+#endif
+{
+#if WITH_EDITORONLY_DATA
+	EditableProperties.Add(ValueName);
+#endif
+}
+ 
+void UDMMaterialValue::OnValueUpdated(bool bInForceStructureUpdate)
+{
+	if (!IsComponentValid())
+	{
+		return;
+	}
+
+	if (FDMUpdateGuard::CanUpdate())
+	{
+		Update(bInForceStructureUpdate ? EDMUpdateType::Structure : EDMUpdateType::Value);
+	}
+}
+ 
+void UDMMaterialValue::Update(EDMUpdateType InUpdateType)
+{
+	if (!IsComponentValid())
+	{
+		return;
+	}
+
+#if WITH_EDITOR
+	if (HasComponentBeenRemoved())
+	{
+		return;
+	}
+ 
+	if (InUpdateType == EDMUpdateType::Structure)
+	{
+		MarkComponentDirty();
+	}
+
+	if (ParentComponent)
+	{
+		ParentComponent->Update(InUpdateType);
+	}
+#endif
+
+	Super::Update(InUpdateType);
+
+#if WITH_EDITOR
+	UDynamicMaterialModel* MaterialModel = GetMaterialModel();
+	check(MaterialModel);
+
+	if (IDynamicMaterialModelEditorOnlyDataInterface* ModelEditorOnlyData = MaterialModel->GetEditorOnlyData())
+	{
+		ModelEditorOnlyData->OnValueUpdated(this, InUpdateType);
+	}
+#endif
+}
+
+#if WITH_EDITOR
+void UDMMaterialValue::CreatePreviewMaterial()
+{
+	if (!IsComponentValid())
+	{
+		return;
+	}
+
+	if (FDynamicMaterialModule::IsMaterialExportEnabled() == false)
+	{
+		UMaterialFactoryNew* MaterialFactory = NewObject<UMaterialFactoryNew>();
+		check(MaterialFactory);
+ 
+		PreviewMaterial = (UMaterial*)MaterialFactory->FactoryCreateNew(
+			UMaterial::StaticClass(),
+			GetTransientPackage(),
+			NAME_None,
+			RF_Transient,
+			nullptr,
+			GWarn
+		);
+ 
+		PreviewMaterial->bIsPreviewMaterial = true;
+	}
+	else
+	{
+		FString MaterialBaseName = GetName() + "-" + FGuid::NewGuid().ToString();
+		const FString FullName = "/Game/DynamicMaterials/" + MaterialBaseName;
+		UPackage* Package = CreatePackage(*FullName);
+ 
+		UMaterialFactoryNew* MaterialFactory = NewObject<UMaterialFactoryNew>();
+		check(MaterialFactory);
+ 
+		PreviewMaterial = Cast<UMaterial>(MaterialFactory->FactoryCreateNew(
+			UMaterial::StaticClass(),
+			Package,
+			*MaterialBaseName,
+			RF_Standalone | RF_Public,
+			nullptr,
+			GWarn
+		));
+ 
+		FAssetRegistryModule::AssetCreated(PreviewMaterial);
+		Package->FullyLoad();
+	}
+}
+ 
+void UDMMaterialValue::UpdatePreviewMaterial()
+{
+	if (!IsComponentValid())
+	{
+		return;
+	}
+
+	if (!PreviewMaterial)
+	{
+		CreatePreviewMaterial();
+ 
+		if (!PreviewMaterial)
+		{
+			return;
+		}
+	}
+ 
+	UDynamicMaterialModel* MaterialModel = GetMaterialModel();
+
+	if (!MaterialModel)
+	{
+		return;
+	}
+
+	IDynamicMaterialModelEditorOnlyDataInterface* ModelEditorOnlyData = MaterialModel->GetEditorOnlyData();
+
+	if (!ModelEditorOnlyData)
+	{
+		return;
+	}
+
+	TSharedRef<IDMMaterialBuildStateInterface> BuildState = ModelEditorOnlyData->CreateBuildStateInterface(PreviewMaterial);
+
+	UE_LOG(LogDynamicMaterial, Display, TEXT("Building Material Designer Value Preview (%s)..."), *GetName());
+
+	GenerateExpression(BuildState);
+	UMaterialExpression* ValueExpression = BuildState->GetLastValueExpression(this);
+ 
+	BuildState->GetBuildUtils().UpdatePreviewMaterial(ValueExpression, 0, FDMMaterialStageConnectorChannel::WHOLE_CHANNEL, 32);
+}
+ 
+int32 UDMMaterialValue::GetInnateMaskOutput(int32 OutputChannels) const
+{
+	return INDEX_NONE;
+}
+#endif // WITH_EDITOR
+
+#undef LOCTEXT_NAMESPACE
diff --git a/Engine/Plugins/Experimental/DynamicMaterial/Source/DynamicMaterial/Private/Components/DMTextureUV.cpp b/Engine/Plugins/Experimental/DynamicMaterial/Source/DynamicMaterial/Private/Components/DMTextureUV.cpp
new file mode 100644
--- /dev/null
+++ b/Engine/Plugins/Experimental/DynamicMaterial/Source/DynamicMaterial/Private/Components/DMTextureUV.cpp
@@ -0,0 +1,739 @@
+// Copyright Epic Games, Inc. All Rights Reserved.
+
+#include "Components/DMTextureUV.h"
+#include "Components/DMMaterialParameter.h"
+#include "DMComponentPath.h"
+#include "Materials/MaterialInstanceDynamic.h"
+#include "Model/DynamicMaterialModel.h"
+#include "Serialization/CustomVersion.h"
+
+#if WITH_EDITOR
+#include "IDetailTreeNode.h"
+#include "IPropertyRowGenerator.h"
+#include "Model/IDynamicMaterialModelEditorOnlyDataInterface.h"
+#include "Modules/ModuleManager.h"
+#include "PropertyEditorModule.h"
+#include "PropertyHandle.h"
+#endif
+
+namespace UE::DynamicMaterial::Private
+{
+	TMap<int32, FName> BaseParameterNames = {
+		{ParamID::PivotX,   FName(TEXT("TextureUV_PivotX"))},
+		{ParamID::PivotY,   FName(TEXT("TextureUV_PivotY"))},
+		{ParamID::ScaleX,   FName(TEXT("TextureUV_ScaleX"))},
+		{ParamID::ScaleY,   FName(TEXT("TextureUV_ScaleY"))},
+		{ParamID::Rotation, FName(TEXT("TextureUV_Rotation"))},
+		{ParamID::OffsetX,  FName(TEXT("TextureUV_OffsetX"))},
+		{ParamID::OffsetY,  FName(TEXT("TextureUV_OffsetY"))},
+	};
+}
+
+enum class EDMTextureUVVersion : int32
+{
+	Initial_Pre_20221102 = 0,
+	Version_22021102 = 1,
+	LatestVersion = Version_22021102
+};
+
+const FGuid UDMTextureUV::GUID(0xFCF57AFB, 0x50764284, 0xB9A9E659, 0xFFA02D33);
+FCustomVersionRegistration GRegisterDMTextureUVVersion(UDMTextureUV::GUID, static_cast<int32>(EDMTextureUVVersion::LatestVersion), TEXT("DMTextureUV"));
+
+const FString UDMTextureUV::OffsetXPathToken  = FString(TEXT("OffsetX"));
+const FString UDMTextureUV::OffsetYPathToken  = FString(TEXT("OffsetY"));
+const FString UDMTextureUV::PivotXPathToken   = FString(TEXT("PivotX"));
+const FString UDMTextureUV::PivotYPathToken   = FString(TEXT("PivotY"));
+const FString UDMTextureUV::RotationPathToken = FString(TEXT("Rotation"));
+const FString UDMTextureUV::ScaleXPathToken   = FString(TEXT("ScaleX"));
+const FString UDMTextureUV::ScaleYPathToken   = FString(TEXT("ScaleY"));
+
+#if WITH_EDITOR
+const FName UDMTextureUV::NAME_UVSource   = GET_MEMBER_NAME_CHECKED(UDMTextureUV, UVSource);
+const FName UDMTextureUV::NAME_Offset     = GET_MEMBER_NAME_CHECKED(UDMTextureUV, Offset);
+const FName UDMTextureUV::NAME_Pivot      = GET_MEMBER_NAME_CHECKED(UDMTextureUV, Pivot);
+const FName UDMTextureUV::NAME_Rotation   = GET_MEMBER_NAME_CHECKED(UDMTextureUV, Rotation);
+const FName UDMTextureUV::NAME_Scale      = GET_MEMBER_NAME_CHECKED(UDMTextureUV, Scale);
+const FName UDMTextureUV::NAME_bMirrorOnX = GET_MEMBER_NAME_CHECKED(UDMTextureUV, bMirrorOnX);
+const FName UDMTextureUV::NAME_bMirrorOnY = GET_MEMBER_NAME_CHECKED(UDMTextureUV, bMirrorOnY);
+
+const TMap<FName, bool> UDMTextureUV::TextureProperties = {
+	{NAME_UVSource,   false},
+	{NAME_Offset,     true},
+	{NAME_Pivot,      true},
+	{NAME_Rotation,   true},
+	{NAME_Scale,      true},
+	{NAME_bMirrorOnX, false},
+	{NAME_bMirrorOnY, false}
+};
+#endif
+
+UDMTextureUV::UDMTextureUV()
+{
+#if WITH_EDITOR
+	EditableProperties.Add(NAME_Offset);
+	EditableProperties.Add(NAME_Pivot);
+	EditableProperties.Add(NAME_Rotation);
+	EditableProperties.Add(NAME_Scale);
+	EditableProperties.Add(NAME_bMirrorOnX);
+	EditableProperties.Add(NAME_bMirrorOnY);
+#endif
+}
+
+#if WITH_EDITOR
+void UDMTextureUV::SetUVSource(EDMUVSource InUVSource)
+{
+	if (!IsComponentValid())
+	{
+		return;
+	}
+
+	if (UVSource == InUVSource)
+	{
+		return;
+	}
+
+	UVSource = InUVSource;
+
+	Update(EDMUpdateType::Structure);
+}
+#endif
+
+void UDMTextureUV::SetOffset(const FVector2D& InOffset)
+{
+	if (!IsComponentValid())
+	{
+		return;
+	}
+
+	if (FMath::IsNearlyEqual(Offset.X, InOffset.X)
+		&& FMath::IsNearlyEqual(Offset.Y, InOffset.Y))
+	{
+		return;
+	}
+
+	Offset = InOffset;
+
+	Update(EDMUpdateType::Value);
+}
+
+void UDMTextureUV::SetPivot(const FVector2D& InPivot)
+{
+	if (!IsComponentValid())
+	{
+		return;
+	}
+
+	if (FMath::IsNearlyEqual(Pivot.X, InPivot.X)
+		&& FMath::IsNearlyEqual(Pivot.Y, InPivot.Y))
+	{
+		return;
+	}
+
+	Pivot = InPivot;
+
+	Update(EDMUpdateType::Value);
+}
+
+void UDMTextureUV::SetRotation(float InRotation)
+{
+	if (!IsComponentValid())
+	{
+		return;
+	}
+
+	if (FMath::IsNearlyEqual(Rotation, InRotation))
+	{
+		return;
+	}
+
+	Rotation = InRotation;
+
+	Update(EDMUpdateType::Value);
+}
+
+void UDMTextureUV::SetScale(const FVector2D& InScale)
+{
+	if (!IsComponentValid())
+	{
+		return;
+	}
+
+	if (FMath::IsNearlyEqual(Scale.X, InScale.X)
+		&& FMath::IsNearlyEqual(Scale.Y, InScale.Y))
+	{
+		return;
+	}
+
+	Scale = InScale;
+
+	Update(EDMUpdateType::Value);
+}
+
+#if WITH_EDITOR
+void UDMTextureUV::SetMirrorOnX(bool bInMirrorOnX)
+{
+	if (!IsComponentValid())
+	{
+		return;
+	}
+
+	if (bMirrorOnX == bInMirrorOnX)
+	{
+		return;
+	}
+
+	bMirrorOnX = bInMirrorOnX;
+
+	Update(EDMUpdateType::Structure);
+}
+
+void UDMTextureUV::SetMirrorOnY(bool bInMirrorOnY)
+{
+	if (!IsComponentValid())
+	{
+		return;
+	}
+
+	if (bMirrorOnY == bInMirrorOnY)
+	{
+		return;
+	}
+
+	bMirrorOnY = bInMirrorOnY;
+
+	Update(EDMUpdateType::Structure);
+}
+
+TSharedPtr<IDetailTreeNode> UDMTextureUV::GetDetailTreeNode(FName InProperty)
+{
+	EnsureDetailObjects();
+
+	return DetailTreeNodes.FindChecked(InProperty);
+}
+
+TSharedPtr<IPropertyHandle> UDMTextureUV::GetPropertyHandle(FName InProperty)
+{
+	EnsureDetailObjects();
+
+	return PropertyHandles.FindChecked(InProperty);
+}
+
+TArray<UDMMaterialParameter*> UDMTextureUV::GetParameters() const
+{
+	TArray<UDMMaterialParameter*> Parameters;
+	Parameters.Reserve(MaterialParameters.Num());
+
+	for (const TPair<int32, TObjectPtr<UDMMaterialParameter>>& Pair : MaterialParameters)
+	{
+		Parameters.Add(Pair.Value.Get());
+	}
+
+	return Parameters;
+}
+
+void UDMTextureUV::PostEditorDuplicate(UDynamicMaterialModel* InMaterialModel, UDMMaterialComponent* InParent)
+{
+	Super::PostEditorDuplicate(InMaterialModel, InParent);
+
+	if (GetOuter() != InMaterialModel)
+	{
+		Rename(nullptr, InMaterialModel, UE::DynamicMaterial::RenameFlags);
+	}
+
+	// Reset this map as it holds copies of the parameters from the copied-from object.
+	// They will not be in this model's parameter list and will share the same name as the old parameters.
+	// Just empty the list and create new parameters.
+	for (TMap<int32, TObjectPtr<UDMMaterialParameter>>::TIterator It(MaterialParameters); It; ++It)
+	{
+		UDMMaterialParameter* Parameter = It->Value.Get();
+
+		if (!Parameter || InMaterialModel->ConditionalFreeParameter(Parameter))
+		{
+			It.RemoveCurrent();
+		}
+	}
+
+	MaterialParameters.Empty();
+
+	// Create new parameters.
+	CreateParameterNames();
+}
+#endif
+
+void UDMTextureUV::SetMIDParameters(UMaterialInstanceDynamic* InMID)
+{
+	if (!IsComponentValid())
+	{
+		return;
+	}
+
+	check(InMID);
+	check(MaterialParameters.IsEmpty() == false);
+
+	using namespace UE::DynamicMaterial;
+
+	auto UpdateMID = [InMID](FName InParamName, float InValue)
+	{
+		if (FMath::IsNearlyEqual(InValue, InMID->K2_GetScalarParameterValue(InParamName)) == false)
+		{
+			InMID->SetScalarParameterValue(InParamName, InValue);
+		}
+	};
+
+	UpdateMID(MaterialParameters[ParamID::PivotX]->GetParameterName(), GetPivot().X);
+	UpdateMID(MaterialParameters[ParamID::PivotY]->GetParameterName(), GetPivot().Y);
+	UpdateMID(MaterialParameters[ParamID::ScaleX]->GetParameterName(), GetScale().X);
+	UpdateMID(MaterialParameters[ParamID::ScaleY]->GetParameterName(), GetScale().Y);
+	UpdateMID(MaterialParameters[ParamID::Rotation]->GetParameterName(), GetRotation());
+	UpdateMID(MaterialParameters[ParamID::OffsetX]->GetParameterName(), GetOffset().X);
+	UpdateMID(MaterialParameters[ParamID::OffsetY]->GetParameterName(), GetOffset().Y);
+}
+
+#if WITH_EDITOR
+bool UDMTextureUV::Modify(bool bInAlwaysMarkDirty)
+{
+	const bool bSaved = Super::Modify(bInAlwaysMarkDirty);
+
+	for (const TPair<int32, TObjectPtr<UDMMaterialParameter>>& Pair : MaterialParameters)
+	{
+		Pair.Value->Modify(bInAlwaysMarkDirty);
+	}
+
+	return bSaved;
+}
+
+void UDMTextureUV::EnsureDetailObjects()
+{
+	if (!IsComponentValid())
+	{
+		return;
+	}
+
+	bool bHasValidDetailObjects = true;
+
+	if (!PropertyRowGenerator.IsValid() || (DetailTreeNodes.Num() != TextureProperties.Num() && PropertyHandles.Num() != TextureProperties.Num()))
+	{
+		bHasValidDetailObjects = false;
+	}
+	else
+	{
+		for (const TPair<FName, TSharedPtr<IDetailTreeNode>>& DetailTreeNode : DetailTreeNodes)
+		{
+			if (DetailTreeNode.Value.IsValid() == false)
+			{
+				bHasValidDetailObjects = false;
+				break;
+			}
+		}
+
+		if (bHasValidDetailObjects)
+		{
+			for (const TPair<FName, TSharedPtr<IPropertyHandle>>& PropertyHandle : PropertyHandles)
+			{
+				if (PropertyHandle.Value.IsValid() == false)
+				{
+					bHasValidDetailObjects = false;
+					break;
+				}
+			}
+		}
+	}
+
+	if (bHasValidDetailObjects)
+	{
+		return;
+	}
+
+	PropertyHandles.Reset();
+	DetailTreeNodes.Reset();
+	PropertyRowGenerator.Reset();
+
+	FPropertyEditorModule& PropertyEditor = FModuleManager::Get().LoadModuleChecked<FPropertyEditorModule>(TEXT("PropertyEditor"));
+
+	FPropertyRowGeneratorArgs RowGeneratorArgs;
+	PropertyRowGenerator = PropertyEditor.CreatePropertyRowGenerator(RowGeneratorArgs);
+	PropertyRowGenerator->SetObjects({this});
+
+	for (const TSharedRef<IDetailTreeNode>& CategoryNode : PropertyRowGenerator->GetRootTreeNodes())
+	{
+		if (CategoryNode->GetNodeName() != TEXT("Material Designer"))
+		{
+			continue;
+		}
+
+		TArray<TSharedRef<IDetailTreeNode>> ChildNodes;
+		CategoryNode->GetChildren(ChildNodes);
+
+		for (const TSharedRef<IDetailTreeNode>& ChildNode : ChildNodes)
+		{
+			if (ChildNode->GetNodeType() != EDetailNodeType::Item)
+			{
+				continue;
+			}
+			
+			if (TextureProperties.Contains(ChildNode->GetNodeName()) == false)
+			{
+				continue;
+			}
+
+			DetailTreeNodes.Emplace(ChildNode->GetNodeName(), ChildNode);
+			PropertyHandles.Emplace(ChildNode->GetNodeName(), ChildNode->CreatePropertyHandle());
+		}
+
+		return;
+	}
+}
+#endif
+
+void UDMTextureUV::Update(EDMUpdateType InUpdateType)
+{
+	if (!IsComponentValid())
+	{
+		return;
+	}
+
+#if WITH_EDITOR
+	if (HasComponentBeenRemoved())
+	{
+		return;
+	}
+
+	if (InUpdateType == EDMUpdateType::Structure)
+	{
+		MarkComponentDirty();
+	}
+
+	if (ParentComponent)
+	{
+		ParentComponent->Update(InUpdateType);
+	}
+#endif
+
+	Super::Update(InUpdateType);
+
+#if WITH_EDITOR
+	UDynamicMaterialModel* MaterialModel = GetMaterialModel();
+	check(MaterialModel);
+
+	if (IDynamicMaterialModelEditorOnlyDataInterface* ModelEditorOnlyData = MaterialModel->GetEditorOnlyData())
+	{
+		ModelEditorOnlyData->OnTextureUVUpdated(this);
+	}
+#endif
+}
+
+UDynamicMaterialModel* UDMTextureUV::GetMaterialModel() const
+{
+	return Cast<UDynamicMaterialModel>(GetOuterSafe());
+}
+
+UDMMaterialComponent* UDMTextureUV::GetSubComponentByPath(FDMComponentPath& InPath, const FDMComponentPathSegment& InPathSegment) const
+{
+	auto GetParameter = [this, &InPath, &InPathSegment](int32 InParamId) -> UDMMaterialComponent*
+		{
+			if (const TObjectPtr<UDMMaterialParameter>* ParameterPtr = MaterialParameters.Find(InParamId))
+			{
+				return *ParameterPtr;
+			}
+
+			return UDMMaterialLinkedComponent::GetSubComponentByPath(InPath, InPathSegment);
+		};
+
+	using namespace UE::DynamicMaterial;
+
+	if (InPathSegment.GetToken() == OffsetXPathToken)
+	{
+		return GetParameter(ParamID::OffsetX);
+	}
+
+	if (InPathSegment.GetToken() == OffsetYPathToken)
+	{
+		return GetParameter(ParamID::OffsetY);
+	}
+
+	if (InPathSegment.GetToken() == PivotXPathToken)
+	{
+		return GetParameter(ParamID::PivotX);
+	}
+
+	if (InPathSegment.GetToken() == PivotYPathToken)
+	{
+		return GetParameter(ParamID::PivotY);
+	}
+
+	if (InPathSegment.GetToken() == RotationPathToken)
+	{
+		return GetParameter(ParamID::Rotation);
+	}
+
+	if (InPathSegment.GetToken() == ScaleXPathToken)
+	{
+		return GetParameter(ParamID::ScaleX);
+	}
+
+	if (InPathSegment.GetToken() == ScaleYPathToken)
+	{
+		return GetParameter(ParamID::ScaleY);
+	}
+
+	return Super::GetSubComponentByPath(InPath, InPathSegment);
+}
+
+#if WITH_EDITOR
+void UDMTextureUV::GetComponentPathInternal(TArray<FString>& OutChildComponentPathComponents) const
+{
+	using namespace UE::DynamicMaterial::Private;
+
+	// Replace parameter object names with the base parameter name
+	if (OutChildComponentPathComponents.IsEmpty() == false)
+	{
+		for (const TPair<int32, TObjectPtr<UDMMaterialParameter>>& MaterialParameter : MaterialParameters)
+		{
+			if (OutChildComponentPathComponents.Last() == MaterialParameter.Value->GetComponentPathComponent())
+			{
+				OutChildComponentPathComponents.Last() = BaseParameterNames[MaterialParameter.Key].ToString();
+				break;
+			}
+		}
+	}
+
+	Super::GetComponentPathInternal(OutChildComponentPathComponents);
+}
+
+void UDMTextureUV::CreateParameterNames()
+{
+	if (!IsComponentValid())
+	{
+		return;
+	}
+
+	UDynamicMaterialModel* MaterialModel = GetMaterialModel();
+	check(MaterialModel);
+
+	using namespace UE::DynamicMaterial;
+
+	auto CreateParam = [this, MaterialModel](int32 ParamId)
+	{
+		if (!MaterialParameters.Contains(ParamId))
+		{
+			using namespace UE::DynamicMaterial::Private;
+
+			MaterialParameters.Add(ParamId, MaterialModel->CreateUniqueParameter(BaseParameterNames[ParamId]));
+			MaterialParameters[ParamId]->SetParentComponent(this);
+		}
+	};
+
+	CreateParam(ParamID::PivotX);
+	CreateParam(ParamID::PivotY);
+	CreateParam(ParamID::ScaleX);
+	CreateParam(ParamID::ScaleY);
+	CreateParam(ParamID::Rotation);
+	CreateParam(ParamID::OffsetX);
+	CreateParam(ParamID::OffsetY);
+}
+
+void UDMTextureUV::RemoveParameterNames()
+{
+	if (!IsComponentValid())
+	{
+		return;
+	}
+
+	UDynamicMaterialModel* MaterialModel = GetMaterialModel();
+	check(MaterialModel);
+
+#if WITH_EDITOR
+	if (GUndo)
+	{
+		MaterialModel->Modify();
+	}
+#endif
+
+	for (const TPair<int32, TObjectPtr<UDMMaterialParameter>>& Pair : MaterialParameters)
+	{
+#if WITH_EDITOR
+		if (GUndo)
+		{
+			Pair.Value->Modify();
+		}
+#endif
+
+		MaterialModel->FreeParameter(Pair.Value);
+	}
+}
+
+void UDMTextureUV::OnComponentAdded()
+{
+	if (!IsComponentValid())
+	{
+		return;
+	}
+
+	CreateParameterNames();
+	
+	Super::OnComponentAdded();
+}
+
+void UDMTextureUV::OnComponentRemoved()
+{
+	RemoveParameterNames();
+
+	Super::OnComponentRemoved();
+}
+
+void UDMTextureUV::PostEditChangeProperty(struct FPropertyChangedEvent& PropertyChangedEvent)
+{
+	Super::PostEditChangeProperty(PropertyChangedEvent);
+
+	if (!IsComponentValid())
+	{
+		return;
+	}
+
+	if (!PropertyChangedEvent.MemberProperty)
+	{
+		return;
+	}
+
+	if (PropertyChangedEvent.MemberProperty->GetFName() == NAME_Offset
+		|| PropertyChangedEvent.MemberProperty->GetFName() == NAME_Pivot
+		|| PropertyChangedEvent.MemberProperty->GetFName() == NAME_Rotation
+		|| PropertyChangedEvent.MemberProperty->GetFName() == NAME_Scale)
+	{
+		Update(EDMUpdateType::Value);
+	}
+	else if (PropertyChangedEvent.MemberProperty->GetFName() == NAME_UVSource
+		|| PropertyChangedEvent.MemberProperty->GetFName() == NAME_bMirrorOnX
+		|| PropertyChangedEvent.MemberProperty->GetFName() == NAME_bMirrorOnY)
+	{
+		Update(EDMUpdateType::Structure);
+	}
+}
+
+void UDMTextureUV::PreEditUndo()
+{
+	Super::PreEditUndo();
+
+	UVSource_PreUndo = UVSource;
+	bMirrorOnX_PreUndo = bMirrorOnX;
+	bMirrorOnY_PreUndo = bMirrorOnY;
+}
+
+void UDMTextureUV::PostEditUndo()
+{
+	Super::PostEditUndo();
+
+	if (UVSource != UVSource_PreUndo
+		|| bMirrorOnX != bMirrorOnX_PreUndo
+		|| bMirrorOnY != bMirrorOnY_PreUndo)
+	{
+		Update(EDMUpdateType::Structure);
+	}
+	else
+	{
+		Update(EDMUpdateType::Value);
+	}
+}
+
+void UDMTextureUV::PostLoad()
+{
+	Super::PostLoad();
+
+	if (!IsComponentValid())
+	{
+		return;
+	}
+
+	if (MaterialParameters.IsEmpty())
+	{
+		CreateParameterNames();
+	}
+
+	for (const TPair<int32, TObjectPtr<UDMMaterialParameter>>& Pair : MaterialParameters)
+	{
+		Pair.Value->SetParentComponent(this);
+	}
+
+	/*
+	 * @TODO GetLinkerCustomVersion() isn't used here to trigger these updates because it always returns 
+	 * the latest version regardless of what was saved to the archive.
+	 * Inside the function, it is unable to find a Loader and thus fails in this way.
+	 */
+
+	if (bNeedsPostLoadStructureUpdate)
+	{
+		Update(EDMUpdateType::Structure);
+	}
+	else if (bNeedsPostLoadValueUpdate)
+	{
+		Update(EDMUpdateType::Value);
+	}
+
+	bNeedsPostLoadStructureUpdate = false;
+	bNeedsPostLoadValueUpdate = false;
+}
+
+void UDMTextureUV::PostEditImport()
+{
+	Super::PostEditImport();
+
+	if (!IsComponentValid())
+	{
+		return;
+	}
+
+	if (MaterialParameters.IsEmpty())
+	{
+		CreateParameterNames();
+	}
+
+	for (const TPair<int32, TObjectPtr<UDMMaterialParameter>>& Pair : MaterialParameters)
+	{
+		Pair.Value->SetParentComponent(this);
+	}
+}
+
+UDMTextureUV* UDMTextureUV::CreateTextureUV(UObject* InOuter)
+{
+	UDMTextureUV* NewTextureUV = NewObject<UDMTextureUV>(InOuter, NAME_None, RF_Transactional);
+	check(NewTextureUV);
+
+	return NewTextureUV;
+}
+#endif
+
+void UDMTextureUV::Serialize(FArchive& Ar)
+{
+	Ar.UsingCustomVersion(UDMTextureUV::GUID);
+
+	Super::Serialize(Ar);
+
+	// @See UDMTextureUV::PostLoad
+
+	int32 TextureUVVersion = Ar.CustomVer(UDMTextureUV::GUID);
+
+	while (TextureUVVersion != static_cast<int32>(EDMTextureUVVersion::LatestVersion))
+	{
+		switch (TextureUVVersion)
+		{
+			case INDEX_NONE:
+			case static_cast<int32>(EDMTextureUVVersion::Initial_Pre_20221102):
+				Offset.X *= -1;
+				Rotation *= -360.f;
+				Scale = FVector2D(1.f, 1.f) / Scale;
+				TextureUVVersion = static_cast<int32>(EDMTextureUVVersion::LatestVersion);
+#if WITH_EDITORONLY_DATA
+				bNeedsPostLoadValueUpdate = true;
+#endif
+				break;
+
+			case static_cast<int32>(EDMTextureUVVersion::LatestVersion):
+				// Do nothing
+				break;
+
+			default:
+				TextureUVVersion = static_cast<int32>(EDMTextureUVVersion::LatestVersion);
+				break;
+		}
+	}
+}
diff --git a/Engine/Plugins/Experimental/DynamicMaterial/Source/DynamicMaterial/Private/Components/MaterialValues/DMMaterialValueTexture.cpp b/Engine/Plugins/Experimental/DynamicMaterial/Source/DynamicMaterial/Private/Components/MaterialValues/DMMaterialValueTexture.cpp
new file mode 100644
--- /dev/null
+++ b/Engine/Plugins/Experimental/DynamicMaterial/Source/DynamicMaterial/Private/Components/MaterialValues/DMMaterialValueTexture.cpp
@@ -0,0 +1,223 @@
+// Copyright Epic Games, Inc. All Rights Reserved.
+ 
+#include "Components/MaterialValues/DMMaterialValueTexture.h"
+#include "Engine/Texture.h"
+#include "Materials/MaterialInstanceDynamic.h"
+
+#if WITH_EDITOR
+#include "DMDefs.h"
+#include "Engine/Texture2D.h"
+#include "Engine/TextureCube.h"
+#include "Engine/VolumeTexture.h"
+#include "Materials/MaterialExpression.h"
+#include "Materials/MaterialExpressionTextureObjectParameter.h"
+#include "Model/IDMMaterialBuildStateInterface.h"
+#include "Model/IDMMaterialBuildUtilsInterface.h"
+#include "RenderUtils.h"
+#endif
+
+#if WITH_EDITOR
+namespace UE::DynamicMaterial::Private
+{
+	bool HasAlpha(UTexture* InTexture)
+	{
+		if (!IsValid(InTexture))
+		{
+			return false;
+		}
+
+		if (InTexture->CompressionNoAlpha)
+		{
+			return false;
+		}
+
+		EPixelFormatChannelFlags ValidTextureChannels = EPixelFormatChannelFlags::None;
+
+		if (UTexture2D* Texture2D = Cast<UTexture2D>(InTexture))
+		{
+			ValidTextureChannels = GetPixelFormatValidChannels(Texture2D->GetPixelFormat());
+		}
+		else if (UTextureCube* TextureCube = Cast<UTextureCube>(InTexture))
+		{
+			ValidTextureChannels = GetPixelFormatValidChannels(TextureCube->GetPixelFormat());
+		}
+		else if (UVolumeTexture* VolumeTexture = Cast<UVolumeTexture>(InTexture))
+		{
+			ValidTextureChannels = GetPixelFormatValidChannels(VolumeTexture->GetPixelFormat());
+		}
+		else
+		{
+			return false;
+		}
+
+		return EnumHasAnyFlags(ValidTextureChannels, EPixelFormatChannelFlags::A);
+	}
+}
+
+FDMGetDefaultRGBTexture UDMMaterialValueTexture::GetDefaultRGBTexture;
+#endif
+
+UDMMaterialValueTexture::UDMMaterialValueTexture()
+	: UDMMaterialValue(EDMValueType::VT_Texture)
+#if WITH_EDITORONLY_DATA
+	, OldValue(nullptr)
+#endif
+{
+#if WITH_EDITOR
+	ResetDefaultValue();
+	Value = DefaultValue;
+#else
+	Value = nullptr;
+#endif
+}
+
+#if WITH_EDITOR
+void UDMMaterialValueTexture::GenerateExpression(const TSharedRef<IDMMaterialBuildStateInterface>& InBuildState) const
+{
+	if (!IsComponentValid())
+	{
+		return;
+	}
+
+	if (InBuildState->HasValue(this))
+	{
+		return;
+	}
+
+	UMaterialExpressionTextureObjectParameter* NewExpression = InBuildState->GetBuildUtils().CreateExpressionParameter<UMaterialExpressionTextureObjectParameter>(GetMaterialParameterName(), UE_DM_NodeComment_Default, Value);
+	check(NewExpression);
+
+	InBuildState->AddValueExpressions(this, {NewExpression});
+}
+ 
+UDMMaterialValueTexture* UDMMaterialValueTexture::CreateMaterialValueTexture(UObject* Outer, UTexture* InTexture)
+{
+	check(InTexture);
+ 
+	UDMMaterialValueTexture* TextureValue = NewObject<UDMMaterialValueTexture>(Outer, NAME_None, RF_Transactional);
+	TextureValue->SetValue(InTexture);
+	return TextureValue;
+}
+ 
+void UDMMaterialValueTexture::PreEditChange(FProperty* PropertyAboutToChange)
+{
+	Super::PreEditChange(PropertyAboutToChange);
+ 
+	if (!IsComponentValid())
+	{
+		return;
+	}
+
+	if (PropertyAboutToChange->GetFName() == ValueName)
+	{
+		OldValue = GetValue();
+	}
+}
+
+bool UDMMaterialValueTexture::IsDefaultValue() const
+{
+	return Value == DefaultValue;
+}
+
+void UDMMaterialValueTexture::PostEditChangeProperty(FPropertyChangedEvent& PropertyChangedEvent)
+{
+	// Skip parent class because we need to do extra logic.
+	Super::Super::PostEditChangeProperty(PropertyChangedEvent);
+
+	if (!IsComponentValid())
+	{
+		return;
+	}
+
+	FName MemberPropertyName = PropertyChangedEvent.GetMemberPropertyName();
+
+	if (MemberPropertyName == NAME_None)
+	{
+		return;
+	}
+
+	for (const FName& EditableProperty : EditableProperties)
+	{
+		if (EditableProperty == MemberPropertyName)
+		{
+			if (EditableProperty == ValueName)
+			{
+				const TextureCompressionSettings CurrentCompression = OldValue ? OldValue->CompressionSettings.GetValue() : TextureCompressionSettings::TC_MAX;
+				const TextureCompressionSettings NewCompression = Value ? Value->CompressionSettings.GetValue() : TextureCompressionSettings::TC_MAX;
+
+				OnValueUpdated(CurrentCompression != NewCompression);
+				return;
+			}
+
+			OnValueUpdated(/* bForceStructureUpdate */ true);
+			return;
+		}
+	}
+}
+#endif
+
+void UDMMaterialValueTexture::PostLoad()
+{
+	Super::PostLoad();
+
+	Type = EDMValueType::VT_Texture;
+}
+
+#if WITH_EDITOR
+void UDMMaterialValueTexture::ApplyDefaultValue()
+{
+	SetValue(DefaultValue);
+}
+
+void UDMMaterialValueTexture::ResetDefaultValue()
+{
+	DefaultValue = nullptr;
+
+	if (GetDefaultRGBTexture.IsBound())
+	{
+		DefaultValue = GetDefaultRGBTexture.Execute();
+	}
+}
+
+void UDMMaterialValueTexture::SetDefaultValue(UTexture* InDefaultValue)
+{
+	DefaultValue = InDefaultValue;
+}
+
+bool UDMMaterialValueTexture::HasAlpha() const
+{
+	return UE::DynamicMaterial::Private::HasAlpha(Value);
+}
+#endif
+
+void UDMMaterialValueTexture::SetValue(UTexture* InValue)
+{
+	if (!IsComponentValid())
+	{
+		return;
+	}
+
+	if (Value == InValue)
+	{
+		return;
+	}
+
+	const TextureCompressionSettings CurrentCompression = Value ? Value->CompressionSettings.GetValue() : TextureCompressionSettings::TC_MAX;
+	const TextureCompressionSettings NewCompression = InValue ? InValue->CompressionSettings.GetValue() : TextureCompressionSettings::TC_MAX;
+
+	Value = InValue;
+
+	OnValueUpdated(CurrentCompression != NewCompression);
+}
+
+void UDMMaterialValueTexture::SetMIDParameter(UMaterialInstanceDynamic* InMID) const
+{
+	if (!IsComponentValid())
+	{
+		return;
+	}
+
+	check(InMID);
+
+	InMID->SetTextureParameterValue(GetMaterialParameterName(), Value);
+}
diff --git a/Engine/Plugins/Experimental/DynamicMaterial/Source/DynamicMaterial/Private/DMDefs.cpp b/Engine/Plugins/Experimental/DynamicMaterial/Source/DynamicMaterial/Private/DMDefs.cpp
new file mode 100644
--- /dev/null
+++ b/Engine/Plugins/Experimental/DynamicMaterial/Source/DynamicMaterial/Private/DMDefs.cpp
@@ -0,0 +1,40 @@
+// Copyright Epic Games, Inc. All Rights Reserved.
+
+#include "DMDefs.h"
+#include "Misc/Paths.h"
+
+int32 FDMUpdateGuard::GuardCount = 0;
+uint32 FDMInitializationGuard::GuardCount = 0;
+
+FString UE::DynamicMaterial::CreateNodeComment(const ANSICHAR* InFile, int InLine, const ANSICHAR* InFunction, const FString* InComment)
+{
+	FString File = FPaths::GetCleanFilename(ANSI_TO_TCHAR(InFile)).RightChop(2);
+
+	if (InComment)
+	{
+		return FString::Printf(TEXT("%s[%d]: %hs: %s"), *File, InLine, InFunction, **InComment);
+	}
+	else
+	{
+		return FString::Printf(TEXT("%s[%d]: %hs"), *File, InLine, InFunction);
+	}
+}
+
+bool FDMInitializationGuard::IsInitializing()
+{
+	return GuardCount > 0;
+}
+
+FDMInitializationGuard::FDMInitializationGuard()
+{
+	// Used the struct name to make it clear it's a static variable.
+	++FDMInitializationGuard::GuardCount;
+}
+
+FDMInitializationGuard::~FDMInitializationGuard()
+{
+	if (FDMInitializationGuard::GuardCount > 0)
+	{
+		--FDMInitializationGuard::GuardCount;
+	}
+}
diff --git a/Engine/Plugins/Experimental/DynamicMaterial/Source/DynamicMaterial/Private/DynamicMaterialModule.cpp b/Engine/Plugins/Experimental/DynamicMaterial/Source/DynamicMaterial/Private/DynamicMaterialModule.cpp
new file mode 100644
--- /dev/null
+++ b/Engine/Plugins/Experimental/DynamicMaterial/Source/DynamicMaterial/Private/DynamicMaterialModule.cpp
@@ -0,0 +1,68 @@
+// Copyright Epic Games, Inc. All Rights Reserved.
+
+#include "DynamicMaterialModule.h"
+#include "HAL/IConsoleManager.h"
+#include "Misc/CoreDelegates.h"
+#include "Modules/ModuleManager.h"
+#include "UObject/UObjectBase.h"
+
+DEFINE_LOG_CATEGORY(LogDynamicMaterial);
+
+bool FDynamicMaterialModule::bIsEngineExiting = false;
+
+#if WITH_EDITOR
+FDMCreateEditorOnlyDataDelegate FDynamicMaterialModule::CreateEditorOnlyDataDelegate;
+#endif
+
+static TAutoConsoleVariable<bool> CVarExportMaterials(
+	TEXT("DM.ExportMaterials"),
+	false,
+	TEXT("If enabled, all materials, including previews, are exported to /Game/DynamicMaterials."),
+	ECVF_SetByConsole);
+
+bool FDynamicMaterialModule::IsMaterialExportEnabled()
+{
+	return CVarExportMaterials.GetValueOnAnyThread();
+}
+
+bool FDynamicMaterialModule::AreUObjectsSafe()
+{
+	return UObjectInitialized() && !bIsEngineExiting;
+}
+
+FDynamicMaterialModule& FDynamicMaterialModule::Get()
+{
+	return FModuleManager::LoadModuleChecked<FDynamicMaterialModule>("DynamicMaterial");
+}
+
+void FDynamicMaterialModule::StartupModule()
+{
+	EnginePreExitHandle = FCoreDelegates::OnEnginePreExit.AddStatic(&FDynamicMaterialModule::HandleEnginePreExit);
+}
+
+void FDynamicMaterialModule::ShutdownModule()
+{
+	if (EnginePreExitHandle.IsValid())
+	{
+		FCoreDelegates::OnEnginePreExit.Remove(EnginePreExitHandle);
+	}
+}
+
+void FDynamicMaterialModule::HandleEnginePreExit()
+{
+	bIsEngineExiting = true;
+}
+
+#if WITH_EDITOR
+TScriptInterface<IDynamicMaterialModelEditorOnlyDataInterface> FDynamicMaterialModule::CreateEditorOnlyData(UDynamicMaterialModel* InMaterialModel)
+{
+	if (CreateEditorOnlyDataDelegate.IsBound())
+	{
+		return CreateEditorOnlyDataDelegate.Execute(InMaterialModel);
+	}
+
+	return nullptr;
+}
+#endif
+
+IMPLEMENT_MODULE(FDynamicMaterialModule, DynamicMaterial)
diff --git a/Engine/Plugins/Experimental/DynamicMaterial/Source/DynamicMaterial/Private/Material/DynamicMaterialInstance.cpp b/Engine/Plugins/Experimental/DynamicMaterial/Source/DynamicMaterial/Private/Material/DynamicMaterialInstance.cpp
new file mode 100644
--- /dev/null
+++ b/Engine/Plugins/Experimental/DynamicMaterial/Source/DynamicMaterial/Private/Material/DynamicMaterialInstance.cpp
@@ -0,0 +1,70 @@
+// Copyright Epic Games, Inc. All Rights Reserved.
+
+#include "Material/DynamicMaterialInstance.h"
+#include "Materials/Material.h"
+#include "Model/DynamicMaterialModel.h"
+#include "UObject/ObjectSaveContext.h"
+
+#if WITH_EDITOR
+#include "Model/IDynamicMaterialModelEditorOnlyDataInterface.h"
+#endif
+
+UDynamicMaterialInstance::UDynamicMaterialInstance()
+{
+	BaseMaterial = nullptr;
+	MaterialModel = nullptr;
+
+	bOutputTranslucentVelocity = true;
+}
+
+#if WITH_EDITOR
+void UDynamicMaterialInstance::InitializeMIDPublic()
+{
+	check(MaterialModel);
+
+	BaseMaterial = MaterialModel->GetGeneratedMaterial();
+	SetParentInternal(BaseMaterial, false);
+	ClearParameterValues();
+	UpdateCachedData();
+}
+
+void UDynamicMaterialInstance::PostDuplicate(bool bDuplicateForPIE)
+{
+	Super::PostDuplicate(bDuplicateForPIE);
+
+	if (MaterialModel)
+	{
+		MaterialModel->SetDynamicMaterialInstance(this);
+
+		if (IDynamicMaterialModelEditorOnlyDataInterface* ModelEditorOnlyData = MaterialModel->GetEditorOnlyData())
+		{
+			ModelEditorOnlyData->RequestMaterialBuild();
+		}
+	}
+}
+
+void UDynamicMaterialInstance::PostEditImport()
+{
+	Super::PostEditImport();
+
+	if (MaterialModel)
+	{
+		MaterialModel->SetDynamicMaterialInstance(this);
+
+		if (IDynamicMaterialModelEditorOnlyDataInterface* ModelEditorOnlyData = MaterialModel->GetEditorOnlyData())
+		{
+			ModelEditorOnlyData->RequestMaterialBuild();
+		}
+	}
+}
+
+void UDynamicMaterialInstance::OnMaterialBuilt(UDynamicMaterialModel* InMaterialModel)
+{
+	if (MaterialModel != InMaterialModel)
+	{
+		return;
+	}
+
+	InitializeMIDPublic();
+}
+#endif
\ No newline at end of file
diff --git a/Engine/Plugins/Experimental/DynamicMaterial/Source/DynamicMaterial/Private/Model/DynamicMaterialModel.cpp b/Engine/Plugins/Experimental/DynamicMaterial/Source/DynamicMaterial/Private/Model/DynamicMaterialModel.cpp
new file mode 100644
--- /dev/null
+++ b/Engine/Plugins/Experimental/DynamicMaterial/Source/DynamicMaterial/Private/Model/DynamicMaterialModel.cpp
@@ -0,0 +1,550 @@
+// Copyright Epic Games, Inc. All Rights Reserved.
+
+#include "Model/DynamicMaterialModel.h"
+#include "Components/DMMaterialParameter.h"
+#include "Components/DMMaterialValue.h"
+#include "Components/MaterialValues/DMMaterialValueFloat1.h"
+#include "DMComponentPath.h"
+#include "DMDefs.h"
+#include "Material/DynamicMaterialInstance.h"
+#include "Model/IDynamicMaterialModelEditorOnlyDataInterface.h"
+
+#if WITH_EDITOR
+#include "Components/DMTextureUV.h"
+#include "DynamicMaterialModule.h"
+#include "Materials/Material.h"
+#endif
+
+const FString UDynamicMaterialModel::ValuesPathToken          = FString(TEXT("Values"));
+const FString UDynamicMaterialModel::ParametersPathToken      = FString(TEXT("Parameters"));
+const FName UDynamicMaterialModel::GlobalOpacityParameterName = FName("GlobalOpacity");
+
+UDynamicMaterialModel::UDynamicMaterialModel()
+{
+	DynamicMaterialInstance = nullptr;
+
+#if WITH_EDITOR
+	EditorOnlyDataSI.SetObject(nullptr);
+#endif
+
+	FDMUpdateGuard Guard;
+	GlobalOpacityValue = CreateDefaultSubobject<UDMMaterialValueFloat1>(GlobalOpacityParameterName);
+
+#if WITH_EDITOR
+	GlobalOpacityValue->SetDefaultValue(1.f);
+	GlobalOpacityValue->ApplyDefaultValue();
+#endif
+
+	GlobalOpacityParameter = CreateDefaultSubobject<UDMMaterialParameter>("GlobalOpacityParameter");
+	GlobalOpacityParameter->ParameterName = GlobalOpacityParameterName;
+	GlobalOpacityValue->Parameter = GlobalOpacityParameter;
+
+	ParameterMap.Add(GlobalOpacityParameterName, GlobalOpacityParameter);
+}
+
+void UDynamicMaterialModel::SetDynamicMaterialInstance(UDynamicMaterialInstance* InDynamicMaterialInstance)
+{
+	DynamicMaterialInstance = InDynamicMaterialInstance;
+
+#if WITH_EDITOR
+	if (IDynamicMaterialModelEditorOnlyDataInterface* ModelEditorOnlyData = GetEditorOnlyData())
+	{
+		ModelEditorOnlyData->PostEditorDuplicate();
+	}
+#endif
+}
+
+bool UDynamicMaterialModel::IsModelValid() const
+{
+	return (!HasAnyFlags(RF_BeginDestroyed | RF_FinishDestroyed) && IsValid(this));
+}
+
+UDMMaterialComponent* UDynamicMaterialModel::GetComponentByPath(const FString& InPath) const
+{
+	FDMComponentPath Path(InPath);
+	return GetComponentByPath(Path);
+}
+
+UDMMaterialComponent* UDynamicMaterialModel::GetComponentByPath(FDMComponentPath& InPath) const
+{
+	if (InPath.IsLeaf())
+	{
+		return nullptr;
+	}
+
+	const FDMComponentPathSegment FirstComponent = InPath.GetFirstSegment();
+
+	if (FirstComponent.GetToken() == ValuesPathToken)
+	{
+		int32 ValueIndex;
+
+		if (FirstComponent.GetParameter(ValueIndex))
+		{
+			if (Values.IsValidIndex(ValueIndex))
+			{
+				return Values[ValueIndex]->GetComponentByPath(InPath);
+			}
+		}
+
+		return nullptr;
+	}
+
+	if (FirstComponent.GetToken() == ParametersPathToken)
+	{
+		FString ParameterStr;
+
+		if (FirstComponent.GetParameter(ParameterStr))
+		{
+			const FName ParameterName = FName(*ParameterStr);
+
+			if (const TWeakObjectPtr<UDMMaterialParameter>* ParameterPtr = ParameterMap.Find(ParameterName))
+			{
+				return (*ParameterPtr)->GetComponentByPath(InPath);
+			}
+		}
+
+		return nullptr;
+	}
+
+#if WITH_EDITOR
+	if (IDynamicMaterialModelEditorOnlyDataInterface* EditorOnlyData = GetEditorOnlyData())
+	{
+		return EditorOnlyData->GetSubComponentByPath(InPath, FirstComponent);
+	}
+#endif
+
+	return nullptr;
+}
+
+#if WITH_EDITOR
+IDynamicMaterialModelEditorOnlyDataInterface* UDynamicMaterialModel::GetEditorOnlyData() const
+{
+	if (IsValid(EditorOnlyDataSI.GetObject()))
+	{
+		return EditorOnlyDataSI.GetInterface();
+	}
+
+	return nullptr;
+}
+
+UDMMaterialValue* UDynamicMaterialModel::GetValueByName(FName InName) const
+{
+	for (UDMMaterialValue* Value : Values)
+	{
+		if (Value->GetParameter() && Value->GetParameter()->GetParameterName() == InName)
+		{
+			return Value;
+		}
+	}
+
+	return nullptr;
+}
+
+UDMMaterialValue* UDynamicMaterialModel::GetValueByIndex(int32 Index) const
+{
+	if (Values.IsValidIndex(Index))
+	{
+		return Values[Index];
+	}
+
+	return nullptr;
+}
+
+UDMMaterialValue* UDynamicMaterialModel::AddValue(EDMValueType InType)
+{
+	UDMMaterialValue* NewValue = UDMMaterialValue::CreateMaterialValue(this, TEXT(""), InType, false);
+	Values.Add(NewValue);
+
+	if (IDynamicMaterialModelEditorOnlyDataInterface* ModelEditorOnlyData = GetEditorOnlyData())
+	{
+		ModelEditorOnlyData->OnValueListUpdate();
+	}
+
+	return NewValue;
+}
+
+void UDynamicMaterialModel::RemoveValueByName(FName InName)
+{
+	int32 FoundIndex = Values.IndexOfByPredicate([InName](UDMMaterialValue* Value)
+		{
+			return Value->GetParameter() && Value->GetParameter()->GetParameterName() == InName;
+		});
+
+	if (FoundIndex == INDEX_NONE)
+	{
+		return;
+	}
+
+	Values.RemoveAt(FoundIndex);
+
+	if (IDynamicMaterialModelEditorOnlyDataInterface* ModelEditorOnlyData = GetEditorOnlyData())
+	{
+		ModelEditorOnlyData->RequestMaterialBuild();
+		ModelEditorOnlyData->OnValueListUpdate();
+	}
+}
+
+void UDynamicMaterialModel::RemoveValueByIndex(int32 Index)
+{
+	if (Values.IsValidIndex(Index))
+	{
+		Values.RemoveAt(Index);
+
+		if (IDynamicMaterialModelEditorOnlyDataInterface* ModelEditorOnlyData = GetEditorOnlyData())
+		{
+			ModelEditorOnlyData->RequestMaterialBuild();
+			ModelEditorOnlyData->OnValueListUpdate();
+		}
+	}
+}
+
+bool UDynamicMaterialModel::HasParameterName(FName InParameterName) const
+{
+	if (Values.ContainsByPredicate([InParameterName](const UDMMaterialValue* InValue)
+		{
+			return IsValid(InValue) && InValue->GetParameter() && InValue->GetParameter()->GetParameterName() == InParameterName;
+		}))
+	{
+		return true;
+	}
+
+	const TWeakObjectPtr<UDMMaterialParameter>* ParameterWeak = ParameterMap.Find(InParameterName);
+
+	return (ParameterWeak && ParameterWeak->IsValid());
+}
+
+UDMMaterialParameter* UDynamicMaterialModel::CreateUniqueParameter(FName InBaseName)
+{
+	check(!InBaseName.IsNone());
+
+	UDMMaterialParameter* NewParameter = nullptr;
+
+	{
+		const FDMInitializationGuard InitGuard;
+
+		NewParameter = NewObject<UDMMaterialParameter>(this, NAME_None, RF_Transactional);
+		check(NewParameter);
+
+		RenameParameter(NewParameter, InBaseName);
+	}
+
+	ParameterMap.Emplace(NewParameter->GetParameterName(), NewParameter);
+	NewParameter->SetComponentState(EDMComponentLifetimeState::Added);
+
+	return NewParameter;
+}
+
+void UDynamicMaterialModel::RenameParameter(UDMMaterialParameter* InParameter, FName InBaseName)
+{
+	check(InParameter);
+	check(!InBaseName.IsNone());
+
+	if (!InParameter->ParameterName.IsNone())
+	{
+		FreeParameter(InParameter);
+	}
+
+	if (GUndo)
+	{
+		InParameter->Modify();
+	}
+
+	InParameter->ParameterName = CreateUniqueParameterName(InBaseName);
+	ParameterMap.Emplace(InParameter->ParameterName, InParameter);
+}
+
+void UDynamicMaterialModel::FreeParameter(UDMMaterialParameter* InParameter)
+{
+	check(InParameter);
+
+	if (InParameter->ParameterName.IsNone())
+	{
+		return;
+	}
+
+	for (TMap<FName, TWeakObjectPtr<UDMMaterialParameter>>::TIterator It(ParameterMap); It; ++It)
+	{
+		const TWeakObjectPtr<UDMMaterialParameter>& ParameterWeak = It->Value;
+
+		if (ParameterWeak.IsValid() == false || ParameterWeak->GetParameterName() == InParameter->GetParameterName())
+		{
+			It.RemoveCurrent();
+		}
+	}
+
+	if (GUndo)
+	{
+		InParameter->Modify();
+	}
+
+	InParameter->ParameterName = NAME_None;
+	InParameter->SetComponentState(EDMComponentLifetimeState::Removed);
+}
+
+bool UDynamicMaterialModel::ConditionalFreeParameter(UDMMaterialParameter* InParameter)
+{
+	check(InParameter);
+
+	// Parameteres without names are not in the map
+	if (InParameter->ParameterName.IsNone())
+	{
+		return true;
+	}
+
+	TWeakObjectPtr<UDMMaterialParameter>* ParameterPtr = ParameterMap.Find(InParameter->ParameterName);
+
+	// Parameter name isn't in the map or it's mapped to a different object.
+	if (!ParameterPtr || InParameter != ParameterPtr->Get())
+	{
+		return true;
+	}
+
+	// We're in the map at the given name.
+	return false;
+}
+
+void UDynamicMaterialModel::ResetData()
+{
+	Values.Empty();
+
+	if (IDynamicMaterialModelEditorOnlyDataInterface* ModelEditorOnlyData = GetEditorOnlyData())
+	{
+		ModelEditorOnlyData->ResetData();
+	}
+}
+
+void UDynamicMaterialModel::PostLoad()
+{
+	Super::PostLoad();
+
+	SetFlags(RF_Transactional);
+
+	FixGlobalOpacityVars();
+
+	IDynamicMaterialModelEditorOnlyDataInterface* ModelEditorOnlyData = GetEditorOnlyData();
+
+	if (!ModelEditorOnlyData)
+	{
+		EditorOnlyDataSI = FDynamicMaterialModule::CreateEditorOnlyData(this);
+		EditorOnlyDataSI->LoadDeprecatedModelData(this);
+	}
+
+	ReinitComponents();
+}
+
+void UDynamicMaterialModel::PostEditUndo()
+{
+	Super::PostEditUndo();
+
+	if (IDynamicMaterialModelEditorOnlyDataInterface* ModelEditorOnlyData = GetEditorOnlyData())
+	{
+		ModelEditorOnlyData->RequestMaterialBuild();
+	}
+}
+
+void UDynamicMaterialModel::PostEditImport()
+{
+	Super::PostEditImport();
+
+	FixGlobalOpacityVars();
+	PostEditorDuplicate();
+	ReinitComponents();
+
+	if (IDynamicMaterialModelEditorOnlyDataInterface* ModelEditorOnlyData = GetEditorOnlyData())
+	{
+		ModelEditorOnlyData->RequestMaterialBuild();
+	}
+}
+
+void UDynamicMaterialModel::PostDuplicate(bool bDuplicateForPIE)
+{
+	Super::PostDuplicate(bDuplicateForPIE);
+
+	FixGlobalOpacityVars();
+	PostEditorDuplicate();
+	ReinitComponents();
+
+	if (IDynamicMaterialModelEditorOnlyDataInterface* ModelEditorOnlyData = GetEditorOnlyData())
+	{
+		ModelEditorOnlyData->RequestMaterialBuild();
+	}
+}
+
+void UDynamicMaterialModel::PostEditorDuplicate()
+{
+	for (const TObjectPtr<UDMMaterialValue>& Value : Values)
+	{
+		if (GUndo)
+		{
+			Value->Modify();
+		}
+
+		Value->PostEditorDuplicate(this, nullptr);
+	}
+
+	if (GUndo)
+	{
+		GlobalOpacityValue->Modify();
+	}
+
+	GlobalOpacityValue->PostEditorDuplicate(this, nullptr);
+
+	for (const TPair<FName, TWeakObjectPtr<UDMMaterialParameter>>& Pair : ParameterMap)
+	{
+		if (UDMMaterialParameter* Parameter = Pair.Value.Get())
+		{
+			if (GUndo)
+			{
+				GlobalOpacityValue->Modify();
+			}
+
+			Parameter->PostEditorDuplicate(this, nullptr);
+		}
+	}
+
+	if (IDynamicMaterialModelEditorOnlyDataInterface* ModelEditorOnlyData = GetEditorOnlyData())
+	{
+		ModelEditorOnlyData->PostEditorDuplicate();
+	}
+}
+
+void UDynamicMaterialModel::ReinitComponents()
+{
+	if (IsValid(DynamicMaterial))
+	{
+		if (GUndo)
+		{
+			DynamicMaterial->Modify();
+		}
+
+		DynamicMaterial->AtomicallySetFlags(RF_DuplicateTransient);
+	}
+
+	if (IsValid(GlobalOpacityValue) && GlobalOpacityValue->GetMaterialParameterName() != GlobalOpacityParameterName)
+	{
+		if (GUndo)
+		{
+			GlobalOpacityValue->Modify();
+		}
+
+		GlobalOpacityValue->SetParameterName(GlobalOpacityParameterName);
+	}
+
+	// Clean up old parameters
+	ParameterMap.Empty();
+
+	TArray<UObject*> Subobjects;
+	GetObjectsWithOuter(this, Subobjects, false);
+
+	for (UObject* Subobject : Subobjects)
+	{
+		if (UDMMaterialValue* Value = Cast<UDMMaterialValue>(Subobject))
+		{
+			if (UDMMaterialParameter* Parameter = Value->GetParameter())
+			{
+				ParameterMap.Emplace(Parameter->GetParameterName(), Parameter);
+			}
+		}
+		else if (UDMTextureUV* TextureUV = Cast<UDMTextureUV>(Subobject))
+		{
+			TArray<UDMMaterialParameter*> Parameters = TextureUV->GetParameters();
+
+			for (UDMMaterialParameter* Parameter : Parameters)
+			{
+				ParameterMap.Emplace(Parameter->GetParameterName(), Parameter);
+			}
+		}
+	}
+
+	if (IDynamicMaterialModelEditorOnlyDataInterface* ModelEditorOnlyData = GetEditorOnlyData())
+	{
+		ModelEditorOnlyData->ReinitComponents();
+	}
+}
+
+void UDynamicMaterialModel::FixGlobalOpacityVars()
+{
+	if (const TWeakObjectPtr<UDMMaterialParameter>* ParameterPtr = ParameterMap.Find(GlobalOpacityParameterName))
+	{
+		UDMMaterialParameter* Parameter = ParameterPtr->Get();
+
+		if (!Parameter || Parameter->HasAnyFlags(RF_ArchetypeObject)
+			|| (GlobalOpacityParameter && Parameter != GlobalOpacityParameter))
+		{
+			if (GlobalOpacityParameter)
+			{
+				ParameterMap[GlobalOpacityParameterName] = GlobalOpacityParameter;
+			}
+			else
+			{
+				ParameterMap.Remove(GlobalOpacityParameterName);
+			}
+		}
+	}
+
+	if (!GlobalOpacityValue || GlobalOpacityParameter == GlobalOpacityValue->Parameter)
+	{
+		return;
+	}
+
+	if (GUndo)
+	{
+		GlobalOpacityParameter->Modify();
+		GlobalOpacityValue->Modify();
+	}
+
+	GlobalOpacityParameter->ParameterName = GlobalOpacityParameterName;
+	GlobalOpacityValue->Parameter = GlobalOpacityParameter;
+}
+
+FName UDynamicMaterialModel::CreateUniqueParameterName(FName InBaseName)
+{
+	int32 CurrentTest = 0;
+	FName UniqueName = InBaseName;
+
+	auto UpdateName = [&InBaseName, &CurrentTest, &UniqueName]()
+		{
+			++CurrentTest;
+			UniqueName = FName(InBaseName.ToString() + TEXT("_") + FString::FromInt(CurrentTest));
+		};
+
+	bool bIsNameUnique = false;
+
+	while (!bIsNameUnique)
+	{
+		bIsNameUnique = true;
+
+		if (Values.ContainsByPredicate([UniqueName](const UDMMaterialValue* Value)
+			{
+				return Value->GetMaterialParameterName() == UniqueName;
+			}))
+		{
+			bIsNameUnique = false;
+			UpdateName();
+		}
+	}
+
+	bIsNameUnique = false;
+
+	while (!bIsNameUnique)
+	{
+		bIsNameUnique = true;
+
+		if (const TWeakObjectPtr<UDMMaterialParameter>* CurrentParameter = ParameterMap.Find(UniqueName))
+		{
+			if (CurrentParameter->IsValid())
+			{
+				bIsNameUnique = false;
+				UpdateName();
+			}
+			else
+			{
+				ParameterMap.Remove(UniqueName);
+				break; // Not needed, but informative.
+			}
+		}
+	}
+
+	return UniqueName;
+}
+#endif
diff --git a/Engine/Plugins/Experimental/DynamicMaterial/Source/DynamicMaterial/Public/Components/DMMaterialComponent.h b/Engine/Plugins/Experimental/DynamicMaterial/Source/DynamicMaterial/Public/Components/DMMaterialComponent.h
new file mode 100644
--- /dev/null
+++ b/Engine/Plugins/Experimental/DynamicMaterial/Source/DynamicMaterial/Public/Components/DMMaterialComponent.h
@@ -0,0 +1,180 @@
+// Copyright Epic Games, Inc. All Rights Reserved.
+
+#pragma once
+
+#include "UObject/Object.h"
+#include "DMDefs.h"
+#include "Misc/NotifyHook.h"
+#include "DMMaterialComponent.generated.h"
+
+class UDynamicMaterialModel;
+enum class EDMUpdateType : uint8;
+struct FDMComponentPathSegment;
+struct FDMComponentPath;
+
+UENUM(BlueprintType)
+enum class EDMComponentLifetimeState : uint8
+{
+	Created,
+	Added,
+	Removed
+};
+
+/**
+ * The base class for all material components. Has a few useful things.
+ */
+UCLASS(Abstract, BlueprintType, meta = (DisplayName = "Material Designer Component"))
+class DYNAMICMATERIAL_API UDMMaterialComponent : public UObject, public FNotifyHook
+{
+	GENERATED_BODY()
+
+	friend class UDynamicMaterialEditorSettings;
+
+public:
+	UDMMaterialComponent();
+
+	/** Checks object flags and IsValid() */
+	UFUNCTION(BlueprintPure, Category = "Material Designer")
+	bool IsComponentValid() const;
+
+	UFUNCTION(BlueprintPure, Category = "Material Designer")
+	virtual bool IsRootComponent() const { return false; }
+
+	UFUNCTION(BlueprintCallable, Category = "Material Designer")
+	UDMMaterialComponent* GetComponentByPath(const FString& InPath) const;
+
+	UDMMaterialComponent* GetComponentByPath(FDMComponentPath& InPath) const;
+
+	template<typename InComponentClass>
+	UDMMaterialComponent* GetComponentByPath(FDMComponentPath& InPath) const
+	{
+		return Cast<InComponentClass>(GetComponentByPath(InPath));
+	}
+
+	UFUNCTION(BlueprintCallable, Category = "Material Designer")
+	virtual void Update(EDMUpdateType InUpdateType);
+
+#if WITH_EDITOR
+	UFUNCTION(BlueprintPure, Category = "Material Designer")
+	/** Returns the complete path from the model to this component. */
+	FString GetComponentPath() const;
+
+	UFUNCTION(BlueprintCallable, Category = "Material Designer")
+	virtual UDMMaterialComponent* GetParentComponent() const;
+
+	UFUNCTION(BlueprintCallable, Category = "Material Designer")
+	UDMMaterialComponent* GetTypedParent(UClass* InParentClass, bool bInAllowSubclasses) const;
+
+	template<class InParentClass>
+	InParentClass* GetTypedParent(bool bInAllowSubclasses) const
+	{
+		return Cast<InParentClass>(GetTypedParent(InParentClass::StaticClass(), bInAllowSubclasses));
+	}
+
+	/* Returns a description of this class/object. */
+	UFUNCTION(BlueprintPure, Category = "Material Designer")
+	virtual FText GetComponentDescription() const;
+
+	/** Returns true if we can attempt to un-dirty this component. */
+	static bool CanClean();
+
+	/** Called to prevent cleaning for MinTimeBeforeClean. **/
+	static void PreventClean(double DelayFor = MinTimeBeforeClean);
+
+	/** Returns true if this component has been marked dirty. */
+	bool NeedsClean();
+
+	/** Performs whatever operation is involved in cleaning this component. */
+	virtual void DoClean();
+
+	UFUNCTION(BlueprintPure, Category = "Material Designer")
+	EDMComponentLifetimeState GetComponentState() const { return ComponentState; }
+
+	UFUNCTION(BlueprintPure, Category = "Material Designer")
+	bool IsComponentCreated() const { return ComponentState == EDMComponentLifetimeState::Created; }
+
+	UFUNCTION(BlueprintPure, Category = "Material Designer")
+	bool HasComponentBeenCreated() const
+	{
+		// This is kind of a "useless" check, a component has _always_ been created.
+		return true;
+	}
+
+	UFUNCTION(BlueprintPure, Category = "Material Designer")
+	bool IsComponentAdded() const { return ComponentState == EDMComponentLifetimeState::Added; }
+
+	UFUNCTION(BlueprintPure, Category = "Material Designer")
+	bool HasComponentBeenAdded() const { return ComponentState >= EDMComponentLifetimeState::Added; }
+
+	UFUNCTION(BlueprintPure, Category = "Material Designer")
+	bool IsComponentRemoved() const { return ComponentState == EDMComponentLifetimeState::Removed; }
+
+	UFUNCTION(BlueprintPure, Category = "Material Designer")
+	bool HasComponentBeenRemoved() const { return ComponentState >= EDMComponentLifetimeState::Removed; }
+
+	UFUNCTION(BlueprintCallable, Category = "Material Designer")
+	void SetComponentState(EDMComponentLifetimeState NewState);
+
+	//~ Begin UObject
+	virtual bool Modify(bool bInAlwaysMarkDirty = true) override;
+	virtual void PostLoad() override;
+	//~ End UObject
+
+	/*** Begin FNotifyHook */
+	virtual void NotifyPreChange(class FEditPropertyChain* PropertyAboutToChange) {}
+	/*** End FNotifyHook */
+
+	UFUNCTION(BlueprintPure, Category = "Material Designer")
+	const TArray<FName>& GetEditableProperties() const { return EditableProperties; }
+
+	UFUNCTION(BlueprintPure, Category = "Material Designer")
+	virtual bool IsPropertyVisible(FName Property) const { return true; }
+
+	/** Returns the part of the component representing just this object */
+	virtual FString GetComponentPathComponent() const;
+
+	virtual void PostEditorDuplicate(UDynamicMaterialModel* InMaterialModel, UDMMaterialComponent* InParent);
+
+	DECLARE_MULTICAST_DELEGATE_TwoParams(FOnUpdate, UDMMaterialComponent*, EDMUpdateType)
+	FOnUpdate& GetOnUpdate() { return OnUpdate; }
+
+	DECLARE_MULTICAST_DELEGATE_TwoParams(FOnLifetimeStateChanged, UDMMaterialComponent*, EDMComponentLifetimeState)
+	FOnLifetimeStateChanged& GetOnAdded() { return OnAdded; }
+	FOnLifetimeStateChanged& GetOnRemoved() { return OnRemoved; }
+#endif
+
+protected:
+	static const double MinTimeBeforeClean;
+	static double MinCleanTime;
+
+#if WITH_EDITORONLY_DATA
+	UPROPERTY(VisibleInstanceOnly, BlueprintReadOnly, Category = "Material Designer")
+	EDMComponentLifetimeState ComponentState;
+
+	UPROPERTY(VisibleInstanceOnly, BlueprintReadOnly, Transient, Category = "Material Designer")
+	bool bComponentDirty;
+
+	UPROPERTY(VisibleInstanceOnly, BlueprintReadOnly, Transient, Category = "Material Designer")
+	TArray<FName> EditableProperties;
+
+	FOnUpdate OnUpdate;
+	FOnLifetimeStateChanged OnAdded;
+	FOnLifetimeStateChanged OnRemoved;
+#endif
+
+	UFUNCTION(BlueprintPure, Category = "Material Designer")
+	UObject* GetOuterSafe() const;
+
+	virtual UDMMaterialComponent* GetSubComponentByPath(FDMComponentPath& InPath, const FDMComponentPathSegment& InPathSegment) const;
+
+	void MarkComponentDirty();
+
+#if WITH_EDITOR
+	/** Allows this object to modify the child path when generating a path. */
+	virtual void GetComponentPathInternal(TArray<FString>& OutChildComponentPathComponents) const;
+
+	virtual void OnComponentStateChange(EDMComponentLifetimeState NewState);
+	virtual void OnComponentAdded();
+	virtual void OnComponentRemoved();
+#endif
+};
diff --git a/Engine/Plugins/Experimental/DynamicMaterial/Source/DynamicMaterial/Public/Components/DMMaterialLinkedComponent.h b/Engine/Plugins/Experimental/DynamicMaterial/Source/DynamicMaterial/Public/Components/DMMaterialLinkedComponent.h
new file mode 100644
--- /dev/null
+++ b/Engine/Plugins/Experimental/DynamicMaterial/Source/DynamicMaterial/Public/Components/DMMaterialLinkedComponent.h
@@ -0,0 +1,31 @@
+// Copyright Epic Games, Inc. All Rights Reserved.
+
+#pragma once
+
+#include "DMMaterialComponent.h"
+#include "DMMaterialLinkedComponent.generated.h"
+
+UCLASS(Abstract, BlueprintType, meta = (DisplayName = "Material Designer Linked Component"))
+class DYNAMICMATERIAL_API UDMMaterialLinkedComponent : public UDMMaterialComponent
+{
+	GENERATED_BODY()
+
+public:
+	UDMMaterialLinkedComponent() = default;
+
+#if WITH_EDITOR
+	virtual UDMMaterialComponent* GetParentComponent() const override;
+
+	UFUNCTION(BlueprintCallable, Category = "Material Designer")
+	void SetParentComponent(UDMMaterialComponent* InParentComponent);
+
+	virtual void PostEditorDuplicate(UDynamicMaterialModel* InMaterialModel, UDMMaterialComponent* InParent) override;
+#endif
+
+protected:
+#if WITH_EDITORONLY_DATA
+	UPROPERTY()
+	TObjectPtr<UDMMaterialComponent> ParentComponent = nullptr;
+#endif
+};
+
diff --git a/Engine/Plugins/Experimental/DynamicMaterial/Source/DynamicMaterial/Public/Components/DMMaterialParameter.h b/Engine/Plugins/Experimental/DynamicMaterial/Source/DynamicMaterial/Public/Components/DMMaterialParameter.h
new file mode 100644
--- /dev/null
+++ b/Engine/Plugins/Experimental/DynamicMaterial/Source/DynamicMaterial/Public/Components/DMMaterialParameter.h
@@ -0,0 +1,52 @@
+// Copyright Epic Games, Inc. All Rights Reserved.
+
+#pragma once
+
+#include "Components/DMMaterialLinkedComponent.h"
+#include "DMMaterialParameter.generated.h"
+
+class UDynamicMaterialModel;
+
+/**
+ * A parameter on a Material Designer Instance.
+ */
+UCLASS(BlueprintType, meta = (DisplayName = "Material Designer Parameter"))
+class DYNAMICMATERIAL_API UDMMaterialParameter : public UDMMaterialLinkedComponent
+{
+	GENERATED_BODY()
+
+public:
+	friend class UDynamicMaterialModel;
+
+	UFUNCTION(BlueprintPure, Category = "Material Designer")
+	UDynamicMaterialModel* GetMaterialModel() const;
+
+	UFUNCTION(BlueprintPure, Category = "Material Designer")
+	FName GetParameterName() const { return ParameterName; }
+
+#if WITH_EDITOR
+	UFUNCTION(BlueprintCallable, Category = "Material Designer")
+	void RenameParameter(FName InBaseParameterName);
+
+	//~ Begin UObject
+	virtual void BeginDestroy() override;
+	//~ End UObject
+
+	//~ Begin UDMMaterialComponent
+	virtual void PostEditorDuplicate(UDynamicMaterialModel* InMaterialModel, UDMMaterialComponent* InParent) override;
+	//~ End UDMMaterialComponent
+#endif
+
+protected:
+	UPROPERTY(BlueprintReadOnly, VisibleInstanceOnly, Category = "Material Designer")
+	FName ParameterName;
+
+private:
+	UDMMaterialParameter();
+
+	#if WITH_EDITOR
+	//~ Begin UDMMaterialComponent
+	virtual void OnComponentRemoved() override;
+	//~ End UDMMaterialComponent
+#endif
+};
diff --git a/Engine/Plugins/Experimental/DynamicMaterial/Source/DynamicMaterial/Public/Components/DMMaterialValue.h b/Engine/Plugins/Experimental/DynamicMaterial/Source/DynamicMaterial/Public/Components/DMMaterialValue.h
new file mode 100644
--- /dev/null
+++ b/Engine/Plugins/Experimental/DynamicMaterial/Source/DynamicMaterial/Public/Components/DMMaterialValue.h
@@ -0,0 +1,201 @@
+// Copyright Epic Games, Inc. All Rights Reserved.
+ 
+#pragma once
+ 
+#include "Components/DMMaterialLinkedComponent.h"
+#include "DMDefs.h"
+#include "Templates/SharedPointer.h"
+#include "UObject/StrongObjectPtr.h"
+#include "DMMaterialValue.generated.h"
+ 
+class UDMMaterialParameter;
+class UDynamicMaterialModel;
+class UMaterial;
+class UMaterialExpressionParameter;
+class UMaterialInstanceDynamic;
+
+#if WITH_EDITOR
+class FAssetThumbnailPool;
+class IDetailTreeNode;
+class IPropertyHandle;
+class IPropertyRowGenerator;
+class SWidget;
+struct IDMMaterialBuildStateInterface;
+#endif
+
+class UDMMaterialValue;
+
+/**
+ * A value used inside a material. Can be exported as a material parameter.
+ */
+UCLASS(BlueprintType, Abstract, ClassGroup = "Material Designer", meta = (DisplayName = "Material Designer Value"))
+class DYNAMICMATERIAL_API UDMMaterialValue : public UDMMaterialLinkedComponent
+{
+	GENERATED_BODY()
+ 
+	friend class UDynamicMaterialModel;
+
+#if WITH_EDITOR
+	friend class SDMMaterialEditor;
+#endif
+ 
+public:
+	static const FString ParameterPathToken;
+	static const TCHAR* ParameterNamePrefix;
+
+#if WITH_EDITOR
+	static const FName ValueName;
+
+	UFUNCTION(BlueprintCallable, Category = "Material Designer")
+	static UDMMaterialValue* CreateMaterialValue(UDynamicMaterialModel* InMaterialModel, const FString& InName, EDMValueType InType, bool bInLocal);
+#endif
+ 
+	UFUNCTION(BlueprintPure, Category = "Material Designer")
+	UDynamicMaterialModel* GetMaterialModel() const;
+ 
+	UFUNCTION(BlueprintPure, Category = "Material Designer")
+	EDMValueType GetType() const { return Type; }
+ 
+#if WITH_EDITOR
+	UFUNCTION(BlueprintPure, Category = "Material Designer")
+	FText GetTypeName() const;
+ 
+	UFUNCTION(BlueprintPure, Category = "Material Designer")
+	FText GetDescription() const;
+
+	UFUNCTION(BlueprintPure, Category = "Material Designer")
+	virtual bool IsDefaultValue() const PURE_VIRTUAL(UDMMaterialValue::IsDefaultValue, return false;)
+
+	bool CanResetToDefault(TSharedPtr<IPropertyHandle> InPropertyHandle) const { return !IsDefaultValue(); }
+ 
+	/**
+	 * Subclasses should implement a SetDefaultValue.
+	 */
+	UFUNCTION(BlueprintCallable, Category = "Material Designer")
+	virtual void ApplyDefaultValue() PURE_VIRTUAL(UDMMaterialValue::ApplyDefaultValue)
+
+	void ResetToDefault(TSharedPtr<IPropertyHandle> InPropertyHandle) { ApplyDefaultValue(); }
+ 
+	/**
+	 * Resets to the default default value. 0, nullptr, etc.
+	 */
+	UFUNCTION(BlueprintCallable, Category = "Material Designer")
+	virtual void ResetDefaultValue() PURE_VIRTUAL(UDMMaterialValue::ResetDefaultValue)
+#endif
+
+	UFUNCTION(BlueprintPure, Category = "Material Designer")
+	bool IsLocal() const { return bLocal; }
+ 
+	UFUNCTION(BlueprintPure, Category = "Material Designer")
+	UDMMaterialParameter* GetParameter() const { return Parameter; }
+ 
+	UFUNCTION(BlueprintPure, Category = "Material Designer")
+	FName GetMaterialParameterName() const;
+ 
+#if WITH_EDITOR
+	UFUNCTION(BlueprintCallable, Category = "Material Designer")
+	bool SetParameterName(FName InBaseName);
+
+	virtual void GenerateExpression(const TSharedRef<IDMMaterialBuildStateInterface>& InBuildState) const
+		PURE_VIRTUAL(UDMMaterialStageSource::GenerateExpressions)
+#endif
+ 
+#if WITH_EDITOR
+	UFUNCTION(BlueprintPure, Category = "Material Designer")
+	int32 FindIndex() const;
+ 
+	int32 FindIndexSafe() const;
+ 
+	UFUNCTION(BlueprintPure, Category = "Material Designer")
+	UMaterial* GetPreviewMaterial();
+ 
+	/**
+	 * Returns the output index (channel WHOLE_CHANNEL) if this expression has pre-masked outputs.
+	 * Returns INDEX_NONE if it is not supported.
+	 */
+	virtual int32 GetInnateMaskOutput(int32 OutputChannels) const;
+
+	TSharedPtr<IDetailTreeNode> GetDetailTreeNode();
+	TSharedPtr<IPropertyHandle> GetPropertyHandle();
+
+	virtual FName GetMainPropertyName() const { return NAME_None; }
+
+	/** Return true if, when setting the base stage, the same value should be applied to the mask stage. */
+	virtual bool IsWholeLayerValue() const { return false; }
+#endif
+ 
+	virtual void SetMIDParameter(UMaterialInstanceDynamic* InMID) const;
+ 
+	//~ Begin UDMMaterialComponent
+	virtual void Update(EDMUpdateType InUpdateType = EDMUpdateType::Value) override;
+#if WITH_EDITOR
+	virtual void DoClean() override;
+	virtual void PostEditorDuplicate(UDynamicMaterialModel* InMaterialModel, UDMMaterialComponent* InParent) override;
+#endif
+	//~ End UDMMaterialComponent
+ 
+#if WITH_EDITOR
+	//~ Begin UObject
+	virtual bool Modify(bool bInAlwaysMarkDirty = true) override;
+	virtual void PostEditUndo() override;
+	virtual void PostEditChangeProperty(FPropertyChangedEvent& PropertyChangedEvent) override;
+	virtual void PostCDOContruct() override;
+	virtual void PostLoad() override;
+	virtual void PostEditImport() override;
+	virtual void BeginDestroy() override;
+	//~ End UObject
+#endif
+ 
+protected:
+	static TMap<EDMValueType, TStrongObjectPtr<UClass>> TypeClasses;
+
+	UPROPERTY(VisibleInstanceOnly, BlueprintReadOnly, Category = "Material Designer")
+	EDMValueType Type;
+ 
+	/**
+	 * True: The value is local to the stage it is used in. 
+	 * False: The value is a global value that can be used anywhere
+	 */
+	UPROPERTY(VisibleInstanceOnly, BlueprintReadOnly, Category = "Material Designer")
+	bool bLocal;
+
+	/**
+	 * The parameter name used to expose this value in a material.
+	 * If it isn't provided, GetFName() will be used instead.
+	 */
+	UPROPERTY(VisibleInstanceOnly, BlueprintReadOnly, Instanced, Category = "Material Designer")
+	TObjectPtr<UDMMaterialParameter> Parameter;
+ 
+#if WITH_EDITORONLY_DATA
+	UPROPERTY(VisibleInstanceOnly, BlueprintReadOnly, Transient, DuplicateTransient, TextExportTransient, Category = "Material Designer")
+	TObjectPtr<UMaterial> PreviewMaterial;
+#endif
+
+#if WITH_EDITOR 
+	TSharedPtr<IPropertyRowGenerator> PropertyRowGenerator;
+	TSharedPtr<IDetailTreeNode> DetailTreeNode;
+	TSharedPtr<IPropertyHandle> PropertyHandle;
+#endif
+ 
+	UDMMaterialValue(EDMValueType InType);
+ 
+	virtual void OnValueUpdated(bool bInForceStructureUpdate);
+ 
+	//~ Begin UDMMaterialComponent
+	virtual UDMMaterialComponent* GetSubComponentByPath(FDMComponentPath& InPath, const FDMComponentPathSegment& InPathSegment) const override;
+	//~ End UDMMaterialComponent
+
+	#if WITH_EDITOR
+	void CreatePreviewMaterial();
+	void UpdatePreviewMaterial();
+ 
+	//~ Begin UDMMaterialComponent
+	virtual void OnComponentRemoved() override;
+	//~ End UDMMaterialComponent
+ 
+	void EnsureDetailObjects();
+#endif
+
+private:
+	UDMMaterialValue();
+};
diff --git a/Engine/Plugins/Experimental/DynamicMaterial/Source/DynamicMaterial/Public/Components/DMTextureUV.h b/Engine/Plugins/Experimental/DynamicMaterial/Source/DynamicMaterial/Public/Components/DMTextureUV.h
new file mode 100644
--- /dev/null
+++ b/Engine/Plugins/Experimental/DynamicMaterial/Source/DynamicMaterial/Public/Components/DMTextureUV.h
@@ -0,0 +1,229 @@
+// Copyright Epic Games, Inc. All Rights Reserved.
+
+#pragma once
+
+#include "Components/DMMaterialLinkedComponent.h"
+#include "DMDefs.h"
+#include "Misc/Guid.h"
+#include "Templates/SharedPointer.h"
+
+#if WITH_EDITOR
+#include "Delegates/Delegate.h"
+#include "Delegates/DelegateCombinations.h"
+#endif
+
+#include "DMTextureUV.generated.h"
+
+class UDMMaterialParameter;
+class UDMTextureUV;
+class UDynamicMaterialModel;
+class UMaterialInstanceDynamic;
+
+#if WITH_EDITOR
+class IDetailTreeNode;
+class IPropertyHandle;
+class IPropertyRowGenerator;
+#endif
+
+namespace UE::DynamicMaterial::ParamID
+{
+	// Individual parameters
+	constexpr int32 PivotX = 0;
+	constexpr int32 PivotY = 1;
+	constexpr int32 ScaleX = 2;
+	constexpr int32 ScaleY = 3;
+	constexpr int32 Rotation = 4;
+	constexpr int32 OffsetX = 5;
+	constexpr int32 OffsetY = 6;
+
+	// Parameter groups
+	constexpr int32 Pivot = PivotX;
+	constexpr int32 Scale = ScaleX;
+	//constexpr int32 Rotation = 4; // No additional group value needed
+	constexpr int32 Offset = OffsetX;
+}
+
+UCLASS(BlueprintType, ClassGroup = "Material Designer", meta = (DisplayName = "Material Designer Texture UV"))
+class DYNAMICMATERIAL_API UDMTextureUV : public UDMMaterialLinkedComponent
+{
+	GENERATED_BODY()
+
+	friend class SDMStageEdit;
+	friend class UDMMaterialStageInputTextureUV;
+
+public:
+#if WITH_EDITOR
+	static const FName NAME_UVSource;
+	static const FName NAME_Offset;
+	static const FName NAME_Pivot;
+	static const FName NAME_Rotation;
+	static const FName NAME_Scale;
+	static const FName NAME_bMirrorOnX;
+	static const FName NAME_bMirrorOnY;
+#endif
+
+	static const FString OffsetXPathToken;
+	static const FString OffsetYPathToken;
+	static const FString PivotXPathToken;
+	static const FString PivotYPathToken;
+	static const FString RotationPathToken;
+	static const FString ScaleXPathToken;
+	static const FString ScaleYPathToken;
+	
+	static const FGuid GUID;
+
+#if WITH_EDITOR
+	// The bool is whether or not the property should be exposed to Sequencer as keyable.
+	static const TMap<FName, bool> TextureProperties;
+
+	static UDMTextureUV* CreateTextureUV(UObject* InOuter);
+#endif
+
+	UDMTextureUV();
+
+#if WITH_EDITORONLY_DATA
+	UPROPERTY(EditInstanceOnly, BlueprintReadWrite, Category = "Material Designer")
+	bool bLinkScale = true;
+#endif
+
+	UFUNCTION(BlueprintPure, Category = "Material Designer")
+	UDynamicMaterialModel* GetMaterialModel() const;
+
+#if WITH_EDITOR
+	UFUNCTION(BlueprintPure, Category = "Material Designer")
+	EDMUVSource GetUVSource() const { return UVSource; }
+
+	UFUNCTION(BlueprintCallable, Category = "Material Designer")
+	void SetUVSource(EDMUVSource InUVSource);
+#endif
+
+	UFUNCTION(BlueprintPure, Category = "Material Designer")
+	const FVector2D& GetOffset() const { return Offset; }
+
+	UFUNCTION(BlueprintCallable, Category = "Material Designer")
+	void SetOffset(const FVector2D& InOffset);
+
+	UFUNCTION(BlueprintPure, Category = "Material Designer")
+	const FVector2D& GetPivot() const { return Pivot; }
+
+	UFUNCTION(BlueprintCallable, Category = "Material Designer")
+	void SetPivot(const FVector2D& InPivot);
+
+	UFUNCTION(BlueprintPure, Category = "Material Designer")
+	float GetRotation() const { return Rotation; }
+
+	UFUNCTION(BlueprintCallable, Category = "Material Designer")
+	void SetRotation(float InRotation);
+
+	UFUNCTION(BlueprintPure, Category = "Material Designer")
+	const FVector2D& GetScale() const { return Scale; }
+
+	UFUNCTION(BlueprintCallable, Category = "Material Designer")
+	void SetScale(const FVector2D& InScale);
+
+#if WITH_EDITOR
+	UFUNCTION(BlueprintPure, Category = "Material Designer")
+	bool GetMirrorOnX() const { return bMirrorOnX; }
+
+	UFUNCTION(BlueprintCallable, Category = "Material Designer")
+	void SetMirrorOnX(bool bInMirrorOnX);
+
+	UFUNCTION(BlueprintPure, Category = "Material Designer")
+	bool GetMirrorOnY() const { return bMirrorOnY; }
+
+	UFUNCTION(BlueprintCallable, Category = "Material Designer")
+	void SetMirrorOnY(bool bInMirrorOnY);
+
+	TSharedPtr<IDetailTreeNode> GetDetailTreeNode(FName InProperty);
+	TSharedPtr<IPropertyHandle> GetPropertyHandle(FName InProperty);
+#endif
+
+	TArray<UDMMaterialParameter*> GetParameters() const;
+
+	void SetMIDParameters(UMaterialInstanceDynamic* InMID);
+
+	//~ Begin UObject
+#if WITH_EDITOR
+	virtual bool Modify(bool bInAlwaysMarkDirty = true) override;
+	virtual void PostLoad() override;
+	virtual void PostEditImport() override;
+	virtual void PostEditChangeProperty(struct FPropertyChangedEvent& PropertyChangedEvent) override;
+	virtual void PreEditUndo() override;
+	virtual void PostEditUndo() override;
+#endif
+	virtual void Serialize(FArchive& Ar) override;
+	//~ End UObject
+
+#if WITH_EDITOR
+	//~ Begin UDMMaterialComponent
+	virtual void PostEditorDuplicate(UDynamicMaterialModel* InMaterialModel, UDMMaterialComponent* InParent) override;
+	//~ End UDMMaterialComponent
+#endif
+
+protected:
+#if WITH_EDITORONLY_DATA
+	UPROPERTY(EditInstanceOnly, BlueprintReadWrite, BlueprintSetter = SetUVSource, Category = "Material Designer",
+		meta = (AllowPrivateAccess = "true", ToolTip = "Forces a material rebuild.", NotKeyframeable))
+	EDMUVSource UVSource = EDMUVSource::Texture;
+
+	EDMUVSource UVSource_PreUndo = UVSource;
+#endif
+
+	UPROPERTY(EditInstanceOnly, BlueprintReadWrite, Getter = GetOffset, Setter = SetOffset, BlueprintSetter = SetOffset, Category = "Material Designer",
+		meta = (AllowPrivateAccess = "true", Delta = 0.001))
+	FVector2D Offset = FVector2D(0.f, 0.f);
+
+	UPROPERTY(EditInstanceOnly, BlueprintReadWrite, Getter = GetPivot, Setter = SetPivot, BlueprintSetter = SetPivot, Category = "Material Designer",
+		meta = (AllowPrivateAccess = "true", ToolTip="Pivot for rotation and scale.", Delta = 0.001))
+	FVector2D Pivot = FVector2D(0.5, 0.5);
+
+	UPROPERTY(EditInstanceOnly, BlueprintReadWrite, Getter = GetRotation, Setter = SetRotation, BlueprintSetter = SetRotation, Category = "Material Designer",
+		meta = (AllowPrivateAccess = "true", Delta = 1.0))
+	float Rotation = 0.f;
+
+	UPROPERTY(EditInstanceOnly, BlueprintReadWrite, Getter = GetScale, Setter = SetScale, BlueprintSetter = SetScale, Category = "Material Designer",
+		meta = (AllowPrivateAccess = "true", AllowPreserveRatio, Delta = 0.001))
+	FVector2D Scale = FVector2D(1.f, 1.f);
+
+#if WITH_EDITORONLY_DATA
+	UPROPERTY(EditInstanceOnly, BlueprintReadWrite, BlueprintSetter = SetMirrorOnX, Category = "Material Designer",
+		meta = (AllowPrivateAccess = "true", ToolTip = "Forces a material rebuild.", NotKeyframeable))
+	bool bMirrorOnX = false;
+
+	UPROPERTY(EditInstanceOnly, BlueprintReadWrite, BlueprintSetter = SetMirrorOnY, Category = "Material Designer",
+		meta = (AllowPrivateAccess = "true", ToolTip = "Forces a material rebuild.", NotKeyframeable))
+	bool bMirrorOnY = false;
+
+	bool bMirrorOnX_PreUndo = bMirrorOnX;
+	bool bMirrorOnY_PreUndo = bMirrorOnY;
+#endif
+
+	UPROPERTY(VisibleInstanceOnly, BlueprintReadOnly, DuplicateTransient, TextExportTransient, Category = "Material Designer")
+	TMap<int32, TObjectPtr<UDMMaterialParameter>> MaterialParameters;
+
+#if WITH_EDITOR
+	mutable TMap<int32, bool> MaterialNodesCreated;
+
+	TSharedPtr<IPropertyRowGenerator> PropertyRowGenerator;
+	TMap<FName, TSharedPtr<IDetailTreeNode>> DetailTreeNodes;
+	TMap<FName, TSharedPtr<IPropertyHandle>> PropertyHandles;
+
+	bool bNeedsPostLoadValueUpdate = false;
+	bool bNeedsPostLoadStructureUpdate = false;
+
+	void EnsureDetailObjects();
+
+	void CreateParameterNames();
+	void RemoveParameterNames();
+#endif
+
+	//~ Begin UDMMaterialComponent
+	virtual void Update(EDMUpdateType InUpdateType) override;
+	virtual UDMMaterialComponent* GetSubComponentByPath(FDMComponentPath& InPath, const FDMComponentPathSegment& InPathSegment) const override;
+#if WITH_EDITOR
+	virtual void GetComponentPathInternal(TArray<FString>& OutChildComponentPathComponents) const override;
+	virtual void OnComponentAdded() override;
+	virtual void OnComponentRemoved() override;
+#endif
+	//~ End UDMMaterialComponent
+};
diff --git a/Engine/Plugins/Experimental/DynamicMaterial/Source/DynamicMaterial/Public/Components/MaterialValues/DMMaterialValueBool.h b/Engine/Plugins/Experimental/DynamicMaterial/Source/DynamicMaterial/Public/Components/MaterialValues/DMMaterialValueBool.h
new file mode 100644
--- /dev/null
+++ b/Engine/Plugins/Experimental/DynamicMaterial/Source/DynamicMaterial/Public/Components/MaterialValues/DMMaterialValueBool.h
@@ -0,0 +1,50 @@
+// Copyright Epic Games, Inc. All Rights Reserved.
+ 
+#pragma once
+ 
+#include "Components/DMMaterialValue.h"
+#include "DMMaterialValueBool.generated.h"
+ 
+UCLASS(BlueprintType, ClassGroup = "Material Designer")
+class DYNAMICMATERIAL_API UDMMaterialValueBool : public UDMMaterialValue
+{
+	GENERATED_BODY()
+ 
+public:
+	UFUNCTION(BlueprintPure, Category = "Material Designer")
+	bool GetValue() const { return Value; }
+ 
+	UFUNCTION(BlueprintCallable, Category = "Material Designer")
+	void SetValue(bool InValue);
+ 
+#if WITH_EDITOR
+	UFUNCTION(BlueprintPure, Category = "Material Designer")
+	bool GetDefaultValue() const { return bDefaultValue; }
+ 
+	UFUNCTION(BlueprintCallable, Category = "Material Designer")
+	void SetDefaultValue(bool bInDefaultValue);
+#endif
+ 
+	//~ Begin UDMMaterialValue
+	virtual void SetMIDParameter(UMaterialInstanceDynamic* InMID) const override;
+#if WITH_EDITOR
+	virtual void GenerateExpression(const TSharedRef<IDMMaterialBuildStateInterface>& InBuildState) const override;
+	virtual bool IsDefaultValue() const override;
+	virtual void ApplyDefaultValue() override;
+	virtual void ResetDefaultValue() override;
+#endif
+	// ~End UDMMaterialValue
+ 
+protected:
+	UPROPERTY(EditInstanceOnly, BlueprintReadWrite, Getter = GetValue, Setter = SetValue, BlueprintSetter = SetValue, Category = "Material Designer",
+		meta = (AllowPrivateAccess = "true", DisplayName = "Bool"))
+	bool Value;
+ 
+#if WITH_EDITORONLY_DATA
+	UPROPERTY(EditInstanceOnly, BlueprintReadWrite, BlueprintSetter = SetDefaultValue, Category = "Material Designer",
+		meta = (AllowPrivateAccess = "true"))
+	bool bDefaultValue;
+#endif
+ 
+	UDMMaterialValueBool();
+};
diff --git a/Engine/Plugins/Experimental/DynamicMaterial/Source/DynamicMaterial/Public/Components/MaterialValues/DMMaterialValueColorAtlas.h b/Engine/Plugins/Experimental/DynamicMaterial/Source/DynamicMaterial/Public/Components/MaterialValues/DMMaterialValueColorAtlas.h
new file mode 100644
--- /dev/null
+++ b/Engine/Plugins/Experimental/DynamicMaterial/Source/DynamicMaterial/Public/Components/MaterialValues/DMMaterialValueColorAtlas.h
@@ -0,0 +1,77 @@
+// Copyright Epic Games, Inc. All Rights Reserved.
+
+#pragma once
+
+#include "Components/DMMaterialValue.h"
+#include "DMMaterialValueColorAtlas.generated.h"
+
+#if WITH_EDITOR
+class UCurveLinearColorAtlas;
+class UCurveLinearColor;
+#endif
+
+UCLASS(BlueprintType, ClassGroup = "Material Designer")
+class DYNAMICMATERIAL_API UDMMaterialValueColorAtlas : public UDMMaterialValue
+{
+	GENERATED_BODY()
+
+public:
+	UFUNCTION(BlueprintPure, Category = "Material Designer")
+	float GetValue() const { return Value; }
+
+	UFUNCTION(BlueprintCallable, Category = "Material Designer")
+	void SetValue(float InValue);
+
+#if WITH_EDITOR
+	UFUNCTION(BlueprintPure, Category = "Material Designer")
+	float GetDefaultValue() const { return DefaultValue; }
+
+	UFUNCTION(BlueprintCallable, Category = "Material Designer")
+	void SetDefaultValue(float InDefaultValue);
+
+	UFUNCTION(BlueprintPure, Category = "Material Designer")
+	UCurveLinearColorAtlas* GetAtlas() const { return Atlas; }
+
+	UFUNCTION(BlueprintCallable, Category = "Material Designer")
+	void SetAtlas(UCurveLinearColorAtlas* InAtlas);
+
+	UFUNCTION(BlueprintPure, Category = "Material Designer")
+	UCurveLinearColor* GetCurve() const { return Curve; }
+
+	UFUNCTION(BlueprintCallable, Category = "Material Designer")
+	void SetCurve(UCurveLinearColor* InCurve);
+
+	virtual bool IsWholeLayerValue() const override { return true; }
+#endif
+
+	//~ Begin UDMMaterialValue
+	virtual void SetMIDParameter(UMaterialInstanceDynamic* InMID) const override;
+#if WITH_EDITOR
+	virtual void GenerateExpression(const TSharedRef<IDMMaterialBuildStateInterface>& InBuildState) const override;
+	virtual bool IsDefaultValue() const override;
+	virtual void ApplyDefaultValue() override;
+	virtual void ResetDefaultValue() override;
+#endif
+	//~ End UDMMaterialValue
+
+protected:
+	UPROPERTY(EditInstanceOnly, BlueprintReadWrite, Getter = GetValue, Setter = SetValue, BlueprintSetter = SetValue, Category = "Material Designer",
+		meta = (DisplayName = "Alpha", AllowPrivateAccess = "true", ClampMin = 0, UIMin = 0, ClampMax = 1, UIMax = 1))
+	float Value;
+
+#if WITH_EDITORONLY_DATA
+	UPROPERTY(EditInstanceOnly, BlueprintReadWrite, BlueprintSetter = SetDefaultValue, Category = "Material Designer",
+		meta = (AllowPrivateAccess = "true", ClampMin = 0, UIMin = 0, ClampMax = 1, UIMax = 1))
+	float DefaultValue;
+
+	UPROPERTY(EditInstanceOnly, BlueprintReadWrite, BlueprintSetter = SetAtlas, Category = "Material Designer",
+		meta = (AllowPrivateAccess = "true"))
+	TObjectPtr<UCurveLinearColorAtlas> Atlas;
+
+	UPROPERTY(EditInstanceOnly, BlueprintReadWrite, BlueprintSetter = SetCurve, Category = "Material Designer",
+		meta = (AllowPrivateAccess = "true"))
+	TObjectPtr<UCurveLinearColor> Curve;
+#endif
+
+	UDMMaterialValueColorAtlas();
+};
diff --git a/Engine/Plugins/Experimental/DynamicMaterial/Source/DynamicMaterial/Public/Components/MaterialValues/DMMaterialValueFloat.h b/Engine/Plugins/Experimental/DynamicMaterial/Source/DynamicMaterial/Public/Components/MaterialValues/DMMaterialValueFloat.h
new file mode 100644
--- /dev/null
+++ b/Engine/Plugins/Experimental/DynamicMaterial/Source/DynamicMaterial/Public/Components/MaterialValues/DMMaterialValueFloat.h
@@ -0,0 +1,30 @@
+// Copyright Epic Games, Inc. All Rights Reserved.
+ 
+#pragma once
+ 
+#include "Components/DMMaterialValue.h"
+#include "DMMaterialValueFloat.generated.h"
+ 
+UCLASS(Abstract, BlueprintType, ClassGroup = "Material Designer")
+class DYNAMICMATERIAL_API UDMMaterialValueFloat : public UDMMaterialValue
+{
+	GENERATED_BODY()
+ 
+public:
+	UFUNCTION(BlueprintPure, Category = "Material Designer")
+	const FFloatInterval& GetValueRange() const { return ValueRange; }
+ 
+	UFUNCTION(BlueprintPure, Category = "Material Designer")
+	bool HasValueRange() const { return ValueRange.Min != ValueRange.Max; }
+ 
+	UFUNCTION(BlueprintCallable, Category = "Material Designer")
+	void SetValueRange(const FFloatInterval& InValueRange) { ValueRange = InValueRange; }
+ 
+protected:
+	UPROPERTY(VisibleInstanceOnly, BlueprintReadWrite, Getter = GetValueRange, Setter = SetValueRange, BlueprintSetter = SetValueRange, Category = "Material Designer",
+		meta = (AllowPrivateAccess = "true", DisplayName = "Value Range"))
+	FFloatInterval ValueRange;
+ 
+	UDMMaterialValueFloat(); 
+	UDMMaterialValueFloat(EDMValueType InValueType);
+};
diff --git a/Engine/Plugins/Experimental/DynamicMaterial/Source/DynamicMaterial/Public/Components/MaterialValues/DMMaterialValueFloat1.h b/Engine/Plugins/Experimental/DynamicMaterial/Source/DynamicMaterial/Public/Components/MaterialValues/DMMaterialValueFloat1.h
new file mode 100644
--- /dev/null
+++ b/Engine/Plugins/Experimental/DynamicMaterial/Source/DynamicMaterial/Public/Components/MaterialValues/DMMaterialValueFloat1.h
@@ -0,0 +1,50 @@
+// Copyright Epic Games, Inc. All Rights Reserved.
+ 
+#pragma once
+ 
+#include "DMMaterialValueFloat.h"
+#include "DMMaterialValueFloat1.generated.h"
+ 
+UCLASS(BlueprintType, ClassGroup = "Material Designer")
+class DYNAMICMATERIAL_API UDMMaterialValueFloat1 : public UDMMaterialValueFloat
+{
+	GENERATED_BODY()
+ 
+public:
+	UFUNCTION(BlueprintPure, Category = "Material Designer")
+	float GetValue() const { return Value; }
+ 
+	UFUNCTION(BlueprintCallable, Category = "Material Designer")
+	void SetValue(float InValue);
+ 
+#if WITH_EDITOR
+	UFUNCTION(BlueprintPure, Category = "Material Designer")
+	float GetDefaultValue() const { return DefaultValue; }
+ 
+	UFUNCTION(BlueprintCallable, Category = "Material Designer")
+	void SetDefaultValue(float InDefaultValue);
+#endif
+ 
+	//~ Begin UDMMaterialValue
+	virtual void SetMIDParameter(UMaterialInstanceDynamic* InMID) const override;
+#if WITH_EDITOR
+	virtual void GenerateExpression(const TSharedRef<IDMMaterialBuildStateInterface>& InBuildState) const override;
+	virtual bool IsDefaultValue() const override;
+	virtual void ApplyDefaultValue() override;
+	virtual void ResetDefaultValue() override;
+#endif
+	// ~End UDMMaterialValue
+ 
+protected:
+	UPROPERTY(EditInstanceOnly, BlueprintReadWrite, Getter = GetValue, Setter = SetValue, BlueprintSetter = SetValue, Category = "Material Designer",
+		meta = (AllowPrivateAccess = "true"), DisplayName = "Float")
+	float Value;
+ 
+#if WITH_EDITORONLY_DATA
+	UPROPERTY(EditInstanceOnly, BlueprintReadWrite, BlueprintSetter = SetDefaultValue, Category = "Material Designer",
+		meta = (AllowPrivateAccess = "true"))
+	float DefaultValue;
+#endif
+ 
+	UDMMaterialValueFloat1();
+};
diff --git a/Engine/Plugins/Experimental/DynamicMaterial/Source/DynamicMaterial/Public/Components/MaterialValues/DMMaterialValueFloat2.h b/Engine/Plugins/Experimental/DynamicMaterial/Source/DynamicMaterial/Public/Components/MaterialValues/DMMaterialValueFloat2.h
new file mode 100644
--- /dev/null
+++ b/Engine/Plugins/Experimental/DynamicMaterial/Source/DynamicMaterial/Public/Components/MaterialValues/DMMaterialValueFloat2.h
@@ -0,0 +1,51 @@
+// Copyright Epic Games, Inc. All Rights Reserved.
+ 
+#pragma once
+ 
+#include "DMMaterialValueFloat.h"
+#include "DMMaterialValueFloat2.generated.h"
+ 
+UCLASS(BlueprintType, ClassGroup = "Material Designer")
+class DYNAMICMATERIAL_API UDMMaterialValueFloat2 : public UDMMaterialValueFloat
+{
+	GENERATED_BODY()
+ 
+public: 
+	UFUNCTION(BlueprintPure, Category = "Material Designer")
+	const FVector2D& GetValue() const { return Value; }
+ 
+	UFUNCTION(BlueprintCallable, Category = "Material Designer")
+	void SetValue(const FVector2D& InValue);
+ 
+#if WITH_EDITOR
+	UFUNCTION(BlueprintPure, Category = "Material Designer")
+	FVector2D GetDefaultValue() const { return DefaultValue; }
+ 
+	UFUNCTION(BlueprintCallable, Category = "Material Designer")
+	void SetDefaultValue(FVector2D InDefaultValue);
+#endif
+ 
+	//~ Begin UDMMaterialValue
+	virtual void SetMIDParameter(UMaterialInstanceDynamic* InMID) const override;
+#if WITH_EDITOR
+	virtual void GenerateExpression(const TSharedRef<IDMMaterialBuildStateInterface>& InBuildState) const override;
+	virtual int32 GetInnateMaskOutput(int32 OutputChannels) const override;
+	virtual bool IsDefaultValue() const override;
+	virtual void ApplyDefaultValue() override;
+	virtual void ResetDefaultValue() override;
+#endif
+	// ~End UDMMaterialValue
+ 
+protected: 
+	UPROPERTY(EditInstanceOnly, BlueprintReadWrite, Getter = GetValue, Setter = SetValue, BlueprintSetter = SetValue, Category = "Material Designer",
+		meta = (AllowPrivateAccess = "true", DisplayName = "Vector"))
+	FVector2D Value;
+ 
+#if WITH_EDITORONLY_DATA
+	UPROPERTY(EditInstanceOnly, BlueprintReadWrite, BlueprintSetter = SetDefaultValue, Category = "Material Designer",
+		meta = (AllowPrivateAccess = "true"))
+	FVector2D DefaultValue;
+#endif
+ 
+	UDMMaterialValueFloat2();
+};
diff --git a/Engine/Plugins/Experimental/DynamicMaterial/Source/DynamicMaterial/Public/Components/MaterialValues/DMMaterialValueFloat3RGB.h b/Engine/Plugins/Experimental/DynamicMaterial/Source/DynamicMaterial/Public/Components/MaterialValues/DMMaterialValueFloat3RGB.h
new file mode 100644
--- /dev/null
+++ b/Engine/Plugins/Experimental/DynamicMaterial/Source/DynamicMaterial/Public/Components/MaterialValues/DMMaterialValueFloat3RGB.h
@@ -0,0 +1,52 @@
+// Copyright Epic Games, Inc. All Rights Reserved.
+ 
+#pragma once
+ 
+#include "DMMaterialValueFloat.h"
+#include "DMMaterialValueFloat3RGB.generated.h"
+ 
+UCLASS(BlueprintType, ClassGroup = "Material Designer")
+class DYNAMICMATERIAL_API UDMMaterialValueFloat3RGB : public UDMMaterialValueFloat
+{
+	GENERATED_BODY()
+ 
+public: 
+	UFUNCTION(BlueprintPure, Category = "Material Designer")
+	const FLinearColor& GetValue() const { return Value; }
+ 
+	UFUNCTION(BlueprintCallable, Category = "Material Designer")
+	void SetValue(const FLinearColor& InValue);
+ 
+#if WITH_EDITOR
+	UFUNCTION(BlueprintPure, Category = "Material Designer")
+	FLinearColor GetDefaultValue() const { return DefaultValue; }
+ 
+	UFUNCTION(BlueprintCallable, Category = "Material Designer")
+	void SetDefaultValue(FLinearColor InDefaultValue);
+#endif
+ 
+	//~ Begin UDMMaterialValue
+	virtual void SetMIDParameter(UMaterialInstanceDynamic* InMID) const override;
+#if WITH_EDITOR
+	virtual void GenerateExpression(const TSharedRef<IDMMaterialBuildStateInterface>& InBuildState) const override;
+	virtual int32 GetInnateMaskOutput(int32 OutputChannels) const override;
+	virtual bool IsDefaultValue() const override;
+	virtual void ApplyDefaultValue() override;
+	virtual void ResetDefaultValue() override;
+	virtual FName GetMainPropertyName() const override { return ValueName; }
+	//~ End UDMMaterialValue
+#endif
+
+protected: 
+	UPROPERTY(EditInstanceOnly, BlueprintReadWrite, Getter = GetValue, Setter = SetValue, BlueprintSetter = SetValue, Category = "Material Designer",
+		meta = (AllowPrivateAccess = "true", DisplayName = "Color", HideAlphaChannel))
+	FLinearColor Value;
+ 
+#if WITH_EDITORONLY_DATA
+	UPROPERTY(EditInstanceOnly, BlueprintReadWrite, BlueprintSetter = SetDefaultValue, Category = "Material Designer",
+		meta = (AllowPrivateAccess = "true"))
+	FLinearColor DefaultValue;
+#endif
+ 
+	UDMMaterialValueFloat3RGB();
+};
