[MaterialDesigner]

diff --git a/Engine/Plugins/Experimental/DynamicMaterial/Source/DynamicMaterial/Public/Components/MaterialValues/DMMaterialValueFloat3RPY.h b/Engine/Plugins/Experimental/DynamicMaterial/Source/DynamicMaterial/Public/Components/MaterialValues/DMMaterialValueFloat3RPY.h
new file mode 100644
--- /dev/null
+++ b/Engine/Plugins/Experimental/DynamicMaterial/Source/DynamicMaterial/Public/Components/MaterialValues/DMMaterialValueFloat3RPY.h
@@ -0,0 +1,51 @@
+// Copyright Epic Games, Inc. All Rights Reserved.
+ 
+#pragma once
+ 
+#include "DMMaterialValueFloat.h"
+#include "DMMaterialValueFloat3RPY.generated.h"
+ 
+UCLASS(BlueprintType, ClassGroup = "Material Designer")
+class DYNAMICMATERIAL_API UDMMaterialValueFloat3RPY : public UDMMaterialValueFloat
+{
+	GENERATED_BODY()
+ 
+public: 
+	UFUNCTION(BlueprintPure, Category = "Material Designer")
+	const FRotator& GetValue() const { return Value; }
+ 
+	UFUNCTION(BlueprintCallable, Category = "Material Designer")
+	void SetValue(const FRotator& InValue);
+ 
+#if WITH_EDITOR
+	UFUNCTION(BlueprintPure, Category = "Material Designer")
+	FRotator GetDefaultValue() const { return DefaultValue; }
+ 
+	UFUNCTION(BlueprintCallable, Category = "Material Designer")
+	void SetDefaultValue(FRotator InDefaultValue);
+#endif
+ 
+	//~ Begin UDMMaterialValue
+	virtual void SetMIDParameter(UMaterialInstanceDynamic* InMID) const override;
+#if WITH_EDITOR
+	virtual void GenerateExpression(const TSharedRef<IDMMaterialBuildStateInterface>& InBuildState) const override;
+	virtual int32 GetInnateMaskOutput(int32 OutputChannels) const override;
+	virtual bool IsDefaultValue() const override;
+	virtual void ApplyDefaultValue() override;
+	virtual void ResetDefaultValue() override;
+#endif
+	// ~End UDMMaterialValue
+ 
+protected: 
+	UPROPERTY(EditInstanceOnly, BlueprintReadWrite, Getter = GetValue, Setter = SetValue, BlueprintSetter = SetValue, Category = "Material Designer",
+		meta = (AllowPrivateAccess = "true", DisplayName = "Rotation"))
+	FRotator Value;
+ 
+#if WITH_EDITORONLY_DATA
+	UPROPERTY(EditAnywhere, BlueprintReadWrite, BlueprintSetter = SetDefaultValue, Category = "Material Designer",
+		meta = (AllowPrivateAccess = "true"))
+	FRotator DefaultValue;
+#endif
+ 
+	UDMMaterialValueFloat3RPY();
+};
diff --git a/Engine/Plugins/Experimental/DynamicMaterial/Source/DynamicMaterial/Public/Components/MaterialValues/DMMaterialValueFloat3XYZ.h b/Engine/Plugins/Experimental/DynamicMaterial/Source/DynamicMaterial/Public/Components/MaterialValues/DMMaterialValueFloat3XYZ.h
new file mode 100644
--- /dev/null
+++ b/Engine/Plugins/Experimental/DynamicMaterial/Source/DynamicMaterial/Public/Components/MaterialValues/DMMaterialValueFloat3XYZ.h
@@ -0,0 +1,51 @@
+// Copyright Epic Games, Inc. All Rights Reserved.
+ 
+#pragma once
+ 
+#include "DMMaterialValueFloat.h"
+#include "DMMaterialValueFloat3XYZ.generated.h"
+ 
+UCLASS(BlueprintType, ClassGroup = "Material Designer")
+class DYNAMICMATERIAL_API UDMMaterialValueFloat3XYZ : public UDMMaterialValueFloat
+{
+	GENERATED_BODY()
+ 
+public: 
+	UFUNCTION(BlueprintPure, Category = "Material Designer")
+	const FVector& GetValue() const { return Value; }
+ 
+	UFUNCTION(BlueprintCallable, Category = "Material Designer")
+	void SetValue(const FVector& InValue);
+ 
+#if WITH_EDITOR
+	UFUNCTION(BlueprintPure, Category = "Material Designer")
+	FVector GetDefaultValue() const { return DefaultValue; }
+ 
+	UFUNCTION(BlueprintCallable, Category = "Material Designer")
+	void SetDefaultValue(FVector InDefaultValue);
+#endif
+ 
+	//~ Begin UDMMaterialValue
+	virtual void SetMIDParameter(UMaterialInstanceDynamic* InMID) const override;
+#if WITH_EDITOR
+	virtual void GenerateExpression(const TSharedRef<IDMMaterialBuildStateInterface>& InBuildState) const override;
+	virtual int32 GetInnateMaskOutput(int32 OutputChannels) const override;
+	virtual bool IsDefaultValue() const override;
+	virtual void ApplyDefaultValue() override;
+	virtual void ResetDefaultValue() override;
+#endif
+	// ~End UDMMaterialValue
+ 
+protected: 
+	UPROPERTY(EditInstanceOnly, BlueprintReadWrite, Getter = GetValue, Setter = SetValue, BlueprintSetter = SetValue, Category = "Material Designer",
+		meta = (AllowPrivateAccess = "true"), DisplayName = "Vector")
+	FVector Value;
+ 
+#if WITH_EDITORONLY_DATA
+	UPROPERTY(EditAnywhere, BlueprintReadWrite, BlueprintSetter = SetDefaultValue, 
+		meta = (AllowPrivateAccess = "true"), Category = "Material Designer")
+	FVector DefaultValue;
+#endif
+ 
+	UDMMaterialValueFloat3XYZ();
+};
diff --git a/Engine/Plugins/Experimental/DynamicMaterial/Source/DynamicMaterial/Public/Components/MaterialValues/DMMaterialValueFloat4.h b/Engine/Plugins/Experimental/DynamicMaterial/Source/DynamicMaterial/Public/Components/MaterialValues/DMMaterialValueFloat4.h
new file mode 100644
--- /dev/null
+++ b/Engine/Plugins/Experimental/DynamicMaterial/Source/DynamicMaterial/Public/Components/MaterialValues/DMMaterialValueFloat4.h
@@ -0,0 +1,53 @@
+// Copyright Epic Games, Inc. All Rights Reserved.
+ 
+#pragma once
+ 
+#include "DMMaterialValueFloat.h"
+#include "DMMaterialValueFloat4.generated.h"
+ 
+UCLASS(BlueprintType, ClassGroup = "Material Designer")
+class DYNAMICMATERIAL_API UDMMaterialValueFloat4 : public UDMMaterialValueFloat
+{
+	GENERATED_BODY()
+ 
+public: 
+	UFUNCTION(BlueprintPure, Category = "Material Designer")
+	const FLinearColor& GetValue() const { return Value; }
+ 
+	UFUNCTION(BlueprintCallable, Category = "Material Designer")
+	void SetValue(const FLinearColor& InValue);
+ 
+#if WITH_EDITOR
+	UFUNCTION(BlueprintPure, Category = "Material Designer")
+	FLinearColor GetDefaultValue() const { return DefaultValue; }
+ 
+	UFUNCTION(BlueprintCallable, Category = "Material Designer")
+	void SetDefaultValue(FLinearColor InDefaultValue);
+#endif
+ 
+	//~ Begin UDMMaterialValue
+	virtual void SetMIDParameter(UMaterialInstanceDynamic* InMID) const override;
+#if WITH_EDITOR
+	virtual void GenerateExpression(const TSharedRef<IDMMaterialBuildStateInterface>& InBuildState) const override;
+	virtual int32 GetInnateMaskOutput(int32 OutputChannels) const override;
+	virtual bool IsDefaultValue() const override;
+	virtual void ApplyDefaultValue() override;
+	virtual void ResetDefaultValue() override;
+	virtual FName GetMainPropertyName() const override { return ValueName; }
+	virtual bool IsWholeLayerValue() const override { return true; }
+	//~ End UDMMaterialValue
+#endif
+
+protected: 
+	UPROPERTY(EditInstanceOnly, BlueprintReadWrite, Getter = GetValue, Setter = SetValue, BlueprintSetter = SetValue, Category = "Material Designer",
+		meta = (AllowPrivateAccess = "true", DisplayName = "Color"))
+	FLinearColor Value;
+ 
+#if WITH_EDITORONLY_DATA
+	UPROPERTY(EditAnywhere, BlueprintReadWrite, BlueprintSetter = SetDefaultValue, Category = "Material Designer",
+		meta = (AllowPrivateAccess = "true"))
+	FLinearColor DefaultValue;
+#endif
+ 
+	UDMMaterialValueFloat4();
+};
diff --git a/Engine/Plugins/Experimental/DynamicMaterial/Source/DynamicMaterial/Public/Components/MaterialValues/DMMaterialValueTexture.h b/Engine/Plugins/Experimental/DynamicMaterial/Source/DynamicMaterial/Public/Components/MaterialValues/DMMaterialValueTexture.h
new file mode 100644
--- /dev/null
+++ b/Engine/Plugins/Experimental/DynamicMaterial/Source/DynamicMaterial/Public/Components/MaterialValues/DMMaterialValueTexture.h
@@ -0,0 +1,85 @@
+// Copyright Epic Games, Inc. All Rights Reserved.
+ 
+#pragma once
+ 
+#include "Components/DMMaterialValue.h"
+#include "DMMaterialValueTexture.generated.h"
+ 
+class UTexture;
+
+#if WITH_EDITOR
+namespace UE::DynamicMaterial::Private
+{
+	bool DYNAMICMATERIAL_API HasAlpha(UTexture* InTexture);
+}
+
+DECLARE_DELEGATE_RetVal(UTexture*, FDMGetDefaultRGBTexture);
+#endif
+
+UCLASS(BlueprintType, ClassGroup = "Material Designer")
+class DYNAMICMATERIAL_API UDMMaterialValueTexture : public UDMMaterialValue
+{
+	GENERATED_BODY()
+ 
+public:
+#if WITH_EDITOR
+	static FDMGetDefaultRGBTexture GetDefaultRGBTexture;
+
+	UFUNCTION(BlueprintPure, Category = "Material Designer")
+	static UDMMaterialValueTexture* CreateMaterialValueTexture(UObject* Outer, UTexture* InTexture);
+#endif
+
+	UFUNCTION(BlueprintPure, Category = "Material Designer")
+	UTexture* GetValue() const { return Value; }
+ 
+	UFUNCTION(BlueprintCallable, Category = "Material Designer")
+	void SetValue(UTexture* InValue);
+
+#if WITH_EDITOR
+	UFUNCTION(BlueprintPure, Category = "Material Designer")
+	bool HasAlpha() const;
+
+	UFUNCTION(BlueprintPure, Category = "Material Designer")
+	UTexture* GetDefaultValue() const { return DefaultValue; }
+
+	UFUNCTION(BlueprintCallable, Category = "Material Designer")
+	void SetDefaultValue(UTexture* InDefaultValue);
+#endif // WITH_EDITOR
+
+	//~ Begin UObject
+#if WITH_EDITOR
+	virtual void PreEditChange(FProperty* PropertyAboutToChange) override;
+	virtual void PostEditChangeProperty(FPropertyChangedEvent& PropertyChangedEvent) override;
+#endif
+	virtual void PostLoad() override;
+	//~ End UObject
+
+	//~ Begin UDMMaterialValue
+	virtual void SetMIDParameter(UMaterialInstanceDynamic* InMID) const override;
+#if WITH_EDITOR
+	virtual FName GetMainPropertyName() const override { return ValueName; }
+	virtual void GenerateExpression(const TSharedRef<IDMMaterialBuildStateInterface>& InBuildState) const override;
+	virtual bool IsDefaultValue() const override;
+	virtual void ApplyDefaultValue() override;
+	virtual void ResetDefaultValue() override;
+#endif
+	//~ End UDMMaterialValue
+
+protected:
+	UDMMaterialValueTexture();
+ 
+	UPROPERTY(EditAnywhere, BlueprintReadWrite, Getter = GetValue, Setter = SetValue, BlueprintSetter = SetValue, Category = "Material Designer", 
+		meta = (DisplayThumbnail = true, DisplayName = "Texture", AllowPrivateAccess = "true", HighPriority, NotKeyframeable))
+	TObjectPtr<UTexture> Value;
+
+#if WITH_EDITORONLY_DATA
+	UPROPERTY(EditAnywhere, BlueprintReadWrite, BlueprintSetter = SetDefaultValue, Category = "Material Designer",
+		meta = (AllowPrivateAccess = "true"))
+	TObjectPtr<UTexture> DefaultValue;
+#endif
+
+#if WITH_EDITORONLY_DATA
+	UPROPERTY()
+	TObjectPtr<UTexture> OldValue;
+#endif
+};
diff --git a/Engine/Plugins/Experimental/DynamicMaterial/Source/DynamicMaterial/Public/DMComponentPath.h b/Engine/Plugins/Experimental/DynamicMaterial/Source/DynamicMaterial/Public/DMComponentPath.h
new file mode 100644
--- /dev/null
+++ b/Engine/Plugins/Experimental/DynamicMaterial/Source/DynamicMaterial/Public/DMComponentPath.h
@@ -0,0 +1,49 @@
+// Copyright Epic Games, Inc. All Rights Reserved.
+
+#pragma once
+
+#include "Containers/StringView.h"
+
+struct FDMComponentPathSegment;
+
+/** A path such as Name.Component.Component[2].Value */
+struct DYNAMICMATERIAL_API FDMComponentPath
+{
+	static constexpr TCHAR Separator = '.';
+	static constexpr TCHAR ParameterOpen = '[';
+	static constexpr TCHAR ParameterClose = ']';
+
+	FDMComponentPath() = delete;
+	FDMComponentPath(FStringView InPath);
+	FDMComponentPath(const FString& InPathString);
+
+	bool IsLeaf() const;
+
+	/**
+	 * Extracts the first component and removes it from the path
+	 */
+	FDMComponentPathSegment GetFirstSegment();
+
+protected:
+	FStringView Path;
+};
+
+/** Represents a single part of a component path */
+struct DYNAMICMATERIAL_API FDMComponentPathSegment
+{
+	friend struct FDMComponentPath;
+
+	FDMComponentPathSegment(FStringView InToken, FStringView InParameter);
+
+	FStringView GetToken() const { return Token; }
+
+	bool HasParameter() const;
+
+	bool GetParameter(int32& OutParameter) const;
+
+	bool GetParameter(FString& OutParameter) const;
+
+protected:
+	FStringView Token;
+	FStringView Parameter;
+};
diff --git a/Engine/Plugins/Experimental/DynamicMaterial/Source/DynamicMaterial/Public/DMDefs.h b/Engine/Plugins/Experimental/DynamicMaterial/Source/DynamicMaterial/Public/DMDefs.h
new file mode 100644
--- /dev/null
+++ b/Engine/Plugins/Experimental/DynamicMaterial/Source/DynamicMaterial/Public/DMDefs.h
@@ -0,0 +1,170 @@
+// Copyright Epic Games, Inc. All Rights Reserved.
+
+#pragma once
+
+#include "Containers/Array.h"
+#include "Internationalization/Text.h"
+#include "Engine/EngineTypes.h"
+#include "DMDefs.generated.h"
+
+class UTexture;
+class UDMMaterialStage;
+
+UENUM(BlueprintType)
+enum class EDMMaterialPropertyType : uint8
+{
+	None = 0,
+	BaseColor,
+	EmissiveColor,
+	Opacity,
+	OpacityMask,
+	Roughness,
+	Specular,
+	Metallic,
+	Normal,
+	PixelDepthOffset,
+	WorldPositionOffset,
+	AmbientOcclusion,
+	Anisotropy,
+	Refraction,
+	Tangent,
+	Custom1,
+	Custom2,
+	Custom3,
+	Custom4,
+	Any
+};
+
+UENUM(BlueprintType)
+enum class EDMValueType : uint8
+{
+	VT_None,
+	VT_Bool,
+	VT_Float1,
+	VT_Float2,
+	VT_Float3_RPY,
+	VT_Float3_RGB,
+	VT_Float3_XYZ,
+	VT_Float4_RGBA,
+	VT_Float_Any,
+	VT_Texture,
+	VT_ColorAtlas
+};
+
+UENUM(BlueprintType)
+enum class EDMUpdateType : uint8
+{
+	Value,
+	Structure
+};
+
+UENUM(BlueprintType)
+enum class EDMMaterialShadingModel : uint8
+{
+	Unlit      = EMaterialShadingModel::MSM_Unlit,
+	DefaultLit = EMaterialShadingModel::MSM_DefaultLit
+};
+
+UENUM(BlueprintType, meta = (DisplayName = "Material Designer UV Source"))
+enum class EDMUVSource : uint8
+{
+	Texture,
+	ScreenPosition,
+	WorldPosition
+};
+
+/**
+ * An individual component of a connector (e.g. G from RGB.)
+ */
+USTRUCT(BlueprintType, Category = "Material Designer", meta = (DisplayName = "Material Designer Stage Connector Channel"))
+struct DYNAMICMATERIAL_API FDMMaterialStageConnectorChannel
+{
+	GENERATED_BODY()
+
+	static constexpr int32 NO_SOURCE = -1;
+	static constexpr int32 PREVIOUS_STAGE = 0;
+	static constexpr int32 FIRST_STAGE_INPUT = 1;
+	static constexpr int32 WHOLE_CHANNEL = 0;
+	static constexpr int32 FIRST_CHANNEL = 1;
+	static constexpr int32 SECOND_CHANNEL = 2;
+	static constexpr int32 THIRD_CHANNEL = 4;
+	static constexpr int32 FOURTH_CHANNEL = 8;
+	static constexpr int32 TWO_CHANNELS = FIRST_CHANNEL | SECOND_CHANNEL;
+	static constexpr int32 THREE_CHANNELS = FIRST_CHANNEL | SECOND_CHANNEL | THIRD_CHANNEL;
+	/** Not really needed? Effectively the whole channel! */
+	static constexpr int32 FOUR_CHANNELS = FIRST_CHANNEL | SECOND_CHANNEL | THIRD_CHANNEL | FOURTH_CHANNEL;
+
+	/**
+	 * The index of the source of this channel
+	 * Index 0 is the previous stage, 1+ are the other inputs required by the current stage (e.g. textures, uvs, etc.)
+	 */
+	UPROPERTY(VisibleInstanceOnly, BlueprintReadOnly, Category = "Material Designer")
+	int32 SourceIndex = PREVIOUS_STAGE;
+
+	/** When using previous stages, this is the material property the previous stage is using */
+	UPROPERTY(VisibleInstanceOnly, BlueprintReadOnly, Category = "Material Designer")
+	EDMMaterialPropertyType MaterialProperty = EDMMaterialPropertyType::None;
+
+	/** The index of the output connector of the given stage. */
+	UPROPERTY(VisibleInstanceOnly, BlueprintReadOnly, Category = "Material Designer")
+	int32 OutputIndex = 0;
+
+	/**
+	 * This can be used to break down float2/3/4 into single pieces of data
+	 * A value of 0 will be the original output. A bitmask (1,2,4,8) will reference (and combine) the specific channels.
+	 */
+	UPROPERTY(VisibleInstanceOnly, BlueprintReadOnly, Category = "Material Designer")
+	int32 OutputChannel = WHOLE_CHANNEL;
+
+	bool operator==(const FDMMaterialStageConnectorChannel& Other) const
+	{
+		return MaterialProperty == Other.MaterialProperty
+			&& SourceIndex == Other.SourceIndex
+			&& OutputIndex == Other.OutputIndex
+			&& OutputChannel == Other.OutputChannel;
+	}
+};
+
+struct DYNAMICMATERIAL_API FDMUpdateGuard
+{
+	FDMUpdateGuard()
+	{
+		++GuardCount;
+	}
+
+	virtual ~FDMUpdateGuard()
+	{
+		--GuardCount;
+	}
+
+	static bool CanUpdate()
+	{
+		return (GuardCount == 0);
+	}
+
+private:
+	static int32 GuardCount;
+};
+
+struct DYNAMICMATERIAL_API FDMInitializationGuard
+{
+public:
+	static bool IsInitializing();
+
+	FDMInitializationGuard();
+	virtual ~FDMInitializationGuard();
+
+private:
+	static uint32 GuardCount;
+};
+
+namespace UE::DynamicMaterial
+{
+	/** Designed to be used with preprocessor macros. See below. */
+	FString DYNAMICMATERIAL_API CreateNodeComment(const ANSICHAR* InFile, int InLine, const ANSICHAR* InFunction, const FString* InComment = nullptr);
+
+	constexpr int32 RenameFlags = REN_DontCreateRedirectors | REN_DoNotDirty | REN_ForceNoResetLoaders | REN_NonTransactional;
+}
+
+#define UE_DM_NodeComment_Default UE::DynamicMaterial::CreateNodeComment(__FILE__, __LINE__, __FUNCTION__)
+#define UE_DM_NodeComment(Comment) UE::DynamicMaterial::CreateNodeComment(__FILE__, __LINE__, __FUNCTION__, &Comment)
diff --git a/Engine/Plugins/Experimental/DynamicMaterial/Source/DynamicMaterial/Public/DMValueDefinition.h b/Engine/Plugins/Experimental/DynamicMaterial/Source/DynamicMaterial/Public/DMValueDefinition.h
new file mode 100644
--- /dev/null
+++ b/Engine/Plugins/Experimental/DynamicMaterial/Source/DynamicMaterial/Public/DMValueDefinition.h
@@ -0,0 +1,94 @@
+// Copyright Epic Games, Inc. All Rights Reserved.
+
+#pragma once
+
+#include "DMDefs.h"
+#include "Engine/EngineTypes.h"
+#include "Internationalization/Text.h"
+#include "Kismet/BlueprintFunctionLibrary.h"
+#include "DMValueDefinition.generated.h"
+
+USTRUCT(BlueprintType)
+struct DYNAMICMATERIAL_API FDMValueDefinition
+{
+	GENERATED_BODY()
+
+	FDMValueDefinition()
+		: FDMValueDefinition(EDMValueType::VT_None, 0, FText::GetEmpty(), {})
+	{		
+	}
+
+	FDMValueDefinition(EDMValueType InType, uint8 InFloatCount, const FText& InDisplayName, const TArray<FText>& InChannelNames)
+		: Type(InType)
+		, FloatCount(InFloatCount)
+		, DisplayName(InDisplayName)
+		, ChannelNames(InChannelNames)
+	{
+	}
+
+	FDMValueDefinition(EDMValueType InType, uint8 InFloatCount, FText&& InDisplayName, TArray<FText>&& InChannelNames)
+		: Type(InType)
+		, FloatCount(InFloatCount)
+		, DisplayName(MoveTemp(InDisplayName))
+		, ChannelNames(MoveTemp(InChannelNames))
+	{
+	}
+
+	EDMValueType GetType() const { return Type; }
+
+	uint8 GetFloatCount() const { return FloatCount; }
+
+	const FText& GetDisplayName() const { return DisplayName; }
+
+	const TArray<FText>& GetChannelNames() const { return ChannelNames; }
+
+	bool IsFloatType() const;
+
+	bool IsFloat3Type() const;
+
+	/** To be consistent without OutputChannel, 1 is the first channel, not 0. */
+	const FText& GetChannelName(int32 InChannel) const;
+
+protected:
+	UPROPERTY(VisibleAnywhere, BlueprintReadOnly, Category = "Material Designer")
+	EDMValueType Type;
+
+	UPROPERTY(VisibleAnywhere, BlueprintReadOnly, Category = "Material Designer")
+	uint8 FloatCount;
+
+	UPROPERTY(VisibleAnywhere, BlueprintReadOnly, Category = "Material Designer")
+	FText DisplayName;
+
+	UPROPERTY(VisibleAnywhere, BlueprintReadOnly, Category = "Material Designer")
+	TArray<FText> ChannelNames;
+};
+
+UCLASS(BlueprintType)
+class DYNAMICMATERIAL_API UDMValueDefinitionLibrary : public UBlueprintFunctionLibrary
+{
+	GENERATED_BODY()
+
+public:
+	UFUNCTION(BlueprintPure, CallInEditor, Category = "Material Designer")
+	static const TArray<EDMValueType>& GetValueTypes();
+
+	UFUNCTION(BlueprintPure, CallInEditor, Category = "Material Designer")
+	static const FDMValueDefinition& GetValueDefinition(EDMValueType InValueType);
+
+	UFUNCTION(BlueprintPure, CallInEditor, Category = "Material Designer", Meta = (DisplayName = "Are Types Compatible"))
+	static bool BP_AreTypesCompatible(EDMValueType A, EDMValueType B, int32 AChannel, int32 BChannel)
+	{
+		return AreTypesCompatible(A, B, AChannel, BChannel);
+	}
+
+	static bool AreTypesCompatible(EDMValueType A, EDMValueType B,
+		int32 AChannel = FDMMaterialStageConnectorChannel::WHOLE_CHANNEL,
+		int32 BChannel = FDMMaterialStageConnectorChannel::WHOLE_CHANNEL);
+
+	/** Converts a number of floats into the value type */
+	static const FDMValueDefinition& GetTypeForFloatCount(uint8 FloatCount);
+
+	/** Converts a number of floats into the value type */
+	UFUNCTION(BlueprintPure, CallInEditor, Category = "Material Designer")
+	static const FDMValueDefinition& GetTypeForFloatCount(int32 FloatCount);
+};
diff --git a/Engine/Plugins/Experimental/DynamicMaterial/Source/DynamicMaterial/Public/DynamicMaterialModule.h b/Engine/Plugins/Experimental/DynamicMaterial/Source/DynamicMaterial/Public/DynamicMaterialModule.h
new file mode 100644
--- /dev/null
+++ b/Engine/Plugins/Experimental/DynamicMaterial/Source/DynamicMaterial/Public/DynamicMaterialModule.h
@@ -0,0 +1,54 @@
+// Copyright Epic Games, Inc. All Rights Reserved.
+
+#pragma once
+
+#include "Delegates/Delegate.h"
+#include "Delegates/DelegateCombinations.h"
+#include "Modules/ModuleInterface.h"
+
+#if WITH_EDITOR
+#include "UObject/ScriptInterface.h"
+#endif
+
+#if WITH_EDITOR
+class IDynamicMaterialModelEditorOnlyDataInterface;
+class UDynamicMaterialModel;
+#endif
+
+DECLARE_LOG_CATEGORY_EXTERN(LogDynamicMaterial, Log, All);
+
+#if WITH_EDITOR
+DECLARE_DELEGATE_RetVal_OneParam(TScriptInterface<IDynamicMaterialModelEditorOnlyDataInterface>, FDMCreateEditorOnlyDataDelegate, UDynamicMaterialModel*)
+#endif
+
+/**
+ * Material Designer - Build your own materials in a slimline editor!
+ */
+class DYNAMICMATERIAL_API FDynamicMaterialModule : public IModuleInterface
+{
+public:
+	static FDynamicMaterialModule& Get();
+	static bool AreUObjectsSafe();
+	static bool IsMaterialExportEnabled();
+
+#if WITH_EDITOR
+	static TScriptInterface<IDynamicMaterialModelEditorOnlyDataInterface> CreateEditorOnlyData(UDynamicMaterialModel* InMaterialModel);
+	static FDMCreateEditorOnlyDataDelegate& GetCreateEditorOnlyDataDelegate() { return CreateEditorOnlyDataDelegate; }
+#endif
+
+	//~ Begin IDynamicMaterialModule
+	virtual void StartupModule() override;
+	virtual void ShutdownModule() override;
+	//~ End IDynamicMaterialModule
+
+protected:
+	static bool bIsEngineExiting;
+
+#if WITH_EDITOR
+	static FDMCreateEditorOnlyDataDelegate CreateEditorOnlyDataDelegate;
+#endif
+	
+	static void HandleEnginePreExit();
+
+	FDelegateHandle EnginePreExitHandle;
+};
diff --git a/Engine/Plugins/Experimental/DynamicMaterial/Source/DynamicMaterial/Public/Material/DynamicMaterialInstance.h b/Engine/Plugins/Experimental/DynamicMaterial/Source/DynamicMaterial/Public/Material/DynamicMaterialInstance.h
new file mode 100644
--- /dev/null
+++ b/Engine/Plugins/Experimental/DynamicMaterial/Source/DynamicMaterial/Public/Material/DynamicMaterialInstance.h
@@ -0,0 +1,45 @@
+// Copyright Epic Games, Inc. All Rights Reserved.
+
+#pragma once
+
+#include "Materials/MaterialInstanceDynamic.h"
+#include "DynamicMaterialInstance.generated.h"
+
+class UDynamicMaterialModel;
+
+#if WITH_EDITOR
+class UDMMaterialStageInputTextureUV;
+class UDMMaterialValue;
+#endif
+
+class UMaterial;
+
+UCLASS(ClassGroup = "Material Designer", DefaultToInstanced, BlueprintType, meta = (DisplayThumbnail = "true"))
+class DYNAMICMATERIAL_API UDynamicMaterialInstance : public UMaterialInstanceDynamic
+{
+	GENERATED_BODY()
+
+public:
+	UDynamicMaterialInstance();
+
+	UDynamicMaterialModel* GetMaterialModel() { return MaterialModel; }
+
+#if WITH_EDITOR
+	void SetMaterialModel(UDynamicMaterialModel* InMaterialModel) { MaterialModel = InMaterialModel; }
+
+	void OnMaterialBuilt(UDynamicMaterialModel* InMaterialModel);
+	void InitializeMIDPublic();
+
+	//~ Begin UObject
+	virtual void PostDuplicate(bool bDuplicateForPIE) override;
+	virtual void PostEditImport() override;
+	//~ End UObject
+#endif
+
+protected:
+	UPROPERTY(VisibleInstanceOnly, BlueprintReadOnly, Instanced, DuplicateTransient, TextExportTransient, Category = "Material Designer")
+	TObjectPtr<UMaterial> BaseMaterial;
+
+	UPROPERTY(VisibleInstanceOnly, BlueprintReadOnly, Instanced, Category = "Material Designer")
+	TObjectPtr<UDynamicMaterialModel> MaterialModel;
+};
diff --git a/Engine/Plugins/Experimental/DynamicMaterial/Source/DynamicMaterial/Public/Model/DynamicMaterialModel.h b/Engine/Plugins/Experimental/DynamicMaterial/Source/DynamicMaterial/Public/Model/DynamicMaterialModel.h
new file mode 100644
--- /dev/null
+++ b/Engine/Plugins/Experimental/DynamicMaterial/Source/DynamicMaterial/Public/Model/DynamicMaterialModel.h
@@ -0,0 +1,197 @@
+// Copyright Epic Games, Inc. All Rights Reserved.
+
+#pragma once
+
+#include "UObject/Object.h"
+#include "Containers/ContainersFwd.h"
+
+#if WITH_EDITOR
+#include "UObject/ScriptInterface.h"
+#endif
+
+#include "DynamicMaterialModel.generated.h"
+
+class UDMMaterialComponent;
+class UDMMaterialParameter;
+class UDMMaterialValue;
+class UDMMaterialValueFloat1;
+class UDMTextureUV;
+class UDynamicMaterialInstance;
+enum class EDMMaterialPropertyType : uint8;
+enum class EDMMaterialShadingModel : uint8;
+enum class EDMUpdateType : uint8;
+enum class EDMValueType : uint8;
+struct FDMComponentPath;
+
+#if WITH_EDITORONLY_DATA
+class IDynamicMaterialModelEditorOnlyDataInterface;
+class UDMMaterialProperty;
+class UDMMaterialSlot;
+class UMaterial;
+class UMaterialExpression;
+enum EBlendMode : int;
+#endif
+
+UCLASS(ClassGroup = "Material Designer", DefaultToInstanced, BlueprintType, meta = (DisplayThumbnail = "true"))
+class DYNAMICMATERIAL_API UDynamicMaterialModel : public UObject
+{
+	GENERATED_BODY()
+
+public:
+	static const FString ValuesPathToken;
+	static const FString ParametersPathToken;
+	static const FName GlobalOpacityParameterName;
+
+	UDynamicMaterialModel();
+
+	bool IsModelValid() const;
+
+	UDynamicMaterialInstance* GetDynamicMaterialInstance() const { return DynamicMaterialInstance; }
+	void SetDynamicMaterialInstance(UDynamicMaterialInstance* InDynamicMaterialInstance);
+
+	UFUNCTION(BlueprintPure, Category = "Material Designer")
+	UMaterial* GetGeneratedMaterial() const { return DynamicMaterial; }
+
+	UFUNCTION(BlueprintPure, Category = "Material Designer")
+	UDMMaterialValueFloat1* GetGlobalOpacityValue() const { return GlobalOpacityValue; }
+
+	UFUNCTION(BlueprintCallable, Category = "Material Designer")
+	UDMMaterialComponent* GetComponentByPath(const FString& InPath) const;
+
+	UDMMaterialComponent* GetComponentByPath(FDMComponentPath& InPath) const;
+
+	template<typename InComponentClass>
+	UDMMaterialComponent* GetComponentByPath(FDMComponentPath& InPath) const
+	{
+		return Cast<InComponentClass>(GetComponentByPath(InPath));
+	}
+
+#if WITH_EDITOR
+	friend class UDynamicMaterialModelEditorOnlyData;
+	friend class UDynamicMaterialModelFactory;
+
+	UFUNCTION(BlueprintPure, Category = "Material Designer", Meta = (DisplayName = "Get Editor Only Data"))
+	TScriptInterface<IDynamicMaterialModelEditorOnlyDataInterface> BP_GetEditorOnlyData() const { return EditorOnlyDataSI; }
+
+	IDynamicMaterialModelEditorOnlyDataInterface* GetEditorOnlyData() const;
+
+	UFUNCTION(BlueprintPure, Category = "Material Designer")
+	const TArray<UDMMaterialValue*>& GetValues() const { return Values; }
+
+	UFUNCTION(BlueprintPure, Category = "Material Designer")
+	UDMMaterialValue* GetValueByName(FName InName) const;
+
+	UFUNCTION(BlueprintPure, Category = "Material Designer")
+	UDMMaterialValue* GetValueByIndex(int32 Index) const;
+
+	UFUNCTION(BlueprintCallable, Category = "Material Designer")
+	UDMMaterialValue* AddValue(EDMValueType InType);
+
+	UFUNCTION(BlueprintCallable, Category = "Material Designer")
+	void RemoveValueByName(FName InName);
+
+	UFUNCTION(BlueprintCallable, Category = "Material Designer")
+	void RemoveValueByIndex(int32 Index);
+
+	UFUNCTION(BlueprintPure, Category = "Material Designer")
+	bool HasParameterName(FName InParameterName) const;
+
+	/** Creates a new parameter and assigns it a unique name. */
+	UFUNCTION(BlueprintCallable, Category = "Material Designer")
+	UDMMaterialParameter* CreateUniqueParameter(FName InBaseName);
+
+	/** Updates the name on an existing parameter. */
+	UFUNCTION(BlueprintCallable, Category = "Material Designer")
+	void RenameParameter(UDMMaterialParameter* InParameter, FName InBaseName);
+
+	/** Removes parameter by the name assigned to this parameter object. */
+	UFUNCTION(BlueprintCallable, Category = "Material Designer")
+	void FreeParameter(UDMMaterialParameter* InParameter);
+
+	/**
+	 * Removes this specific object from the parameter map if the name is in use by a different parameter.
+	 * Returns true if, after this call, the object is not in the parameter map.
+	 */
+	UFUNCTION(BlueprintCallable, Category = "Material Designer")
+	bool ConditionalFreeParameter(UDMMaterialParameter* InParameter);
+
+	UFUNCTION(BlueprintCallable, Category = "Material Designer")
+	void ResetData();
+
+	//~ Begin UObject
+	virtual void PostLoad() override;
+	virtual void PostEditUndo() override;
+	virtual void PostEditImport() override;
+	virtual void PostDuplicate(bool bDuplicateForPIE) override;
+	//~ End UObject
+
+	void PostEditorDuplicate();
+#endif
+
+protected:
+	UPROPERTY(VisibleInstanceOnly, BlueprintReadOnly, Instanced, Category = "Material Designer")
+	TArray<TObjectPtr<UDMMaterialValue>> Values;
+
+	UPROPERTY(VisibleInstanceOnly, BlueprintReadOnly, Category = "Material Designer")
+	TObjectPtr<UDMMaterialValueFloat1> GlobalOpacityValue;
+
+	UPROPERTY(VisibleInstanceOnly, BlueprintReadOnly, Category = "Material Designer")
+	TObjectPtr<UDMMaterialParameter> GlobalOpacityParameter;
+
+	UPROPERTY(VisibleInstanceOnly, DuplicateTransient, TextExportTransient, Category = "Material Designer")
+	TMap<FName, TWeakObjectPtr<UDMMaterialParameter>> ParameterMap;
+
+	UPROPERTY(VisibleInstanceOnly, Instanced, BlueprintReadOnly, DuplicateTransient, TextExportTransient, Category = "Material Designer")
+	TObjectPtr<UMaterial> DynamicMaterial = nullptr;
+
+	UPROPERTY(VisibleInstanceOnly, BlueprintReadOnly, DuplicateTransient, TextExportTransient, Category = "Material Designer")
+	TObjectPtr<UDynamicMaterialInstance> DynamicMaterialInstance;
+
+#if WITH_EDITORONLY_DATA
+	UPROPERTY(VisibleInstanceOnly, BlueprintReadOnly, Category = "Material Designer")
+	TScriptInterface<IDynamicMaterialModelEditorOnlyDataInterface> EditorOnlyDataSI;
+#endif
+
+#if WITH_EDITOR
+	FName CreateUniqueParameterName(FName InBaseName);
+
+	void ReinitComponents();
+
+	void FixGlobalOpacityVars();
+#endif
+
+	/**
+	 * These properties used to exist on the editor-only subclass of this model.
+	 * They now exist on the editor-only component of this model.
+	 */
+
+#if WITH_EDITORONLY_DATA
+	UE_DEPRECATED(5.3, "Moved to editor-only subobject.")
+	UPROPERTY()
+	TEnumAsByte<EBlendMode> BlendMode;
+
+	UE_DEPRECATED(5.3, "Moved to editor-only subobject.")
+	UPROPERTY()
+	EDMMaterialShadingModel ShadingModel;
+
+	UE_DEPRECATED(5.3, "Moved to editor-only subobject.")
+	UPROPERTY()
+	TMap<EDMMaterialPropertyType, TObjectPtr<UObject>> Properties;
+
+	UE_DEPRECATED(5.3, "Moved to editor-only subobject.")
+	UPROPERTY()
+	TMap<EDMMaterialPropertyType, TObjectPtr<UObject>> PropertySlotMap;
+
+	UE_DEPRECATED(5.3, "Moved to editor-only subobject.")
+	UPROPERTY()
+	TArray<TObjectPtr<UObject>> Slots;
+
+	UE_DEPRECATED(5.3, "Moved to editor-only subobject.")
+	UPROPERTY()
+	TArray<TObjectPtr<UMaterialExpression>> Expressions;
+
+	UE_DEPRECATED(5.3, "Moved to editor-only subobject.")
+	UPROPERTY()
+	bool bCreateMaterialPackage;
+#endif
+};
diff --git a/Engine/Plugins/Experimental/DynamicMaterial/Source/DynamicMaterialEditor/Private/Components/DMMaterialSlot.cpp b/Engine/Plugins/Experimental/DynamicMaterial/Source/DynamicMaterialEditor/Private/Components/DMMaterialSlot.cpp
new file mode 100644
--- /dev/null
+++ b/Engine/Plugins/Experimental/DynamicMaterial/Source/DynamicMaterialEditor/Private/Components/DMMaterialSlot.cpp
@@ -0,0 +1,1401 @@
+// Copyright Epic Games, Inc. All Rights Reserved.
+
+#include "Components/DMMaterialSlot.h"
+#include "Components/DMMaterialLayer.h"
+#include "Components/DMMaterialLayer_Deprecated.h"
+#include "Components/DMMaterialProperty.h"
+#include "Components/DMMaterialStage.h"
+#include "Components/DMMaterialStageThroughputLayerBlend.h"
+#include "Components/DMMaterialSubStage.h"
+#include "CoreGlobals.h"
+#include "DMComponentPath.h"
+#include "DMDefs.h"
+#include "DynamicMaterialEditorModule.h"
+#include "Factories/MaterialFactoryNew.h"
+#include "Materials/Material.h"
+#include "Materials/MaterialExpression.h"
+#include "Materials/MaterialExpressionMax.h"
+#include "Materials/MaterialExpressionMultiply.h"
+#include "MaterialValueType.h"
+#include "Model/DMMaterialBuildState.h"
+#include "Model/DMMaterialBuildUtils.h"
+#include "Model/DynamicMaterialModel.h"
+#include "Model/DynamicMaterialModelEditorOnlyData.h"
+#include "UObject/Package.h"
+#include "UObject/UObjectGlobals.h"
+
+#define LOCTEXT_NAMESPACE "DMMaterialStage"
+
+const FString UDMMaterialSlot::LayersPathToken = FString(TEXT("Layers"));
+
+bool UDMMaterialSlot::MoveLayer(UDMMaterialLayerObject* InLayer, int32 InNewIndex)
+{
+	if (!IsComponentValid())
+	{
+		return false;
+	}
+
+	check(InLayer);
+
+	InNewIndex = FMath::Clamp(InNewIndex, 0, LayerObjects.Num() - 1);
+	const int32 CurrentIndex = InLayer->FindIndex();
+
+	if (InNewIndex == CurrentIndex)
+	{
+		return false;
+	}
+
+	if (InNewIndex == 0 && !InLayer->IsStageEnabled(EDMMaterialLayerStage::Base))
+	{
+		if (UDMMaterialStage* Stage = InLayer->GetStage(EDMMaterialLayerStage::Base))
+		{
+			if (GUndo)
+			{
+				Stage->Modify();
+			}
+
+			Stage->SetEnabled(true);
+		}
+	}
+		
+	const int MinIndex = FMath::Min(CurrentIndex, InNewIndex);
+	const int MaxIndex = FMath::Max(CurrentIndex, InNewIndex);
+
+	LayerObjects.RemoveAt(CurrentIndex, 1, EAllowShrinking::No); // Don't allow shrinking.
+	LayerObjects.Insert(InLayer, InNewIndex);
+
+	for (int32 LayerIndex = MinIndex; LayerIndex <= MaxIndex; ++LayerIndex)
+	{
+		LayerObjects[LayerIndex]->ForEachValidStage(
+			EDMMaterialLayerStage::All,
+			[](UDMMaterialStage* InStage)
+			{
+				if (GUndo)
+				{
+					InStage->Modify();
+				}
+
+				InStage->ResetInputConnectionMap();
+			});
+	}
+
+	if (InNewIndex == (LayerObjects.Num() - 1))
+	{
+		UpdateOutputConnectorTypes();
+	}
+
+	if (UDMMaterialStage* Stage = LayerObjects[MinIndex]->GetFirstEnabledStage(EDMMaterialLayerStage::All))
+	{
+		Stage->Update(EDMUpdateType::Structure);
+	}
+	else
+	{
+		Update(EDMUpdateType::Structure);
+	}
+
+	OnLayersUpdateDelegate.Broadcast(this);
+
+	return true;
+}
+
+bool UDMMaterialSlot::MoveLayerBefore(UDMMaterialLayerObject* InLayer, UDMMaterialLayerObject* InBeforeLayer)
+{
+	check(InLayer);
+
+	if (InBeforeLayer == nullptr)
+	{
+		return MoveLayer(InLayer, 0);
+	}
+	else
+	{
+		return MoveLayer(InLayer, InBeforeLayer->FindIndex() - 1);
+	}
+}
+
+bool UDMMaterialSlot::MoveLayerAfter(UDMMaterialLayerObject* InLayer, UDMMaterialLayerObject* InAfterLayer)
+{
+	check(InLayer);
+
+	if (InAfterLayer == nullptr)
+	{
+		return MoveLayer(InLayer, LayerObjects.Num());
+	}
+	else
+	{
+		return MoveLayer(InLayer, InAfterLayer->FindIndex() + 1);
+	}
+}
+
+UDMMaterialLayerObject* UDMMaterialSlot::FindLayer(const UDMMaterialStage* InBaseOrMask) const
+{
+	if (const UDMMaterialSubStage* SubStage = Cast<UDMMaterialSubStage>(InBaseOrMask))
+	{
+		InBaseOrMask = SubStage->GetParentMostStage();
+	}
+
+	const TObjectPtr<UDMMaterialLayerObject>* FoundLayer = LayerObjects.FindByPredicate(
+		[InBaseOrMask](const TObjectPtr<UDMMaterialLayerObject>& Element)
+		{
+			return IsValid(Element) && Element->HasValidStage(InBaseOrMask);
+		}
+	);
+
+	if (FoundLayer)
+	{
+		return FoundLayer->Get();
+	}
+
+	return nullptr;
+}
+
+TArray<UDMMaterialLayerObject*> UDMMaterialSlot::BP_GetLayers() const
+{
+	TArray<UDMMaterialLayerObject*> OutLayerObjects;
+	OutLayerObjects.Reserve(LayerObjects.Num());
+
+	Algo::Transform(LayerObjects, OutLayerObjects, [](const TObjectPtr<UDMMaterialLayerObject> InElement) { return InElement; });
+
+	return OutLayerObjects;
+}
+
+UDMMaterialLayerObject* UDMMaterialSlot::GetLastLayerForMaterialProperty(EDMMaterialPropertyType InMaterialProperty) const
+{
+	for (int32 LayerIndex = LayerObjects.Num() - 1; LayerIndex >= 0; --LayerIndex)
+	{
+		if (!LayerObjects[LayerIndex]->IsStageEnabled(EDMMaterialLayerStage::Base))
+		{
+			continue;
+		}
+
+		if (LayerObjects[LayerIndex]->GetMaterialProperty() != InMaterialProperty)
+		{
+			continue;
+		}
+
+		return LayerObjects[LayerIndex];
+	}
+
+	return nullptr;
+}
+
+void UDMMaterialSlot::Update(EDMUpdateType InUpdateType)
+{
+	if (!IsComponentValid())
+	{
+		return;
+	}
+
+	if (HasComponentBeenRemoved())
+	{
+		return;
+	}
+
+	if (InUpdateType == EDMUpdateType::Structure)
+	{
+		MarkComponentDirty();
+		UpdateMaterialProperties();
+	}
+
+	Super::Update(InUpdateType);
+
+	if (InUpdateType == EDMUpdateType::Structure)
+	{
+		UDynamicMaterialModelEditorOnlyData* ModelEditorOnlyData = GetMaterialModelEditorOnlyData();
+		check(ModelEditorOnlyData);
+
+		ModelEditorOnlyData->RequestMaterialBuild();
+	}
+}
+
+void UDMMaterialSlot::OnComponentAdded()
+{
+	if (!IsComponentValid())
+	{
+		return;
+	}
+
+	Super::OnComponentAdded();
+
+	for (const TObjectPtr<UDMMaterialLayerObject>& LayerObj : LayerObjects)
+	{
+		if (GUndo)
+		{
+			LayerObj->Modify();
+		}
+
+		LayerObj->SetComponentState(EDMComponentLifetimeState::Added);
+	}
+}
+
+void UDMMaterialSlot::OnComponentRemoved()
+{
+	Super::OnComponentRemoved();
+
+	for (const TObjectPtr<UDMMaterialLayerObject>& LayerObj : LayerObjects)
+	{
+		if (GUndo)
+		{
+			LayerObj->Modify();
+		}
+
+		LayerObj->SetComponentState(EDMComponentLifetimeState::Removed);
+	}
+}
+
+void UDMMaterialSlot::UpdateOutputConnectorTypes()
+{
+	if (!IsComponentValid())
+	{
+		return;
+	}
+
+	OutputConnectorTypes.Empty();
+
+	if (LayerObjects.IsEmpty())
+	{
+		UpdateMaterialProperties();
+		return;
+	}
+
+	TMap<EDMMaterialPropertyType, UDMMaterialLayerObject*> LastOutputForProperty;
+
+	for (const TObjectPtr<UDMMaterialLayerObject>& Layer : LayerObjects)
+	{
+		LastOutputForProperty.FindOrAdd(Layer->GetMaterialProperty()) = Layer;
+	}
+
+	for (const TPair<EDMMaterialPropertyType, UDMMaterialLayerObject*>& PropStage : LastOutputForProperty)
+	{
+		if (UDMMaterialStage* Mask = PropStage.Value->GetStage(EDMMaterialLayerStage::Mask, /* Enabled Only */ true))
+		{
+			if (UDMMaterialStageSource* Source = Mask->GetSource())
+			{
+				TArray<EDMValueType> MaterialPropertyORutputConnectorTypes;
+				const TArray<FDMMaterialStageConnector>& LastStageOutputConnectors = Source->GetOutputConnectors();
+
+				for (const FDMMaterialStageConnector& Connector : LastStageOutputConnectors)
+				{
+					MaterialPropertyORutputConnectorTypes.Add(Connector.Type);
+				}
+
+				FDMMaterialSlotOutputConnectorTypes ConnectorTypes = {MaterialPropertyORutputConnectorTypes};
+				OutputConnectorTypes.Emplace(PropStage.Key, ConnectorTypes);
+			}
+		}
+	}
+
+	UpdateMaterialProperties();
+
+	OnConnectorsUpdateDelegate.Broadcast(this);
+}
+
+void UDMMaterialSlot::UpdateMaterialProperties()
+{
+	if (!IsComponentValid())
+	{
+		return;
+	}
+
+	UDynamicMaterialModelEditorOnlyData* ModelEditorOnlyData = GetMaterialModelEditorOnlyData();
+	check(ModelEditorOnlyData);
+
+	if (LayerObjects.IsEmpty())
+	{
+		const TArray<EDMMaterialPropertyType> SlotMaterialProperties = ModelEditorOnlyData->GetMaterialPropertiesForSlot(this);
+
+		for (EDMMaterialPropertyType MaterialProperty : SlotMaterialProperties)
+		{
+			UDMMaterialSlot* CurrentSlot = ModelEditorOnlyData->GetSlotForMaterialProperty(MaterialProperty);
+			check(CurrentSlot == nullptr || CurrentSlot == this);
+
+			if (CurrentSlot == this)
+			{
+				if (GUndo)
+				{
+					ModelEditorOnlyData->Modify();
+				}
+
+				ModelEditorOnlyData->UnassignMaterialProperty(MaterialProperty);
+			}
+		}
+
+		return;
+	}
+
+	TSet<EDMMaterialPropertyType> CurrentStageMaterialProperties;
+
+	for (const TObjectPtr<UDMMaterialLayerObject>& Layer : LayerObjects)
+	{
+		EDMMaterialPropertyType StageProperty = Layer->GetMaterialProperty();
+
+		if (StageProperty != EDMMaterialPropertyType::None && StageProperty != EDMMaterialPropertyType::Any)
+		{
+			CurrentStageMaterialProperties.Add(StageProperty);
+		}
+	}
+
+	const TArray<EDMMaterialPropertyType> CurrentSlotMaterialProperties = ModelEditorOnlyData->GetMaterialPropertiesForSlot(this);
+
+	for (EDMMaterialPropertyType MaterialProperty : CurrentSlotMaterialProperties)
+	{
+		if (!CurrentStageMaterialProperties.Contains(MaterialProperty))
+		{
+			UDMMaterialSlot* CurrentSlot = ModelEditorOnlyData->GetSlotForMaterialProperty(MaterialProperty);
+			check(CurrentSlot == nullptr || CurrentSlot == this);
+
+			if (CurrentSlot == this)
+			{
+				if (GUndo)
+				{
+					ModelEditorOnlyData->Modify();
+				}
+
+				ModelEditorOnlyData->UnassignMaterialProperty(MaterialProperty);
+			}
+		}
+	}
+
+	for (EDMMaterialPropertyType MaterialProperty : CurrentStageMaterialProperties)
+	{
+		if (!CurrentSlotMaterialProperties.Contains(MaterialProperty))
+		{
+			UDMMaterialSlot* CurrentSlot = ModelEditorOnlyData->GetSlotForMaterialProperty(MaterialProperty);
+			check(CurrentSlot == nullptr || CurrentSlot == this);
+
+			if (CurrentSlot == nullptr)
+			{
+				if (GUndo)
+				{
+					ModelEditorOnlyData->Modify();
+				}
+
+				ModelEditorOnlyData->AssignMaterialPropertyToSlot(MaterialProperty, this);
+			}
+		}
+	}
+}
+
+UDMMaterialSlot::UDMMaterialSlot()
+	: Index(INDEX_NONE)
+	, bIsEditingLayers(true)
+	, BasePreviewMaterial(nullptr)
+	, MaskPreviewMaterial(nullptr)
+{
+}
+
+UDynamicMaterialModelEditorOnlyData* UDMMaterialSlot::GetMaterialModelEditorOnlyData() const
+{
+	return Cast<UDynamicMaterialModelEditorOnlyData>(GetOuterSafe());
+}
+
+FText UDMMaterialSlot::GetDescription() const
+{
+	static const FText Template = LOCTEXT("StageInputSlotTempate", "Slot {0}");
+
+	return FText::Format(Template, FText::AsNumber(Index));
+}
+
+UDMMaterialLayerObject* UDMMaterialSlot::GetLayer(int32 InLayerIndex) const
+{
+	if (!LayerObjects.IsValidIndex(InLayerIndex))
+	{
+		return nullptr;
+	}
+
+	return LayerObjects[InLayerIndex];
+}
+
+const TArray<EDMValueType>& UDMMaterialSlot::GetOutputConnectorTypesForMaterialProperty(EDMMaterialPropertyType InMaterialProperty) const
+{
+	static const TArray<EDMValueType> NoConnectors;
+
+	const FDMMaterialSlotOutputConnectorTypes* ConnectorTypes = OutputConnectorTypes.Find(InMaterialProperty);
+
+	if (!ConnectorTypes)
+	{
+		return NoConnectors;
+	}
+
+	return ConnectorTypes->ConnectorTypes;
+}
+
+TSet<EDMValueType> UDMMaterialSlot::GetAllOutputConnectorTypes() const
+{
+	TSet<EDMValueType> AllOutputTypes;
+
+	for (const TPair<EDMMaterialPropertyType, FDMMaterialSlotOutputConnectorTypes>& OutputTypes : OutputConnectorTypes)
+	{
+		for (EDMValueType OutputType : OutputTypes.Value.ConnectorTypes)
+		{
+			AllOutputTypes.Emplace(OutputType);
+		}
+	}
+
+	return AllOutputTypes;
+}
+
+UDMMaterialLayerObject* UDMMaterialSlot::AddLayer(EDMMaterialPropertyType InMaterialProperty, UDMMaterialStage* InNewBase)
+{
+	UDMMaterialStage* NewMask = UDMMaterialStageThroughputLayerBlend::CreateStage();
+	check(NewMask);
+
+	return AddLayerWithMask(InMaterialProperty, InNewBase, NewMask);
+}
+
+UDMMaterialLayerObject* UDMMaterialSlot::AddLayerWithMask(EDMMaterialPropertyType InMaterialProperty, UDMMaterialStage* InNewBase, 
+	UDMMaterialStage* InNewMask)
+{
+	if (!IsComponentValid())
+	{
+		return nullptr;
+	}
+
+	check(InNewBase);
+	check(InNewBase->GetSource());
+	check(InNewBase->GetSource()->GetOutputConnectors().IsEmpty() == false);
+
+	check(InNewMask);
+	check(InNewMask->GetSource());
+	check(InNewMask->GetSource()->GetOutputConnectors().IsEmpty() == false);
+
+	if (GUndo)
+	{
+		InNewBase->Modify();
+		InNewMask->Modify();
+	}
+
+	InNewBase->SetBeingEdited(false);
+	InNewMask->SetBeingEdited(false);
+
+	UDMMaterialLayerObject* NewLayer = UDMMaterialLayerObject::CreateLayer(this, InMaterialProperty, {InNewBase, InNewMask});
+	LayerObjects.Add(NewLayer);
+
+	if (IsComponentAdded())
+	{
+		NewLayer->SetComponentState(EDMComponentLifetimeState::Added);
+	}
+
+	UpdateOutputConnectorTypes();
+
+	InNewBase->Update(EDMUpdateType::Structure);
+
+	OnLayersUpdateDelegate.Broadcast(this);
+
+	return LayerObjects.Last();
+}
+
+bool UDMMaterialSlot::PasteLayer(UDMMaterialLayerObject* InLayer)
+{
+	if (!InLayer)
+	{
+		return false;
+	}
+
+	if (GUndo)
+	{
+		InLayer->Modify();
+	}
+
+	UDynamicMaterialModelEditorOnlyData* ModelEditorOnlyData = GetMaterialModelEditorOnlyData();
+
+	if (!ModelEditorOnlyData)
+	{
+		return false;
+	}
+
+	UDynamicMaterialModel* MaterialModel = ModelEditorOnlyData->GetMaterialModel();
+
+	LayerObjects.Add(InLayer);
+
+	if (IsComponentAdded())
+	{
+		InLayer->SetComponentState(EDMComponentLifetimeState::Added);
+	}
+
+	UpdateOutputConnectorTypes();
+
+	if (UDMMaterialStage* Stage = InLayer->GetFirstEnabledStage(EDMMaterialLayerStage::All))
+	{
+		Stage->Update(EDMUpdateType::Structure);
+	}
+	else
+	{
+		Update(EDMUpdateType::Structure);
+	}
+
+	OnLayersUpdateDelegate.Broadcast(this);
+
+	return true;
+}
+
+bool UDMMaterialSlot::RemoveLayer(UDMMaterialLayerObject* InLayer)
+{
+	if (!IsComponentValid())
+	{
+		return false;
+	}
+
+	if (LayerObjects.Num() == 1)
+	{
+		return false;
+	}
+
+	check(InLayer);
+	check(InLayer->GetSlot() == this);
+
+	const int32 LayerIndex = InLayer->FindIndex();
+
+	if (LayerIndex == INDEX_NONE)
+	{
+		return false;
+	}
+
+	LayerObjects.RemoveAt(LayerIndex);
+
+	if (LayerIndex == 0 && LayerObjects.IsEmpty() == false)
+	{
+		if (UDMMaterialStage* Stage = LayerObjects[0]->GetStage(EDMMaterialLayerStage::Base))
+		{
+			if (GUndo)
+			{
+				Stage->Modify();
+			}
+
+			Stage->SetEnabled(true);
+		}
+	}
+
+	if (GUndo)
+	{
+		InLayer->Modify();
+	}
+
+	InLayer->SetComponentState(EDMComponentLifetimeState::Removed);
+
+	if (!LayerObjects.IsEmpty())
+	{
+		if (UDMMaterialStage* Stage = LayerObjects[0]->GetFirstEnabledStage(EDMMaterialLayerStage::All))
+		{
+			Stage->Update(EDMUpdateType::Structure);
+		}
+		else
+		{
+			Update(EDMUpdateType::Structure);
+		}
+	}
+
+	OnLayersUpdateDelegate.Broadcast(this);
+
+	return true;
+}
+
+void UDMMaterialSlot::GenerateExpressions(const TSharedRef<FDMMaterialBuildState>& InBuildState) const
+{
+	if (!IsComponentValid() || !IsComponentAdded())
+	{
+		return;
+	}
+
+	if (InBuildState->HasSlot(this) || LayerObjects.IsEmpty())
+	{
+		return;
+	}
+
+	TArray<UMaterialExpression*> SlotExpressions;
+	TMap<EDMMaterialPropertyType, TArray<UMaterialExpression*>> SlotPropertyExpressions;
+
+	for (const TObjectPtr<UDMMaterialLayerObject>& Layer : LayerObjects)
+	{
+		if (!Layer->IsEnabled())
+		{
+			continue;
+		}
+
+		Layer->GenerateExpressions(InBuildState);
+
+		const TArray<UMaterialExpression*>& LayerExpressions = InBuildState->GetLayerExpressions(Layer);
+
+		if (LayerExpressions.IsEmpty())
+		{
+			continue;
+		}
+
+		SlotExpressions.Append(LayerExpressions);
+		SlotPropertyExpressions.FindOrAdd(Layer->GetMaterialProperty()).Append(LayerExpressions);
+	}
+
+	if (SlotExpressions.IsEmpty())
+	{
+		return;
+	}
+
+	InBuildState->AddSlotExpressions(this, SlotExpressions);
+	InBuildState->AddSlotPropertyExpressions(this, SlotPropertyExpressions);
+}
+
+TArray<UDMMaterialSlot*> UDMMaterialSlot::K2_GetSlotsReferencedBy() const
+{
+	TArray<TWeakObjectPtr<UDMMaterialSlot>> WeakKeys;
+	SlotsReferencedBy.GetKeys(WeakKeys);
+
+	TArray<UDMMaterialSlot*> Keys;
+	Keys.SetNum(Keys.Num());
+
+	for (int32 KeyIndex = 0; KeyIndex < WeakKeys.Num(); ++KeyIndex)
+	{
+		Keys[KeyIndex] = WeakKeys[KeyIndex].Get();
+	}
+
+	return Keys;
+}
+
+bool UDMMaterialSlot::ReferencedBySlot(UDMMaterialSlot* InOtherSlot)
+{
+	if (!IsComponentValid())
+	{
+		return false;
+	}
+
+	check(InOtherSlot);
+	check(InOtherSlot != this);
+
+	int32* CountPtr = SlotsReferencedBy.Find(InOtherSlot);
+
+	if (CountPtr)
+	{
+		++(*CountPtr);
+		return false;
+	}
+	else
+	{
+		SlotsReferencedBy.Emplace(InOtherSlot, 1);
+		OnPropertiesUpdateDelegate.Broadcast(this);
+		return true;
+	}
+}
+
+bool UDMMaterialSlot::UnreferencedBySlot(UDMMaterialSlot* InOtherSlot)
+{
+	if (!IsComponentValid())
+	{
+		return false;
+	}
+
+	check(InOtherSlot);
+	check(InOtherSlot != this);
+
+	int32* CountPtr = SlotsReferencedBy.Find(InOtherSlot);
+	check(CountPtr);
+
+	--(*CountPtr);
+
+	if ((*CountPtr) == 0)
+	{
+		SlotsReferencedBy.Remove(InOtherSlot);
+		OnPropertiesUpdateDelegate.Broadcast(this);
+		return true;
+	}
+	else
+	{
+		return false;
+	}
+}
+
+UMaterial* UDMMaterialSlot::GetPreviewMaterial(EDMMaterialLayerStage InLayerStage)
+{
+	check(InLayerStage == EDMMaterialLayerStage::Base || InLayerStage == EDMMaterialLayerStage::Mask);
+	TObjectPtr<UMaterial>& PreviewMaterial = (InLayerStage == EDMMaterialLayerStage::Base ? BasePreviewMaterial : MaskPreviewMaterial);
+
+	if (!PreviewMaterial)
+	{
+		CreatePreviewMaterial(InLayerStage);
+
+		if (PreviewMaterial)
+		{
+			MarkComponentDirty();
+		}
+	}
+
+	return PreviewMaterial;
+}
+
+void UDMMaterialSlot::CreatePreviewMaterial(EDMMaterialLayerStage InLayerStage)
+{
+	if (!IsComponentValid())
+	{
+		return;
+	}
+
+	check(InLayerStage == EDMMaterialLayerStage::Base || InLayerStage == EDMMaterialLayerStage::Mask);
+	TObjectPtr<UMaterial>& PreviewMaterial = (InLayerStage == EDMMaterialLayerStage::Base ? BasePreviewMaterial : MaskPreviewMaterial);
+
+	UMaterialFactoryNew* MaterialFactory = NewObject<UMaterialFactoryNew>();
+	check(MaterialFactory);
+
+	PreviewMaterial = (UMaterial*)MaterialFactory->FactoryCreateNew(
+		UMaterial::StaticClass(),
+		GetTransientPackage(),
+		NAME_None,
+		RF_Transient,
+		nullptr,
+		GWarn
+	);
+
+	PreviewMaterial->bIsPreviewMaterial = true;
+}
+
+void UDMMaterialSlot::UpdatePreviewMaterial(EDMMaterialLayerStage InLayerStage)
+{
+	if (!IsComponentValid())
+	{
+		return;
+	}
+
+	bool bHasActivateStage = false;
+
+	for (const TObjectPtr<UDMMaterialLayerObject>& Layer : LayerObjects)
+	{
+		if (!Layer->IsEnabled())
+		{
+			continue;
+		}
+
+		if (Layer->IsStageEnabled(EDMMaterialLayerStage::Base))
+		{
+			bHasActivateStage = true;
+			break;
+		}
+	}
+
+	if (!bHasActivateStage)
+	{
+		return;
+	}
+
+	check(InLayerStage == EDMMaterialLayerStage::Base || InLayerStage == EDMMaterialLayerStage::Mask);
+	TObjectPtr<UMaterial>& PreviewMaterial = (InLayerStage == EDMMaterialLayerStage::Base ? BasePreviewMaterial : MaskPreviewMaterial);
+
+	if (!PreviewMaterial)
+	{
+		CreatePreviewMaterial(InLayerStage);
+
+		if (!PreviewMaterial)
+		{
+			return;
+		}
+	}
+
+	UDynamicMaterialModelEditorOnlyData* ModelEditorOnlyData = GetMaterialModelEditorOnlyData();
+	check(ModelEditorOnlyData);
+
+	if (InLayerStage == EDMMaterialLayerStage::Base)
+	{
+		TSharedRef<FDMMaterialBuildState> BuildState = ModelEditorOnlyData->CreateBuildState(BasePreviewMaterial);
+		BuildState->SetPreviewMaterial();
+
+		GenerateExpressions(BuildState);
+		UpdateBasePreviewMaterial(BuildState);
+	}
+	else
+	{
+		TSharedRef<FDMMaterialBuildState> BuildState = ModelEditorOnlyData->CreateBuildState(MaskPreviewMaterial);
+		BuildState->SetPreviewMaterial();
+
+		GenerateExpressions(BuildState);
+		UpdateMaskPreviewMaterial(BuildState);
+	}
+}
+
+void UDMMaterialSlot::UpdateBasePreviewMaterial(const TSharedRef<FDMMaterialBuildState>& InBuildState)
+{
+	if (!IsComponentValid())
+	{
+		return;
+	}
+
+	UE_LOG(LogDynamicMaterialEditor, Display, TEXT("Building Material Designer Slot Base Preview (%s)..."), *GetName());
+
+	BasePreviewMaterial->GetEditorOnlyData()->EmissiveColor.Expression = nullptr;
+	BasePreviewMaterial->GetEditorOnlyData()->EmissiveColor.OutputIndex = 0;
+
+	const TMap<EDMMaterialPropertyType, TArray<UMaterialExpression*>>& SlotPropertyExpressions = InBuildState->GetSlotPropertyExpressions(this);
+
+	// Single material property, connect it up to emissive and output it.
+	if (SlotPropertyExpressions.Num() == 1)
+	{
+		for (const TPair<EDMMaterialPropertyType, TArray<UMaterialExpression*>>& Pair : SlotPropertyExpressions)
+		{
+			UpdateBasePreviewMaterialProperty(InBuildState, Pair.Key);
+		}
+	}
+	else
+	{
+		UpdateBasePreviewMaterialFull(InBuildState);
+	}
+}
+
+void UDMMaterialSlot::UpdateBasePreviewMaterialProperty(const TSharedRef<FDMMaterialBuildState>& InBuildState, EDMMaterialPropertyType InBaseProperty)
+{
+	UpdatePreviewMaterialProperty(InBuildState, BasePreviewMaterial, InBaseProperty);
+}
+
+void UDMMaterialSlot::UpdateBasePreviewMaterialFull(const TSharedRef<FDMMaterialBuildState>& InBuildState)
+{
+	if (!IsComponentValid())
+	{
+		return;
+	}
+
+	UDynamicMaterialModelEditorOnlyData* ModelEditorOnlyData = GetMaterialModelEditorOnlyData();
+	check(ModelEditorOnlyData);
+
+	const TMap<EDMMaterialPropertyType, UDMMaterialProperty*> MaterialProperties = ModelEditorOnlyData->GetMaterialProperties();
+	const TMap<EDMMaterialPropertyType, TArray<UMaterialExpression*>>& SlotPropertyExpressions = InBuildState->GetSlotPropertyExpressions(this);
+
+	for (const TPair<EDMMaterialPropertyType, UDMMaterialProperty*>& Pair : MaterialProperties)
+	{
+		if (UE::DynamicMaterialEditor::IsCustomMaterialProperty(Pair.Key))
+		{
+			continue;
+		}
+
+		// For now we don't have channel remapping!
+		FExpressionInput* MaterialPropertyPtr = InBuildState->GetMaterialProperty(Pair.Key);
+		check(MaterialPropertyPtr);
+
+		MaterialPropertyPtr->Expression = nullptr;
+		MaterialPropertyPtr->OutputIndex = 0;
+
+		if (!SlotPropertyExpressions.Contains(Pair.Key))
+		{
+			continue;
+		}
+
+		const TArray<UMaterialExpression*>& PropertyExpressions = SlotPropertyExpressions[Pair.Key];
+
+		if (PropertyExpressions.IsEmpty())
+		{
+			continue;
+		}
+
+		MaterialPropertyPtr->Expression = PropertyExpressions.Last();
+
+		if (Pair.Value->GetInputConnectionMap().Channels.IsEmpty() == false)
+		{
+			MaterialPropertyPtr->OutputIndex = Pair.Value->GetInputConnectionMap().Channels[0].OutputIndex;
+		}
+		else
+		{
+			MaterialPropertyPtr->OutputIndex = 0;
+		}
+	}
+}
+
+void UDMMaterialSlot::UpdateMaskPreviewMaterial(const TSharedRef<FDMMaterialBuildState>& InBuildState)
+{
+	if (!IsComponentValid())
+	{
+		return;
+	}
+
+	UE_LOG(LogDynamicMaterialEditor, Display, TEXT("Building Material Designer Slot Mask Preview (%s)..."), *GetName());
+
+	MaskPreviewMaterial->GetEditorOnlyData()->EmissiveColor.Expression = nullptr;
+	MaskPreviewMaterial->GetEditorOnlyData()->EmissiveColor.OutputIndex = 0;
+
+	UDynamicMaterialModelEditorOnlyData* ModelEditorOnlyData = GetMaterialModelEditorOnlyData();
+	check(ModelEditorOnlyData);
+
+	const TMap<EDMMaterialPropertyType, TArray<UMaterialExpression*>>& SlotPropertyExpressions = InBuildState->GetSlotPropertyExpressions(this);
+
+	if (SlotPropertyExpressions.Num() == 1)
+	{
+		for (const TPair<EDMMaterialPropertyType, TArray<UMaterialExpression*>>& Pair : SlotPropertyExpressions)
+		{
+			UpdateMaskPreviewMaterialMaskCombination(InBuildState, Pair.Key);
+		}
+	}
+	else if (SlotPropertyExpressions.Contains(EDMMaterialPropertyType::Opacity))
+	{
+		UpdateMaskPreviewMaterialProperty(InBuildState, EDMMaterialPropertyType::Opacity);
+	}
+	else if (SlotPropertyExpressions.Contains(EDMMaterialPropertyType::OpacityMask))
+	{
+		UpdateMaskPreviewMaterialProperty(InBuildState, EDMMaterialPropertyType::OpacityMask);
+	}
+	else if (SlotPropertyExpressions.Contains(EDMMaterialPropertyType::BaseColor))
+	{
+		UpdateMaskPreviewMaterialMaskCombination(InBuildState, EDMMaterialPropertyType::BaseColor);
+	}
+	else if (SlotPropertyExpressions.Contains(EDMMaterialPropertyType::EmissiveColor))
+	{
+		UpdateMaskPreviewMaterialMaskCombination(InBuildState, EDMMaterialPropertyType::EmissiveColor);
+	}
+}
+
+void UDMMaterialSlot::UpdateMaskPreviewMaterialProperty(const TSharedRef<FDMMaterialBuildState>& InBuildState, EDMMaterialPropertyType InMaskProperty)
+{
+	UpdatePreviewMaterialProperty(InBuildState, MaskPreviewMaterial, InMaskProperty);
+}
+
+void UDMMaterialSlot::UpdateMaskPreviewMaterialMaskCombination(const TSharedRef<FDMMaterialBuildState>& InBuildState, 
+	EDMMaterialPropertyType InMaskProperty)
+{
+	if (!IsComponentValid())
+	{
+		return;
+	}
+
+	int32 OutOutputChannel = FDMMaterialStageConnectorChannel::WHOLE_CHANNEL;
+
+	for (const TObjectPtr<UDMMaterialLayerObject>& Layer : LayerObjects)
+	{
+		// Although we are working with masks, if the base is disabled, this is handled by the GenerateExpressions
+		// of the LayerBlend code (to multiply alpha together, instead of maxing it).
+		if (Layer->GetMaterialProperty() != InMaskProperty || !Layer->AreAllStagesEnabled(EDMMaterialLayerStage::All))
+		{
+			continue;
+		}
+
+		UDMMaterialStage* BaseStage = Layer->GetStage(EDMMaterialLayerStage::Base);
+		UDMMaterialStage* MaskStage = Layer->GetStage(EDMMaterialLayerStage::Mask);
+
+		MaskStage->GenerateExpressions(InBuildState);
+		UDMMaterialStageThroughputLayerBlend* LayerBlend = Cast<UDMMaterialStageThroughputLayerBlend>(MaskStage->GetSource());
+
+		if (!LayerBlend)
+		{
+			continue;
+		}
+
+		UMaterialExpression* MaskOutputExpression;
+		int32 MaskOutputIndex;
+		int32 MaskOutputChannel;
+		LayerBlend->GetMaskOutput(InBuildState, MaskOutputExpression, MaskOutputIndex, MaskOutputChannel);
+
+		if (!MaskOutputExpression)
+		{
+			continue;
+		}
+
+		if (LayerBlend->UsePremultiplyAlpha())
+		{
+			if (UDMMaterialStageSource* Source = BaseStage->GetSource())
+			{
+				UMaterialExpression* LayerAlphaOutputExpression;
+				int32 LayerAlphaOutputIndex;
+				int32 LayerAlphaOutputChannel;
+
+				Source->GetMaskAlphaBlendNode(InBuildState, LayerAlphaOutputExpression, LayerAlphaOutputIndex, LayerAlphaOutputChannel);
+
+				if (LayerAlphaOutputExpression)
+				{
+					UMaterialExpressionMultiply* AlphaMultiply = InBuildState->GetBuildUtils().CreateExpression<UMaterialExpressionMultiply>(UE_DM_NodeComment_Default);
+
+					AlphaMultiply->A.Expression = MaskOutputExpression;
+					AlphaMultiply->A.OutputIndex = MaskOutputIndex;
+					AlphaMultiply->A.Mask = 0;
+
+					if (MaskOutputChannel != FDMMaterialStageConnectorChannel::WHOLE_CHANNEL)
+					{
+						AlphaMultiply->A.Mask = 1;
+						AlphaMultiply->A.MaskR = !!(MaskOutputChannel & FDMMaterialStageConnectorChannel::FIRST_CHANNEL);
+						AlphaMultiply->A.MaskG = !!(MaskOutputChannel & FDMMaterialStageConnectorChannel::SECOND_CHANNEL);
+						AlphaMultiply->A.MaskB = !!(MaskOutputChannel & FDMMaterialStageConnectorChannel::THIRD_CHANNEL);
+						AlphaMultiply->A.MaskA = !!(MaskOutputChannel & FDMMaterialStageConnectorChannel::FOURTH_CHANNEL);
+					}
+
+					AlphaMultiply->B.Expression = LayerAlphaOutputExpression;
+					AlphaMultiply->B.OutputIndex = LayerAlphaOutputIndex;
+					AlphaMultiply->B.Mask = 0;
+
+					if (LayerAlphaOutputChannel != FDMMaterialStageConnectorChannel::WHOLE_CHANNEL)
+					{
+						AlphaMultiply->B.Mask = 1;
+						AlphaMultiply->B.MaskR = !!(LayerAlphaOutputChannel & FDMMaterialStageConnectorChannel::FIRST_CHANNEL);
+						AlphaMultiply->B.MaskG = !!(LayerAlphaOutputChannel & FDMMaterialStageConnectorChannel::SECOND_CHANNEL);
+						AlphaMultiply->B.MaskB = !!(LayerAlphaOutputChannel & FDMMaterialStageConnectorChannel::THIRD_CHANNEL);
+						AlphaMultiply->B.MaskA = !!(LayerAlphaOutputChannel & FDMMaterialStageConnectorChannel::FOURTH_CHANNEL);
+					}
+
+					MaskOutputExpression = AlphaMultiply;
+					MaskOutputIndex = 0;
+					MaskOutputChannel = FDMMaterialStageConnectorChannel::WHOLE_CHANNEL;
+				}
+			}
+		}
+
+		if (MaskPreviewMaterial->GetEditorOnlyData()->EmissiveColor.Expression == nullptr)
+		{
+			MaskPreviewMaterial->GetEditorOnlyData()->EmissiveColor.Expression = MaskOutputExpression;
+
+			// The first output will use the node's output info.
+			MaskPreviewMaterial->GetEditorOnlyData()->EmissiveColor.OutputIndex = MaskOutputIndex;
+			OutOutputChannel = MaskOutputChannel;
+			continue;
+		}
+
+		UMaterialExpressionMax* Max = InBuildState->GetBuildUtils().CreateExpression<UMaterialExpressionMax>(UE_DM_NodeComment_Default);
+		check(Max);
+
+		Max->A.Expression = MaskPreviewMaterial->GetEditorOnlyData()->EmissiveColor.Expression;
+		Max->A.OutputIndex = MaskPreviewMaterial->GetEditorOnlyData()->EmissiveColor.OutputIndex;
+		Max->A.Mask = 0;
+
+		if (OutOutputChannel != FDMMaterialStageConnectorChannel::WHOLE_CHANNEL)
+		{
+			Max->A.Mask = 1;
+			Max->A.MaskR = !!(OutOutputChannel & FDMMaterialStageConnectorChannel::FIRST_CHANNEL);
+			Max->A.MaskG = !!(OutOutputChannel & FDMMaterialStageConnectorChannel::SECOND_CHANNEL);
+			Max->A.MaskB = !!(OutOutputChannel & FDMMaterialStageConnectorChannel::THIRD_CHANNEL);
+			Max->A.MaskA = !!(OutOutputChannel & FDMMaterialStageConnectorChannel::FOURTH_CHANNEL);
+		}
+
+		Max->B.Expression = MaskOutputExpression;
+		Max->B.OutputIndex = MaskOutputIndex;
+		Max->B.Mask = 0;
+
+		if (MaskOutputChannel != FDMMaterialStageConnectorChannel::WHOLE_CHANNEL)
+		{
+			Max->B.Mask = 1;
+			Max->B.MaskR = !!(MaskOutputChannel & FDMMaterialStageConnectorChannel::FIRST_CHANNEL);
+			Max->B.MaskG = !!(MaskOutputChannel & FDMMaterialStageConnectorChannel::SECOND_CHANNEL);
+			Max->B.MaskB = !!(MaskOutputChannel & FDMMaterialStageConnectorChannel::THIRD_CHANNEL);
+			Max->B.MaskA = !!(MaskOutputChannel & FDMMaterialStageConnectorChannel::FOURTH_CHANNEL);
+		}
+
+		MaskPreviewMaterial->GetEditorOnlyData()->EmissiveColor.Expression = Max;
+
+		// If we have to combine, it will use the Max node's output info
+		MaskPreviewMaterial->GetEditorOnlyData()->EmissiveColor.OutputIndex = 0;
+		OutOutputChannel = FDMMaterialStageConnectorChannel::WHOLE_CHANNEL;
+	}
+
+	MaskPreviewMaterial->GetEditorOnlyData()->EmissiveColor.Mask = 0;
+
+	if (OutOutputChannel != FDMMaterialStageConnectorChannel::WHOLE_CHANNEL)
+	{
+		MaskPreviewMaterial->GetEditorOnlyData()->EmissiveColor.Mask = 1;
+		MaskPreviewMaterial->GetEditorOnlyData()->EmissiveColor.MaskR = !!(OutOutputChannel & FDMMaterialStageConnectorChannel::FIRST_CHANNEL);
+		MaskPreviewMaterial->GetEditorOnlyData()->EmissiveColor.MaskG = !!(OutOutputChannel & FDMMaterialStageConnectorChannel::SECOND_CHANNEL);
+		MaskPreviewMaterial->GetEditorOnlyData()->EmissiveColor.MaskB = !!(OutOutputChannel & FDMMaterialStageConnectorChannel::THIRD_CHANNEL);
+		MaskPreviewMaterial->GetEditorOnlyData()->EmissiveColor.MaskA = !!(OutOutputChannel & FDMMaterialStageConnectorChannel::FOURTH_CHANNEL);
+	}
+}
+
+void UDMMaterialSlot::UpdatePreviewMaterialProperty(const TSharedRef<FDMMaterialBuildState>& InBuildState, UMaterial* InPreviewMaterial, 
+	EDMMaterialPropertyType InProperty)
+{
+	if (!IsComponentValid())
+	{
+		return;
+	}
+
+	UDynamicMaterialModelEditorOnlyData* ModelEditorOnlyData = GetMaterialModelEditorOnlyData();
+	check(ModelEditorOnlyData);
+	
+	UDMMaterialProperty* PropertyObj = ModelEditorOnlyData->GetMaterialProperty(InProperty);
+	check(PropertyObj);
+
+	const TArray<UMaterialExpression*>& SlotExpressions = InBuildState->GetSlotExpressions(this);
+
+	if (SlotExpressions.IsEmpty())
+	{
+		return;
+	}
+
+	UMaterialExpression* LastExpression = SlotExpressions.Last();
+	int32 BestMatch = INDEX_NONE;
+	int32 OutputCount = 0;
+	const int32 FloatsForPropertyType = UDMValueDefinitionLibrary::GetValueDefinition(PropertyObj->GetInputConnectorType()).GetFloatCount();
+
+	for (int32 OutputIdx = 0; OutputIdx < LastExpression->GetOutputs().Num(); ++OutputIdx)
+	{
+		const EMaterialValueType CurrentOutputType = static_cast<EMaterialValueType>(LastExpression->GetOutputType(OutputIdx));
+		int32 CurrentOutputCount = 0;
+
+		switch (CurrentOutputType)
+		{
+			case MCT_Float:
+			case MCT_Float1:
+				CurrentOutputCount = 1;
+				break;
+
+			case MCT_Float2:
+				CurrentOutputCount = 2;
+				break;
+
+			case MCT_Float3:
+				CurrentOutputCount = 3;
+				break;
+
+			case MCT_Float4:
+				CurrentOutputCount = 4;
+				break;
+
+			default:
+				continue; // For loop
+		}
+
+		if (CurrentOutputCount > OutputCount)
+		{
+			BestMatch = OutputIdx;
+			OutputCount = CurrentOutputCount;
+
+			if (CurrentOutputCount >= FloatsForPropertyType)
+			{
+				break;
+			}
+		}
+	}
+
+	if (BestMatch != INDEX_NONE)
+	{
+		BasePreviewMaterial->GetEditorOnlyData()->EmissiveColor.Expression = LastExpression;
+		BasePreviewMaterial->GetEditorOnlyData()->EmissiveColor.OutputIndex = BestMatch;
+	}
+}
+
+bool UDMMaterialSlot::SetLayerMaterialPropertyAndReplaceOthers(UDMMaterialLayerObject* InLayer, EDMMaterialPropertyType InMaterialProperty, 
+	EDMMaterialPropertyType InReplaceWithProperty)
+{
+	if (!IsComponentValid())
+	{
+		return false;
+	}
+
+	UDynamicMaterialModelEditorOnlyData* ModelEditorOnlyData = GetMaterialModelEditorOnlyData();
+	check(ModelEditorOnlyData);
+
+	UDMMaterialSlot* CurrentSlot = ModelEditorOnlyData->GetSlotForMaterialProperty(InMaterialProperty);
+
+	if (!CurrentSlot || CurrentSlot == this)
+	{
+		if (GUndo)
+		{
+			InLayer->Modify();
+		}
+
+		InLayer->SetMaterialProperty(InMaterialProperty);
+		return false; // Could be caused by asynchronous input
+	}
+
+	{
+		FDMUpdateGuard Guard;
+		
+		for (TObjectPtr<UDMMaterialLayerObject>& Layer : CurrentSlot->LayerObjects)
+		{
+			if (Layer->GetMaterialProperty() == InMaterialProperty)
+			{
+				if (GUndo)
+				{
+					Layer->Modify();
+				}
+
+				Layer->SetMaterialProperty(InReplaceWithProperty);
+			}
+
+			if (UDMMaterialStage* BaseStage = Layer->GetStage(EDMMaterialLayerStage::Base))
+			{
+				TArray<FDMMaterialStageConnection>& CurrentSlotStageInputMap = BaseStage->GetInputConnectionMap();
+
+				for (int32 InputIdx = 0; InputIdx < CurrentSlotStageInputMap.Num(); ++InputIdx)
+				{
+					for (int32 ChannelIdx = 0; ChannelIdx < CurrentSlotStageInputMap[InputIdx].Channels.Num(); ++ChannelIdx)
+					{
+						FDMMaterialStageConnectorChannel& Channel = CurrentSlotStageInputMap[InputIdx].Channels[ChannelIdx];
+
+						if (Channel.SourceIndex == FDMMaterialStageConnectorChannel::PREVIOUS_STAGE
+							&& Channel.MaterialProperty == InMaterialProperty)
+						{
+							// Delve into class internals, avoiding the const issues above.
+							Channel.MaterialProperty = InReplaceWithProperty;
+						}
+					}
+				}
+			}
+		}
+	}
+
+	if (!CurrentSlot->LayerObjects.IsEmpty())
+	{
+		if (UDMMaterialStage* Stage = CurrentSlot->LayerObjects[0]->GetFirstEnabledStage(EDMMaterialLayerStage::All))
+		{
+			Stage->Update(EDMUpdateType::Structure);
+		}
+	}
+
+	return InLayer->SetMaterialProperty(InMaterialProperty);
+}
+
+void UDMMaterialSlot::DoClean()
+{
+	if (IsComponentValid())
+	{
+		UpdatePreviewMaterial(EDMMaterialLayerStage::Base);
+		UpdatePreviewMaterial(EDMMaterialLayerStage::Mask);
+	}
+
+	Super::DoClean();
+}
+
+FString UDMMaterialSlot::GetComponentPathComponent() const
+{
+	// @TODO These will probably need to change when other slot types are opened up.
+	switch (Index)
+	{
+		case 0:
+			return UDynamicMaterialModelEditorOnlyData::RGBSlotPathToken;
+
+		case 1:
+			return UDynamicMaterialModelEditorOnlyData::OpacitySlotPathToken;
+
+		default:
+			return FString::Printf(
+				TEXT("%s%hc%i%hc"), 
+				*UDynamicMaterialModelEditorOnlyData::SlotsPathToken,
+				FDMComponentPath::ParameterOpen,
+				Index, 
+				FDMComponentPath::ParameterClose
+			);
+	}
+}
+
+UDMMaterialComponent* UDMMaterialSlot::GetSubComponentByPath(FDMComponentPath& InPath, const FDMComponentPathSegment& InPathSegment) const
+{
+	if (InPathSegment.GetToken() == LayersPathToken)
+	{
+		int32 LayerIndex = INDEX_NONE;
+
+		if (InPathSegment.GetParameter(LayerIndex))
+		{
+			if (LayerObjects.IsValidIndex(LayerIndex))
+			{
+				return LayerObjects[LayerIndex];
+			}
+		}
+	}
+
+	return Super::GetSubComponentByPath(InPath, InPathSegment);
+}
+
+void UDMMaterialSlot::PostEditorDuplicate(UDynamicMaterialModel* InMaterialModel, UDMMaterialComponent* InParent)
+{
+	Super::PostEditorDuplicate(InMaterialModel, InParent);
+
+	UDynamicMaterialModelEditorOnlyData* ModelEditorOnlyData = UDynamicMaterialModelEditorOnlyData::Get(InMaterialModel);
+
+	if (ModelEditorOnlyData && GetOuter() != ModelEditorOnlyData)
+	{
+		Rename(nullptr, ModelEditorOnlyData, UE::DynamicMaterial::RenameFlags);
+	}
+
+	for (const TObjectPtr<UDMMaterialLayerObject>& Layer : LayerObjects)
+	{
+		Layer->PostEditorDuplicate(InMaterialModel, this);
+	}
+}
+
+bool UDMMaterialSlot::Modify(bool bInAlwaysMarkDirty)
+{
+	const bool bSaved = Super::Modify(bInAlwaysMarkDirty);
+
+	for (const TObjectPtr<UDMMaterialLayerObject>& Layer : LayerObjects)
+	{
+		Layer->Modify(bInAlwaysMarkDirty);
+	}
+
+	return bSaved;
+}
+
+void UDMMaterialSlot::PostEditUndo()
+{
+	Super::PostEditUndo();
+
+	if (!IsComponentValid())
+	{
+		return;
+	}
+
+	UDynamicMaterialModelEditorOnlyData* ModelEditorOnlyData = GetMaterialModelEditorOnlyData();
+	check(ModelEditorOnlyData);
+
+	UDynamicMaterialModel* MaterialModel = ModelEditorOnlyData->GetMaterialModel();
+	check(MaterialModel);
+
+	for (const TObjectPtr<UDMMaterialLayerObject>& Layer : LayerObjects)
+	{
+		if (GUndo)
+		{
+			Layer->Modify();
+		}
+
+		Layer->PostEditorDuplicate(MaterialModel, this);
+	}
+
+	MarkComponentDirty();
+
+	Update(EDMUpdateType::Structure);
+
+	// Fire all of these to make sure everything is updated.
+	OnPropertiesUpdateDelegate.Broadcast(this);
+	OnLayersUpdateDelegate.Broadcast(this);
+
+	UpdateOutputConnectorTypes();
+}
+
+void UDMMaterialSlot::PostLoad()
+{
+	Super::PostLoad();
+
+	PRAGMA_DISABLE_DEPRECATION_WARNINGS
+
+	if (!Layers.IsEmpty())
+	{
+		ConvertDeprecatedLayers(Layers);
+	}
+
+	PRAGMA_ENABLE_DEPRECATION_WARNINGS
+}
+
+void UDMMaterialSlot::ConvertDeprecatedLayers(TArray<FDMMaterialLayer>& InLayers)
+{
+	FDMUpdateGuard Guard;
+
+	MarkPackageDirty();
+
+	UDynamicMaterialModel* MaterialModel = GetMaterialModelEditorOnlyData()->GetMaterialModel();
+
+	for (const FDMMaterialLayer& Layer : InLayers)
+	{
+		UDMMaterialLayerObject* NewLayer = AddLayerWithMask(Layer.MaterialProperty, Layer.Base, Layer.Mask);
+		NewLayer->SetLayerName(Layer.LayerName);
+		NewLayer->SetEnabled(Layer.bEnabled);
+		NewLayer->SetTextureUVLinkEnabled(Layer.bLinkedUVs);
+
+		if (Layer.Base)
+		{
+			Layer.Base->SetEnabled(Layer.bBaseEnabled);
+		}
+
+		if (Layer.Mask)
+		{
+			Layer.Mask->SetEnabled(Layer.bMaskEnabled);
+		}
+
+		NewLayer->PostEditorDuplicate(MaterialModel, this);
+	}
+
+	InLayers.Empty();	
+
+	if (!LayerObjects.IsEmpty())
+	{
+		LayerObjects[0]->Update(EDMUpdateType::Structure);
+	}
+	else
+	{
+		Update(EDMUpdateType::Structure);
+	}
+}
+
+#undef LOCTEXT_NAMESPACE
