[MaterialDesigner]

diff --git a/Engine/Plugins/Experimental/DynamicMaterial/Source/DynamicMaterialEditor/Private/Components/DMMaterialStage.cpp b/Engine/Plugins/Experimental/DynamicMaterial/Source/DynamicMaterialEditor/Private/Components/DMMaterialStage.cpp
new file mode 100644
--- /dev/null
+++ b/Engine/Plugins/Experimental/DynamicMaterial/Source/DynamicMaterialEditor/Private/Components/DMMaterialStage.cpp
@@ -0,0 +1,2082 @@
+// Copyright Epic Games, Inc. All Rights Reserved.
+
+#include "Components/DMMaterialStage.h"
+#include "AssetRegistry/AssetRegistryModule.h"
+#include "Components/DMMaterialLayer.h"
+#include "Components/DMMaterialProperty.h"
+#include "Components/DMMaterialSlot.h"
+#include "Components/DMMaterialStageBlend.h"
+#include "Components/DMMaterialStageExpression.h"
+#include "Components/DMMaterialStageFunction.h"
+#include "Components/DMMaterialStageGradient.h"
+#include "Components/DMMaterialStageInput.h"
+#include "Components/DMMaterialStageSource.h"
+#include "Components/DMMaterialStageThroughputLayerBlend.h"
+#include "Components/DMMaterialValue.h"
+#include "Components/DMTextureUV.h"
+#include "Components/MaterialStageExpressions/DMMSEMathBase.h"
+#include "Components/MaterialStageInputs/DMMSIExpression.h"
+#include "Components/MaterialStageInputs/DMMSIFunction.h"
+#include "Components/MaterialStageInputs/DMMSIGradient.h"
+#include "Components/MaterialStageInputs/DMMSISlot.h"
+#include "Components/MaterialStageInputs/DMMSITextureUV.h"
+#include "Components/MaterialStageInputs/DMMSIValue.h"
+#include "Containers/TransArray.h"
+#include "DMComponentPath.h"
+#include "DynamicMaterialEditorModule.h"
+#include "DynamicMaterialEditorSettings.h"
+#include "DynamicMaterialModule.h"
+#include "Factories/MaterialFactoryNew.h"
+#include "Materials/Material.h"
+#include "Materials/MaterialExpression.h"
+#include "Materials/MaterialExpressionAppendVector.h"
+#include "Materials/MaterialExpressionComponentMask.h"
+#include "Materials/MaterialInstanceDynamic.h"
+#include "MaterialValueType.h"
+#include "Model/DMMaterialBuildState.h"
+#include "Model/DMMaterialBuildUtils.h"
+#include "Model/DynamicMaterialModel.h"
+#include "Model/DynamicMaterialModelEditorOnlyData.h"
+#include "UObject/Package.h"
+
+#define LOCTEXT_NAMESPACE "DMMaterialStage"
+
+const FString UDMMaterialStage::SourcePathToken = FString(TEXT("Source"));
+const FString UDMMaterialStage::InputsPathToken = FString(TEXT("Inputs"));
+
+UDMMaterialStage* UDMMaterialStage::CreateMaterialStage(UDMMaterialLayerObject* InLayer)
+{
+	UObject* Outer = IsValid(InLayer) ? (UObject*)InLayer : (UObject*)GetTransientPackage();
+
+	return NewObject<UDMMaterialStage>(Outer, NAME_None, RF_Transactional);
+}
+
+UDMMaterialStage::UDMMaterialStage()
+	: Source(nullptr)
+	, bEnabled(true)
+	, bCanChangeSource(true)
+	, bIsBeingEdited(false)
+	, PreviewMaterialBase(nullptr)
+	, PreviewMaterialDynamic(nullptr)
+{
+	EditableProperties.Add(GET_MEMBER_NAME_CHECKED(UDMMaterialStage, Source));
+	EditableProperties.Add(GET_MEMBER_NAME_CHECKED(UDMMaterialStage, Inputs));
+}
+
+void UDMMaterialStage::OnComponentAdded()
+{
+	if (!IsComponentValid())
+	{
+		return;
+	}
+
+	Super::OnComponentAdded();
+
+	ResetInputConnectionMap();
+
+	for (UDMMaterialStageInput* Input : Inputs)
+	{
+		Input->SetComponentState(EDMComponentLifetimeState::Added);
+	}
+		
+	if (Source)
+	{
+		if (GUndo)
+		{
+			Source->Modify();
+		}
+
+		Source->SetComponentState(EDMComponentLifetimeState::Added);
+	}
+
+	AddDelegates();
+}
+
+void UDMMaterialStage::OnComponentRemoved()
+{
+	Super::OnComponentRemoved();
+
+	RemoveDelegates();
+
+	for (UDMMaterialStageInput* Input : Inputs)
+	{
+		Input->SetComponentState(EDMComponentLifetimeState::Removed);
+	}
+		
+	if (Source)
+	{
+		if (GUndo)
+		{
+			Source->Modify();
+		}
+
+		Source->SetComponentState(EDMComponentLifetimeState::Removed);
+	}
+}
+
+UDMMaterialComponent* UDMMaterialStage::GetParentComponent() const
+{
+	return GetLayer();
+}
+
+void UDMMaterialStage::PostEditorDuplicate(UDynamicMaterialModel* InMaterialModel, UDMMaterialComponent* InParent)
+{
+	Super::PostEditorDuplicate(InMaterialModel, InParent);
+
+	if (GetOuter() != InParent)
+	{
+		Rename(nullptr, InParent, UE::DynamicMaterial::RenameFlags);
+	}
+
+	PreviewMaterialBase = nullptr;
+	PreviewMaterialDynamic = nullptr;
+
+	AddDelegates();
+
+	if (Source)
+	{
+		if (GUndo)
+		{
+			Source->Modify();
+		}
+
+		Source->PostEditorDuplicate(InMaterialModel, this);
+	}
+
+	for (const TObjectPtr<UDMMaterialStageInput>& Input : Inputs)
+	{
+		if (Input)
+		{
+			Input->PostEditorDuplicate(InMaterialModel, this);
+		}
+	}
+}
+
+bool UDMMaterialStage::Modify(bool bInAlwaysMarkDirty)
+{
+	const bool bSaved = Super::Modify(bInAlwaysMarkDirty);
+
+	if (Source)
+	{
+		Source->Modify(bInAlwaysMarkDirty);
+	}
+
+	for (const TObjectPtr<UDMMaterialStageInput>& Input : Inputs)
+	{
+		if (Input)
+		{
+			Input->Modify(bInAlwaysMarkDirty);
+		}
+	}
+
+	return bSaved;
+}
+
+FString UDMMaterialStage::GetComponentPathComponent() const
+{
+	if (UDMMaterialLayerObject* Layer = GetLayer())
+	{
+		FString TypeStr = "?";
+
+		switch (Layer->GetStageType(this))
+		{
+			case EDMMaterialLayerStage::Base:
+				TypeStr = UDMMaterialLayerObject::BasePathToken;
+				break;
+
+			case EDMMaterialLayerStage::Mask:
+				TypeStr = UDMMaterialLayerObject::MaskPathToken;
+				break;
+
+			default:
+				TypeStr = FString::FromInt(Layer->GetAllStages().IndexOfByKey(this));
+				break;
+		}
+
+		return FString::Printf(
+			TEXT("%s%hc%s%hc"),
+			*UDMMaterialLayerObject::StagesPathToken,
+			FDMComponentPath::ParameterOpen,
+			*TypeStr,
+			FDMComponentPath::ParameterClose
+		);
+	}
+
+	return Super::GetComponentPathComponent();
+}
+
+void UDMMaterialStage::GetComponentPathInternal(TArray<FString>& OutChildComponentPathComponents) const
+{
+	// Strip off the type index of the substage
+	if (OutChildComponentPathComponents.IsEmpty() == false && Source)
+	{
+		if (OutChildComponentPathComponents.Last() == Source->GetComponentPathComponent())
+		{
+			OutChildComponentPathComponents.Last() = Source->GetClass()->GetName();
+		}
+	}
+
+	Super::GetComponentPathInternal(OutChildComponentPathComponents);
+}
+
+UDMMaterialComponent* UDMMaterialStage::GetSubComponentByPath(FDMComponentPath& InPath,
+	const FDMComponentPathSegment& InPathSegment) const
+{
+	if (InPathSegment.GetToken() == SourcePathToken)
+	{
+		return Source;
+	}
+
+	if (InPathSegment.GetToken() == InputsPathToken)
+	{
+		int32 InputIndex = INDEX_NONE;
+		FString InputType;
+
+		if (InPathSegment.GetParameter(InputIndex))
+		{
+			if (Inputs.IsValidIndex(InputIndex))
+			{
+				return Inputs[InputIndex];
+			}
+		}
+		else if (InPathSegment.GetParameter(InputType))
+		{
+			for (UDMMaterialStageInput* Input : Inputs)
+			{
+				if (Input && (Input->GetClass()->GetName().Equals(InputType)
+					|| Input->GetClass()->GetName().Equals(UDMMaterialStageInput::StageInputPrefixStr + InputType)))
+				{
+					return Input;
+				}
+			}
+		}
+	}
+
+	return Super::GetSubComponentByPath(InPath, InPathSegment);
+}
+
+void UDMMaterialStage::AddDelegates()
+{
+	if (UDMMaterialLayerObject* Layer = GetLayer())
+	{
+		if (UDMMaterialSlot* Slot = Layer->GetSlot())
+		{
+			if (UDynamicMaterialModelEditorOnlyData* ModelEditorOnlyData = Slot->GetMaterialModelEditorOnlyData())
+			{
+				ModelEditorOnlyData->GetOnValueUpdateDelegate().AddUObject(this, &UDMMaterialStage::OnValueUpdated);
+				ModelEditorOnlyData->GetOnTextureUVUpdateDelegate().AddUObject(this, &UDMMaterialStage::OnTextureUVUpdated);
+			}
+		}
+	}
+}
+
+void UDMMaterialStage::RemoveDelegates()
+{
+	if (UDMMaterialLayerObject* Layer = GetLayer())
+	{
+		if (UDMMaterialSlot* Slot = Layer->GetSlot())
+		{
+			if (UDynamicMaterialModelEditorOnlyData* ModelEditorOnlyData = Slot->GetMaterialModelEditorOnlyData())
+			{
+				ModelEditorOnlyData->GetOnValueUpdateDelegate().RemoveAll(this);
+				ModelEditorOnlyData->GetOnTextureUVUpdateDelegate().RemoveAll(this);
+			}
+		}
+	}
+}
+
+void UDMMaterialStage::OnValueUpdated(UDynamicMaterialModel* InMaterialModel, UDMMaterialValue* InValue)
+{
+	if (PreviewMaterialDynamic && InMaterialModel && InValue)
+	{
+		InValue->SetMIDParameter(PreviewMaterialDynamic);
+	}
+}
+
+void UDMMaterialStage::OnTextureUVUpdated(UDynamicMaterialModel* InMaterialModel, UDMTextureUV* InTextureUV)
+{
+	if (!UDynamicMaterialEditorSettings::Get()->bPreviewImagesUseTextureUVs)
+	{
+		return;
+	}
+
+	if (!PreviewMaterialDynamic || !InMaterialModel || !InTextureUV)
+	{
+		return;
+	}
+
+	UDMMaterialStage* Stage = InTextureUV->GetTypedParent<UDMMaterialStage>(false);
+
+	if (Stage != this)
+	{
+		return;
+	}
+
+	InTextureUV->SetMIDParameters(PreviewMaterialDynamic);
+
+	UDMMaterialLayerObject* Layer = Stage->GetLayer();
+
+	if (!Layer || !Layer->IsTextureUVLinkEnabled() || Layer->GetStageType(Stage) != EDMMaterialLayerStage::Base)
+	{
+		return;
+	}
+
+	UDMMaterialStageThroughput* Throughput = Cast<UDMMaterialStageThroughput>(Stage->GetSource());
+
+	if (!Throughput || !Throughput->SupportsLayerMaskTextureUVLink())
+	{
+		return;
+	}
+
+	UDMMaterialStage* MaskStage = Layer->GetStage(EDMMaterialLayerStage::Mask);
+
+	if (!MaskStage)
+	{
+		return;
+	}
+
+	if (UMaterialInstanceDynamic* MaskMID = Cast<UMaterialInstanceDynamic>(MaskStage->GetPreviewMaterial()))
+	{
+		InTextureUV->SetMIDParameters(MaskMID);
+	}
+}
+
+void UDMMaterialStage::ApplyWholeLayerValue(UDMMaterialValue* InValue)
+{
+	if (!InValue || !InValue->IsWholeLayerValue())
+	{
+		return;
+	}
+
+	UDMMaterialLayerObject* Layer = GetLayer();
+
+	if (!Layer)
+	{
+		return;
+	}
+
+	if (!Layer 
+		|| Layer->GetStageType(this) != EDMMaterialLayerStage::Base
+		|| !Layer->IsStageEnabled(EDMMaterialLayerStage::Mask))
+	{
+		return;
+	}
+
+	UDMMaterialStage* MaskStage = Layer->GetStage(EDMMaterialLayerStage::Mask);
+
+	if (!IsValid(MaskStage))
+	{
+		return;
+	}
+
+	if (GUndo)
+	{
+		MaskStage->Modify();
+	}
+
+	MaskStage->ChangeInput_Value(
+		UDMMaterialStageThroughputLayerBlend::InputMaskSource,
+		FDMMaterialStageConnectorChannel::WHOLE_CHANNEL,
+		InValue,
+		FDMMaterialStageConnectorChannel::FOURTH_CHANNEL
+	);
+}
+
+UDMMaterialLayerObject* UDMMaterialStage::GetLayer() const
+{
+	return Cast<UDMMaterialLayerObject>(GetOuterSafe());
+}
+
+bool UDMMaterialStage::SetEnabled(bool bInEnabled)
+{
+	if (bEnabled == bInEnabled)
+	{
+		return false;
+	}
+
+	bEnabled = bInEnabled;
+
+	Update(EDMUpdateType::Structure);
+
+	return true;
+}
+
+void UDMMaterialStage::SetSource(UDMMaterialStageSource* InSource)
+{
+	if (!IsComponentValid())
+	{
+		return;
+	}
+
+	if (!bCanChangeSource)
+	{
+		return;
+	}
+
+	if (Source)
+	{
+		if (GUndo)
+		{
+			Source->Modify();
+		}
+
+		Source->SetComponentState(EDMComponentLifetimeState::Removed);
+	}
+
+	Source = InSource;
+
+	if (FDMUpdateGuard::CanUpdate())
+	{
+		ResetInputConnectionMap();
+
+		if (IsComponentAdded())
+		{
+			if (GUndo)
+			{
+				Source->Modify();
+			}
+
+			Source->SetComponentState(EDMComponentLifetimeState::Added);
+		}
+
+		Update(EDMUpdateType::Structure);
+	}
+}
+
+FText UDMMaterialStage::GetComponentDescription() const
+{
+	if (Source)
+	{
+		if (Source->IsComponentValid())
+		{
+			return Source->GetStageDescription();
+		}
+	}
+
+	return LOCTEXT("StageDescription", "Material Stage");
+}
+
+EDMValueType UDMMaterialStage::GetSourceType(const FDMMaterialStageConnectorChannel& InChannel) const
+{
+	if (InChannel.SourceIndex == FDMMaterialStageConnectorChannel::PREVIOUS_STAGE)
+	{
+		UDMMaterialLayerObject* Layer = GetLayer();
+		check(Layer);
+
+		EDMMaterialPropertyType StageProperty = Layer->GetMaterialProperty();
+		check(StageProperty != EDMMaterialPropertyType::None && StageProperty != EDMMaterialPropertyType::Any);
+
+		const UDMMaterialLayerObject* PreviousLayer = Layer->GetPreviousLayer(StageProperty, EDMMaterialLayerStage::Base);
+		check(PreviousLayer);
+		check(PreviousLayer->GetStage(EDMMaterialLayerStage::Mask));
+		check(PreviousLayer->GetStage(EDMMaterialLayerStage::Mask)->GetSource());
+
+		const TArray<FDMMaterialStageConnector>& OutputConnectors = PreviousLayer->GetStage(EDMMaterialLayerStage::Mask)->GetSource()->GetOutputConnectors();
+		check(OutputConnectors.IsValidIndex(InChannel.OutputIndex));
+
+		return OutputConnectors[InChannel.OutputIndex].Type;
+	}
+	else
+	{
+		const int32 StageInputIdx = InChannel.SourceIndex - FDMMaterialStageConnectorChannel::FIRST_STAGE_INPUT;
+		check(Inputs.IsValidIndex(StageInputIdx));
+
+		const TArray<FDMMaterialStageConnector>& OutputConnectors = Inputs[StageInputIdx]->GetOutputConnectors();
+		check(OutputConnectors.IsValidIndex(InChannel.OutputIndex));
+
+		return OutputConnectors[InChannel.OutputIndex].Type;
+	}
+}
+
+bool UDMMaterialStage::IsInputMapped(int32 InputIndex) const
+{
+	if (!InputConnectionMap.IsValidIndex(InputIndex))
+	{
+		return false;
+	}
+
+	for (const FDMMaterialStageConnectorChannel& Channel : InputConnectionMap[InputIndex].Channels)
+	{
+		if (Channel.SourceIndex != FDMMaterialStageConnectorChannel::NO_SOURCE)
+		{
+			return true;
+		}
+	}
+
+	return false;
+}
+
+void UDMMaterialStage::Update(EDMUpdateType InUpdateType)
+{
+	if (!IsComponentValid())
+	{
+		return;
+	}
+
+	if (HasComponentBeenRemoved())
+	{
+		return;
+	}
+
+	if (InUpdateType == EDMUpdateType::Structure)
+	{
+		MarkComponentDirty();
+		VerifyAllInputMaps();
+	}
+
+	Super::Update(InUpdateType);
+
+	if (UDMMaterialStage* NextStage = GetNextStage())
+	{
+		NextStage->Update(InUpdateType);
+	}
+	else if (UDMMaterialLayerObject* Layer = GetLayer())
+	{
+		Layer->Update(InUpdateType);
+	}	
+}
+
+void UDMMaterialStage::InputUpdated(UDMMaterialStageInput* InInput, EDMUpdateType InUpdateType)
+{
+	if (!IsComponentValid())
+	{
+		return;
+	}
+
+	UDMMaterialStageThroughput* Throughput = Cast<UDMMaterialStageThroughput>(Source);
+
+	if (!Throughput)
+	{
+		return;
+	}
+
+	int32 InputIdx = Inputs.Find(InInput);
+	check(InputIdx != INDEX_NONE);
+
+	InputIdx += FDMMaterialStageConnectorChannel::FIRST_STAGE_INPUT;
+
+	for (int32 InputMapIdx = 0; InputMapIdx < Inputs.Num() && InputMapIdx < InputConnectionMap.Num(); ++InputMapIdx)
+	{
+		bool bIsUsedInInput = false;
+
+		for (const FDMMaterialStageConnectorChannel& Channel : InputConnectionMap[InputMapIdx].Channels)
+		{
+			if (Channel.SourceIndex == InputIdx)
+			{
+				bIsUsedInInput = true;
+				break;
+			}
+		}
+
+		if (bIsUsedInInput)
+		{
+			Throughput->InputUpdated(InputMapIdx, InUpdateType);
+		}
+	}
+}
+
+void UDMMaterialStage::ResetInputConnectionMap()
+{
+	if (!IsComponentValid())
+	{
+		return;
+	}
+
+	VerifyAllInputMaps();
+}
+
+void UDMMaterialStage::GenerateExpressions(const TSharedRef<FDMMaterialBuildState>& InBuildState) const
+{
+	if (!IsComponentValid() || !IsComponentAdded())
+	{
+		return;
+	}
+
+	check(Source);
+
+	if (InBuildState->HasStage(this))
+	{
+		return;
+	}
+
+	UDMMaterialLayerObject* Layer = GetLayer();
+
+	if (!Layer)
+	{
+		return;
+	}
+
+	TArray<UMaterialExpression*> StageExpressions;
+
+	Source->GenerateExpressions(InBuildState);
+	const TArray<UMaterialExpression*>& StageSourceExpressions = InBuildState->GetStageSourceExpressions(Source);
+
+	if (!StageSourceExpressions.IsEmpty())
+	{
+		if (UDMMaterialStageThroughput* Throughput = Cast<UDMMaterialStageThroughput>(Source))
+		{
+			const TArray<FDMMaterialStageConnector>& ThroughputInputs = Throughput->GetInputConnectors();
+
+			for (int32 InputIdx = 0; InputIdx < ThroughputInputs.Num() && InputIdx < InputConnectionMap.Num(); ++InputIdx)
+			{
+				FDMMaterialStageConnectorChannel Channel;
+				TArray<UMaterialExpression*> Expressions;
+
+				const int32 OutputChannelOverride = Throughput->GetOutputChannelOverride(Channel.OutputIndex);
+
+				if (OutputChannelOverride != INDEX_NONE)
+				{
+					Channel.OutputChannel = OutputChannelOverride;
+				}
+
+				const int32 NodeOutputIndex = Throughput->ResolveInput(InBuildState, InputIdx, Channel, Expressions);
+				
+				if (!Expressions.IsEmpty() && NodeOutputIndex != INDEX_NONE)
+				{
+					StageExpressions.Append(Expressions);
+
+					Throughput->ConnectOutputToInput(InBuildState, ThroughputInputs[InputIdx].Index,
+						Expressions.Last(), NodeOutputIndex, Channel.OutputChannel);
+				}
+			}
+		}
+
+		StageExpressions.Append(InBuildState->GetStageSourceExpressions(Source));		
+	}
+
+	InBuildState->AddStageExpressions(this, StageExpressions);
+}
+
+bool UDMMaterialStage::SetBeingEdited(bool bInBeingEdited)
+{ 
+	if (bIsBeingEdited == bInBeingEdited)
+	{
+		return false;
+	}
+
+	bIsBeingEdited = bInBeingEdited; 
+
+	return true;
+}
+
+TMap<EDMMaterialPropertyType, UDMMaterialLayerObject*> UDMMaterialStage::GetPreviousStagesPropertyMap()
+{
+	UDMMaterialLayerObject* Layer = GetLayer();
+
+	if (!Layer)
+	{
+		return {};
+	}
+
+	UDMMaterialSlot* Slot = Layer->GetSlot();
+	check(Slot);
+
+	TMap<EDMMaterialPropertyType, UDMMaterialLayerObject*> PropertyMap;
+	const TArray<TObjectPtr<UDMMaterialLayerObject>> Layers = Slot->GetLayers();
+
+	for (UDMMaterialLayerObject* LayerIter : Layers)
+	{
+		if (LayerIter == Layer)
+		{
+			break;
+		}
+
+		PropertyMap.FindOrAdd(LayerIter->GetMaterialProperty()) = LayerIter;
+	}
+
+	return PropertyMap;
+}
+
+TMap<EDMMaterialPropertyType, UDMMaterialLayerObject*> UDMMaterialStage::GetPropertyMap()
+{
+	UDMMaterialLayerObject* Layer = GetLayer();
+
+	if (!Layer)
+	{
+		return {};
+	}
+
+	UDMMaterialSlot* Slot = Layer->GetSlot();
+	check(Slot);
+
+	TMap<EDMMaterialPropertyType, UDMMaterialLayerObject*> PropertyMap;
+	const TArray<TObjectPtr<UDMMaterialLayerObject>> Layers = Slot->GetLayers();
+
+	for (UDMMaterialLayerObject* LayerIter : Layers)
+	{
+		PropertyMap.FindOrAdd(LayerIter->GetMaterialProperty()) = LayerIter;
+
+		if (Layer->HasValidStage(this))
+		{
+			break;
+		}
+	}
+
+	return PropertyMap;
+}
+
+void UDMMaterialStage::RemoveUnusedInputs()
+{
+	if (!IsComponentValid())
+	{
+		return;
+	}
+
+	TArray<UDMMaterialStageInput*> UnusedInputs = Inputs;
+
+	for (const FDMMaterialStageConnection& Connection : InputConnectionMap)
+	{
+		for (const FDMMaterialStageConnectorChannel& Channel : Connection.Channels)
+		{
+			if (Channel.SourceIndex == 0)
+			{
+				continue;
+			}
+
+			const int32 InputIdx = Channel.SourceIndex - FDMMaterialStageConnectorChannel::FIRST_STAGE_INPUT;
+
+			if (!Inputs.IsValidIndex(InputIdx))
+			{
+				continue;
+			}
+
+			UnusedInputs.Remove(Inputs[InputIdx]);
+		}
+	}
+
+	for (UDMMaterialStageInput* Input : UnusedInputs)
+	{
+		int32 InputIdx;
+
+		if (!Inputs.Find(Input, InputIdx))
+		{
+			continue;
+		}
+
+		InputIdx += FDMMaterialStageConnectorChannel::FIRST_STAGE_INPUT;
+
+		for (FDMMaterialStageConnection& Connection : InputConnectionMap)
+		{
+			for (int32 ChannelIdx = 0; ChannelIdx < Connection.Channels.Num(); ++ChannelIdx)
+			{
+				// Delete channel if we're using a deleted input
+				if (Connection.Channels[ChannelIdx].SourceIndex == InputIdx)
+				{
+					Connection.Channels.RemoveAt(ChannelIdx);
+					--ChannelIdx;
+					continue;
+				}
+
+				if (Connection.Channels[ChannelIdx].SourceIndex > InputIdx)
+				{
+					Connection.Channels[ChannelIdx].SourceIndex -= 1;
+				}
+			}
+		}
+
+		InputIdx -= FDMMaterialStageConnectorChannel::FIRST_STAGE_INPUT;
+		Inputs.RemoveAt(InputIdx);
+
+		Input->SetComponentState(EDMComponentLifetimeState::Removed);
+	}
+}
+
+UDMMaterialStageSource* UDMMaterialStage::ChangeInput_PreviousStage(int32 InInputIdx, int32 InInputChannel, EDMMaterialPropertyType InPreviousStageProperty, int32 InOutputIdx, int32 InOutputChannel)
+{
+	check(Source);
+
+	UDMMaterialStageThroughput* Throughput = Cast<UDMMaterialStageThroughput>(Source);
+	check(Throughput);
+
+	const TArray<FDMMaterialStageConnector>& InputConnectors = Throughput->GetInputConnectors();
+	check(InputConnectors.IsValidIndex(InInputIdx));
+
+	UDMMaterialLayerObject* Layer = GetLayer();
+	check(Layer);
+
+	EDMMaterialPropertyType StageProperty = Layer->GetMaterialProperty();
+	check(StageProperty != EDMMaterialPropertyType::None);
+
+	const UDMMaterialLayerObject* PreviousLayer = Layer->GetPreviousLayer(StageProperty, EDMMaterialLayerStage::Base);
+	UDMMaterialStageSource* PreviousSource = nullptr;
+	
+	if (PreviousLayer)
+	{
+		if (UDMMaterialStage* Stage = PreviousLayer->GetLastEnabledStage(EDMMaterialLayerStage::All))
+		{
+			PreviousSource = Stage->GetSource();
+		}
+	}
+
+	if (PreviousSource)
+	{
+		const TArray<FDMMaterialStageConnector>& PreviousStageOutputs = PreviousSource->GetOutputConnectors();
+		check(PreviousStageOutputs.IsValidIndex(InOutputIdx));
+		check(Throughput->CanInputConnectTo(InInputIdx, PreviousStageOutputs[InOutputIdx], InOutputChannel));
+	}
+	else
+	{
+		check(InOutputIdx == 0);
+		check(Throughput->GetInputConnectors().IsEmpty() == false);
+		check(Throughput->GetInputConnectors()[0].IsCompatibleWith(EDMValueType::VT_Float3_RGB));
+	}
+
+	UpdateInputMap(
+		InInputIdx,
+		FDMMaterialStageConnectorChannel::PREVIOUS_STAGE,
+		InInputChannel,
+		InOutputIdx,
+		InOutputChannel,
+		InPreviousStageProperty
+	);
+
+	return PreviousSource;
+}
+
+UDMMaterialStageInputValue* UDMMaterialStage::ChangeInput_NewLocalValue(int32 InInputIdx, int32 InInputChannel, EDMValueType InType, int32 InOutputChannel)
+{
+	check(Source);
+
+	UDMMaterialStageThroughput* Throughput = Cast<UDMMaterialStageThroughput>(Source);
+	check(Throughput);
+
+	const TArray<FDMMaterialStageConnector>& InputConnectors = Throughput->GetInputConnectors();
+	check(InputConnectors.IsValidIndex(InInputIdx));
+	check(Throughput->CanInputAcceptType(InInputIdx, InType));
+
+	UDMMaterialLayerObject* Layer = GetLayer();
+	check(Layer);
+
+	UDMMaterialSlot* Slot = Layer->GetSlot();
+	check(Slot);
+
+	UDynamicMaterialModelEditorOnlyData* ModelEditorOnlyData = Slot->GetMaterialModelEditorOnlyData();
+	check(ModelEditorOnlyData);
+
+	UDynamicMaterialModel* MaterialModel = ModelEditorOnlyData->GetMaterialModel();
+	check(MaterialModel);
+
+	// Don't add via the builder, just parent it to the builder.
+	// It won't appear in the global value list.
+	UDMMaterialValue* NewValue = UDMMaterialValue::CreateMaterialValue(MaterialModel, TEXT(""), InType, true);
+	check(NewValue);
+
+	UDMMaterialStageInputValue* InputValue = NewObject<UDMMaterialStageInputValue>(this, NAME_None, RF_Transactional);
+	check(InputValue);
+
+	{
+		const FDMUpdateGuard Guard;
+		InputValue->SetValue(NewValue);
+	}
+
+	AddInput(InputValue);
+
+	UpdateInputMap(
+		InInputIdx,
+		FDMMaterialStageConnectorChannel::FIRST_STAGE_INPUT + Inputs.Num() - 1,
+		InInputChannel,
+		0,
+		InOutputChannel,
+		EDMMaterialPropertyType::None
+	);
+
+	ApplyWholeLayerValue(NewValue);
+
+	return InputValue;
+}
+
+UDMMaterialStageInputValue* UDMMaterialStage::ChangeInput_Value(int32 InInputIdx, int32 InInputChannel, UDMMaterialValue* InValue, int32 OutputChannel)
+{
+	check(Source);
+	check(InValue);
+
+	UDMMaterialStageThroughput* Throughput = Cast<UDMMaterialStageThroughput>(Source);
+	check(Throughput);
+
+	const TArray<FDMMaterialStageConnector>& InputConnectors = Throughput->GetInputConnectors();
+	check(InputConnectors.IsValidIndex(InInputIdx));
+
+	if (OutputChannel == FDMMaterialStageConnectorChannel::WHOLE_CHANNEL)
+	{
+		check(Throughput->CanInputAcceptType(InInputIdx, InValue->GetType()));
+	}
+	else
+	{
+		check(UDMValueDefinitionLibrary::GetValueDefinition(InValue->GetType()).IsFloatType());
+		check(Throughput->CanInputAcceptType(InInputIdx, EDMValueType::VT_Float1));
+	}
+
+	UDMMaterialLayerObject* Layer = GetLayer();
+	check(Layer);
+
+	UDMMaterialSlot* Slot = Layer->GetSlot();
+	check(Slot);
+
+	UDynamicMaterialModelEditorOnlyData* ModelEditorOnlyData = Slot->GetMaterialModelEditorOnlyData();
+	check(ModelEditorOnlyData);
+
+	UDynamicMaterialModel* MaterialModel = ModelEditorOnlyData->GetMaterialModel();
+	check(MaterialModel);
+	check(MaterialModel == InValue->GetMaterialModel());
+
+	UDMMaterialStageInputValue* InputValue = NewObject<UDMMaterialStageInputValue>(this, NAME_None, RF_Transactional);
+	check(InputValue);
+
+	{
+		const FDMUpdateGuard Guard;
+		InputValue->SetValue(InValue);
+	}
+
+	AddInput(InputValue);
+
+	UpdateInputMap(
+		InInputIdx,
+		FDMMaterialStageConnectorChannel::FIRST_STAGE_INPUT + Inputs.Num() - 1,
+		InInputChannel,
+		0,
+		OutputChannel,
+		EDMMaterialPropertyType::None
+	);
+
+	ApplyWholeLayerValue(InValue);
+
+	return InputValue;
+}
+
+UDMMaterialStageInputValue* UDMMaterialStage::ChangeInput_NewValue(int32 InInputIdx, int32 InInputChannel, EDMValueType InType, int32 InOutputChannel)
+{
+	check(Source);
+
+	UDMMaterialStageThroughput* Throughput = Cast<UDMMaterialStageThroughput>(Source);
+	check(Throughput);
+
+	const TArray<FDMMaterialStageConnector>& InputConnectors = Throughput->GetInputConnectors();
+	check(InputConnectors.IsValidIndex(InInputIdx));
+	check(Throughput->CanInputAcceptType(InInputIdx, InType));
+
+	UDMMaterialLayerObject* Layer = GetLayer();
+	check(Layer);
+
+	UDMMaterialSlot* Slot = Layer->GetSlot();
+	check(Slot);
+
+	UDynamicMaterialModelEditorOnlyData* ModelEditorOnlyData = Slot->GetMaterialModelEditorOnlyData();
+	check(ModelEditorOnlyData);
+
+	UDynamicMaterialModel* MaterialModel = ModelEditorOnlyData->GetMaterialModel();
+	check(MaterialModel);
+
+	UDMMaterialValue* NewValue = MaterialModel->AddValue(InType);
+	check(NewValue);
+
+	UDMMaterialStageInputValue* InputValue = NewObject<UDMMaterialStageInputValue>(this, NAME_None, RF_Transactional);
+	check(InputValue);
+
+	{
+		const FDMUpdateGuard Guard;
+		InputValue->SetValue(NewValue);
+	}
+
+	AddInput(InputValue);
+
+	UpdateInputMap(
+		InInputIdx,
+		FDMMaterialStageConnectorChannel::FIRST_STAGE_INPUT + Inputs.Num() - 1,
+		InInputChannel,
+		0,
+		InOutputChannel,
+		EDMMaterialPropertyType::None
+	);
+
+	ApplyWholeLayerValue(NewValue);
+
+	return InputValue;
+}
+
+UDMMaterialStageInputSlot* UDMMaterialStage::ChangeInput_Slot(int32 InInputIdx, int32 InInputChannel, UDMMaterialSlot* InSlot, EDMMaterialPropertyType InProperty, int32 InOutputIdx, int32 InOutputChannel)
+{
+	check(Source);
+
+	UDMMaterialStageThroughput* Throughput = Cast<UDMMaterialStageThroughput>(Source);
+	check(Throughput);
+
+	const TArray<FDMMaterialStageConnector>& InputConnectors = Throughput->GetInputConnectors();
+	check(InputConnectors.IsValidIndex(InInputIdx));
+
+	UDMMaterialLayerObject* Layer = GetLayer();
+	check(Layer);
+
+	UDMMaterialSlot* Slot = Layer->GetSlot();
+	check(Slot);
+
+	UDynamicMaterialModelEditorOnlyData* ModelEditorOnlyData = Slot->GetMaterialModelEditorOnlyData();
+	check(ModelEditorOnlyData == InSlot->GetMaterialModelEditorOnlyData());
+
+	const TArray<EDMMaterialPropertyType> SlotPropertes = ModelEditorOnlyData->GetMaterialPropertiesForSlot(InSlot);
+	check(SlotPropertes.Contains(InProperty));
+
+	const TArray<EDMValueType>& SlotPropertyOutputTypes = InSlot->GetOutputConnectorTypesForMaterialProperty(InProperty);
+	check(SlotPropertyOutputTypes.IsValidIndex(InOutputIdx));
+
+	if (InOutputChannel == FDMMaterialStageConnectorChannel::WHOLE_CHANNEL)
+	{
+		check(Throughput->CanInputAcceptType(InInputIdx, SlotPropertyOutputTypes[InOutputIdx]));
+	}
+	else
+	{
+		check(UDMValueDefinitionLibrary::GetValueDefinition(SlotPropertyOutputTypes[InOutputIdx]).IsFloatType());
+		check(Throughput->CanInputAcceptType(InInputIdx, EDMValueType::VT_Float1));
+	}
+
+	UDMMaterialStageInputSlot* InputSlot = NewObject<UDMMaterialStageInputSlot>(this, NAME_None, RF_Transactional);
+	check(InputSlot);
+
+	{
+		const FDMUpdateGuard Guard;
+		InputSlot->SetSlot(InSlot);
+		InputSlot->SetMaterialProperty(InProperty);
+	}
+
+	AddInput(InputSlot);
+
+	UpdateInputMap(
+		InInputIdx,
+		FDMMaterialStageConnectorChannel::FIRST_STAGE_INPUT + Inputs.Num() - 1,
+		InInputChannel,
+		InOutputIdx,
+		InOutputChannel,
+		EDMMaterialPropertyType::None
+	);
+
+	return InputSlot;
+}
+
+UDMMaterialStageInputExpression* UDMMaterialStage::ChangeInput_Expression(int32 InInputIdx, int32 InInputChannel, TSubclassOf<UDMMaterialStageExpression> InExpressionClass, int32 InOutputIdx, int32 InOutputChannel)
+{
+	check(Source);
+
+	UDMMaterialStageThroughput* Throughput = Cast<UDMMaterialStageThroughput>(Source);
+	check(Throughput);
+
+	const TArray<FDMMaterialStageConnector>& InputConnectors = Throughput->GetInputConnectors();
+	check(InputConnectors.IsValidIndex(InInputIdx));
+
+	check(InExpressionClass.Get());
+	check(InExpressionClass.Get() != UDMMaterialStageExpression::StaticClass());
+
+	UDMMaterialStageExpression* ExpressionCDO = Cast<UDMMaterialStageExpression>(InExpressionClass->GetDefaultObject(true));
+	check(ExpressionCDO);
+	check(!ExpressionCDO->IsInputRequired() || ExpressionCDO->AllowsNestedInputs());
+
+	const TArray<FDMMaterialStageConnector>& OutputConnectors = ExpressionCDO->GetOutputConnectors();
+	check(Throughput->CanInputConnectTo(InInputIdx, OutputConnectors[InOutputIdx], InOutputChannel));
+
+	UDMMaterialStageInputExpression* InputExpression = NewObject<UDMMaterialStageInputExpression>(this, NAME_None, RF_Transactional);
+	check(InputExpression);
+
+	{
+		const FDMUpdateGuard Guard;
+		InputExpression->SetMaterialStageExpressionClass(InExpressionClass);
+	}
+
+	AddInput(InputExpression);
+
+	UpdateInputMap(
+		InInputIdx,
+		FDMMaterialStageConnectorChannel::FIRST_STAGE_INPUT + Inputs.Num() - 1,
+		InInputChannel,
+		InOutputIdx,
+		InOutputChannel,
+		EDMMaterialPropertyType::None
+	);
+
+	return InputExpression;
+}
+
+UDMMaterialStageInputTextureUV* UDMMaterialStage::ChangeInput_UV(int32 InInputIdx, int32 InInputChannel, int32 OutputChannel)
+{
+	UDMMaterialStageInputTextureUV* InputTextureUV = NewObject<UDMMaterialStageInputTextureUV>(this, NAME_None, RF_Transactional);
+	check(InputTextureUV);
+
+	UDMMaterialLayerObject* Layer = GetLayer();
+	check(Layer);
+
+	UDMMaterialSlot* Slot = Layer->GetSlot();
+	check(Slot);
+
+	UDynamicMaterialModelEditorOnlyData* ModelEditorOnlyData = Slot->GetMaterialModelEditorOnlyData();
+	check(ModelEditorOnlyData);
+
+	UDynamicMaterialModel* MaterialModel = ModelEditorOnlyData->GetMaterialModel();
+	check(MaterialModel);
+
+	InputTextureUV->Init(MaterialModel);
+
+	AddInput(InputTextureUV);
+
+	UpdateInputMap(
+		InInputIdx,
+		FDMMaterialStageConnectorChannel::FIRST_STAGE_INPUT + Inputs.Num() - 1,
+		InInputChannel,
+		0, // Texture UVs only have 1 output.
+		OutputChannel,
+		EDMMaterialPropertyType::None
+	);
+
+	return InputTextureUV;
+}
+
+UDMMaterialStageInputGradient* UDMMaterialStage::ChangeInput_Gradient(int32 InInputIdx, int32 InInputChannel, TSubclassOf<UDMMaterialStageGradient> InGradientClass, int32 InOutputChannel)
+{
+	check(Source);
+
+	UDMMaterialStageThroughput* Throughput = Cast<UDMMaterialStageThroughput>(Source);
+	check(Throughput);
+
+	const TArray<FDMMaterialStageConnector>& InputConnectors = Throughput->GetInputConnectors();
+	check(InputConnectors.IsValidIndex(InInputIdx));
+
+	check(InGradientClass.Get());
+	check(InGradientClass.Get() != UDMMaterialStageGradient::StaticClass());
+
+	UDMMaterialStageGradient* GradientCDO = Cast<UDMMaterialStageGradient>(InGradientClass->GetDefaultObject(true));
+	check(GradientCDO);
+
+	const TArray<FDMMaterialStageConnector>& OutputConnectors = GradientCDO->GetOutputConnectors();
+	check(Throughput->CanInputAcceptType(InInputIdx, EDMValueType::VT_Float1));
+
+	UDMMaterialStageInputGradient* InputGradient = NewObject<UDMMaterialStageInputGradient>(this, NAME_None, RF_Transactional);
+	check(InputGradient);
+
+	{
+		const FDMUpdateGuard Guard;
+		InputGradient->SetMaterialStageGradientClass(InGradientClass);
+	}
+
+	AddInput(InputGradient);
+
+	UpdateInputMap(
+		InInputIdx,
+		FDMMaterialStageConnectorChannel::FIRST_STAGE_INPUT + Inputs.Num() - 1,
+		InInputChannel,
+		0,
+		InOutputChannel,
+		EDMMaterialPropertyType::None
+	);
+
+	return InputGradient;
+}
+
+UDMMaterialStageInputFunction* UDMMaterialStage::ChangeInput_MaterialFunction(int32 InInputIdx, int32 InInputChannel)
+{
+	check(Source);
+
+	UDMMaterialStageThroughput* Throughput = Cast<UDMMaterialStageThroughput>(Source);
+	check(Throughput);
+
+	const TArray<FDMMaterialStageConnector>& InputConnectors = Throughput->GetInputConnectors();
+	check(InputConnectors.IsValidIndex(InInputIdx));
+
+	UDMMaterialStageInputFunction* InputMaterialFunction = NewObject<UDMMaterialStageInputFunction>(this, NAME_None, RF_Transactional);
+	check(InputMaterialFunction);
+
+	{
+		const FDMUpdateGuard Guard;
+		InputMaterialFunction->Init();
+	}
+
+	AddInput(InputMaterialFunction);
+
+	UpdateInputMap(
+		InInputIdx,
+		FDMMaterialStageConnectorChannel::FIRST_STAGE_INPUT + Inputs.Num() - 1,
+		InInputChannel,
+		0,
+		FDMMaterialStageConnectorChannel::WHOLE_CHANNEL,
+		EDMMaterialPropertyType::None
+	);
+
+	return InputMaterialFunction;
+}
+
+void UDMMaterialStage::UpdateInputMap(int32 InInputIdx, int32 InSourceIndex, int32 InInputChannel, int32 InOutputIdx, int32 InOutputChannel, EDMMaterialPropertyType InStageProperty)
+{
+	if (!IsComponentValid())
+	{
+		return;
+	}
+
+	check(InSourceIndex != FDMMaterialStageConnectorChannel::PREVIOUS_STAGE || InStageProperty != EDMMaterialPropertyType::None);
+	check(InSourceIndex >= 0 && InSourceIndex <= Inputs.Num());
+
+	VerifyAllInputMaps();
+	check(InputConnectionMap.IsValidIndex(InInputIdx));
+
+	check(Source);
+
+	UDMMaterialStageThroughput* Throughput = Cast<UDMMaterialStageThroughput>(Source);
+	check(Throughput);
+
+	const TArray<FDMMaterialStageConnector>& InputConnectors = Throughput->GetInputConnectors();
+	check(InputConnectors.IsValidIndex(InInputIdx));
+
+	// Check the validity of the incoming data. It must be valid at set time.
+	// There are no guarantees it'll be valid at a later time.
+	if (InSourceIndex == FDMMaterialStageConnectorChannel::PREVIOUS_STAGE)
+	{
+		UDMMaterialLayerObject* Layer = GetLayer();
+		check(Layer);
+
+		const UDMMaterialLayerObject* PreviousLayer = Layer->GetPreviousLayer(InStageProperty, EDMMaterialLayerStage::Base);
+		UDMMaterialStageSource* PreviousSource = nullptr;
+
+		if (PreviousLayer)
+		{
+			if (UDMMaterialStage* Stage = PreviousLayer->GetLastEnabledStage(EDMMaterialLayerStage::All))
+			{
+				PreviousSource = Stage->GetSource();
+			}
+		}
+
+		if (PreviousSource)
+		{
+			const TArray<FDMMaterialStageConnector>& PreviousStageOutputs = PreviousSource->GetOutputConnectors();
+			check(PreviousStageOutputs.IsValidIndex(InOutputIdx));
+			check(Throughput->CanInputConnectTo(InInputIdx, PreviousStageOutputs[InOutputIdx], InOutputChannel));
+		}
+		else
+		{
+			check(InOutputIdx == 0);
+			check(Throughput->GetInputConnectors().IsEmpty() == false);
+			check(Throughput->GetInputConnectors()[0].IsCompatibleWith(EDMValueType::VT_Float3_RGB));
+		}
+	}
+	else
+	{
+		const int32 StageInputIdx = InSourceIndex - FDMMaterialStageConnectorChannel::FIRST_STAGE_INPUT;
+		check(Inputs.IsValidIndex(StageInputIdx));
+
+		const TArray<FDMMaterialStageConnector>& StageInputOutputConnectors = Inputs[StageInputIdx]->GetOutputConnectors();
+		check(StageInputOutputConnectors.IsValidIndex(InOutputIdx));
+		check(Throughput->CanInputConnectTo(InInputIdx, StageInputOutputConnectors[InOutputIdx], InOutputChannel));
+	}
+
+	// Remove all mapping for this input and replace them with a whole channel mapping
+	if (InInputChannel == FDMMaterialStageConnectorChannel::WHOLE_CHANNEL)
+	{
+		InputConnectionMap[InInputIdx].Channels.Empty();
+		InputConnectionMap[InInputIdx].Channels.Add({InSourceIndex, InStageProperty, InOutputIdx, InOutputChannel});
+	}
+	// Add/replace the new channel-specific mapping
+	else
+	{
+		for (int32 NewIndex = InputConnectionMap.Num(); NewIndex < InInputChannel; ++NewIndex)
+		{
+			InputConnectionMap[InInputIdx].Channels.Add({FDMMaterialStageConnectorChannel::NO_SOURCE, EDMMaterialPropertyType::None, 0, FDMMaterialStageConnectorChannel::WHOLE_CHANNEL});
+		}
+
+		InputConnectionMap[InInputIdx].Channels[InInputChannel].SourceIndex = InSourceIndex;
+		InputConnectionMap[InInputIdx].Channels[InInputChannel].MaterialProperty = InStageProperty;
+		InputConnectionMap[InInputIdx].Channels[InInputChannel].OutputIndex = InOutputIdx;
+		InputConnectionMap[InInputIdx].Channels[InInputChannel].OutputChannel = InOutputChannel;
+	}
+
+	RemoveUnusedInputs();
+
+	if (FDMUpdateGuard::CanUpdate())
+	{
+		Source->Update(EDMUpdateType::Structure);
+	}
+}
+
+int32 UDMMaterialStage::FindIndex() const
+{
+	UDMMaterialLayerObject* Layer = GetLayer();
+
+	if (!Layer)
+	{
+		return INDEX_NONE;
+	}
+
+	const TArray<TObjectPtr<UDMMaterialStage>>& Stages = Layer->GetAllStages();
+
+	for (int32 StageIndex = 0; StageIndex < Stages.Num(); ++StageIndex)
+	{
+		if (Stages[StageIndex] == this)
+		{
+			return StageIndex;
+		}
+	}
+
+	return INDEX_NONE;
+}
+
+UDMMaterialStage* UDMMaterialStage::GetPreviousStage() const
+{
+	UDMMaterialLayerObject* Layer = GetLayer();
+
+	if (!Layer)
+	{
+		return nullptr;
+	}
+
+	TArray<UDMMaterialStage*> Stages = Layer->GetStages(EDMMaterialLayerStage::All);
+	int32 StageIndex = Stages.IndexOfByKey(this);
+
+	if (StageIndex == INDEX_NONE)
+	{
+		return nullptr;
+	}
+
+	--StageIndex;
+
+	if (Stages.IsValidIndex(StageIndex))
+	{
+		return Stages[StageIndex];
+	}
+
+	return nullptr;
+}
+
+UDMMaterialStage* UDMMaterialStage::GetNextStage() const
+{
+	UDMMaterialLayerObject* Layer = GetLayer();
+
+	if (!Layer)
+	{
+		return nullptr;
+	}
+
+	TArray<UDMMaterialStage*> Stages = Layer->GetStages(EDMMaterialLayerStage::All);
+	int32 StageIndex = Stages.IndexOfByKey(this);
+
+	if (StageIndex == INDEX_NONE)
+	{
+		return nullptr;
+	}
+
+	++StageIndex;
+
+	if (Stages.IsValidIndex(StageIndex))
+	{
+		return Stages[StageIndex];
+	}
+
+	return nullptr;
+}
+
+bool UDMMaterialStage::IsRootStage() const
+{
+	return true;
+}
+
+bool UDMMaterialStage::VerifyAllInputMaps()
+{
+	if (!IsComponentValid())
+	{
+		return false;
+	}
+
+	bool bVerified = true;
+
+	UDMMaterialStageThroughput* Throughput = Cast<UDMMaterialStageThroughput>(Source);
+
+	if (!Throughput)
+	{
+		if (InputConnectionMap.Num() > 1)
+		{
+			InputConnectionMap.Empty();
+			bVerified = false;
+		}
+
+		Inputs.Empty(); // If we have no inputs connectors, we don't need any inputs...
+
+		return bVerified;
+	}
+
+	const TArray<FDMMaterialStageConnector>& ExpressionInputs = Throughput->GetInputConnectors();
+
+	if (ExpressionInputs.IsEmpty())
+	{
+		if (InputConnectionMap.Num() > 1)
+		{
+			InputConnectionMap.Empty();
+			bVerified = false;
+		}
+
+		Inputs.Empty(); // If we have no inputs connectors, we don't need any inputs...
+
+		return bVerified;
+	}
+
+	if (InputConnectionMap.Num() != ExpressionInputs.Num())
+	{
+		InputConnectionMap.SetNum(ExpressionInputs.Num());
+		bVerified = false;
+	}
+
+	for (int32 InputIdx = 0; InputIdx < InputConnectionMap.Num(); ++InputIdx)
+	{
+		bVerified = bVerified && VerifyInputMap(InputIdx);
+	}
+
+	return bVerified;
+}
+
+bool UDMMaterialStage::VerifyInputMap(int32 InInputIdx)
+{
+	if (!IsComponentValid())
+	{
+		return false;
+	}
+
+	bool bVerified = true;
+
+	UDMMaterialStageThroughput* Throughput = Cast<UDMMaterialStageThroughput>(Source);
+
+	if (!Throughput)
+	{
+		if (!InputConnectionMap.IsEmpty())
+		{
+			InputConnectionMap.Empty();
+			bVerified = false;
+		}
+
+		Inputs.Empty(); // If we have no inputs connectors, we don't need any inputs...
+
+		return bVerified;
+	}
+
+	const TArray<FDMMaterialStageConnector>& ExpressionInputs = Throughput->GetInputConnectors();
+
+	if (ExpressionInputs.IsEmpty())
+	{
+		if (!InputConnectionMap.IsEmpty())
+		{
+			InputConnectionMap.Empty();
+			bVerified = false;
+		}
+
+		Inputs.Empty(); // If we have no inputs connectors, we don't need any inputs...
+
+		return bVerified;
+	}
+
+	if (!InputConnectionMap.IsValidIndex(InInputIdx))
+	{
+		return false;
+	}
+
+	if (InputConnectionMap[InInputIdx].Channels.IsEmpty())
+	{
+		return true;
+	}
+
+	UDMMaterialLayerObject* Layer = GetLayer();
+	check(Layer);
+
+	EDMMaterialPropertyType StageProperty = Layer->GetMaterialProperty();
+	check(StageProperty != EDMMaterialPropertyType::None);
+
+	for (FDMMaterialStageConnectorChannel& Channel : InputConnectionMap[InInputIdx].Channels)
+	{
+		bool bValidConnectionMap = false;
+
+		if (Channel.SourceIndex == FDMMaterialStageConnectorChannel::NO_SOURCE)
+		{
+			continue;
+		}
+
+		// Check previous stage
+		if (Channel.SourceIndex == FDMMaterialStageConnectorChannel::PREVIOUS_STAGE)
+		{
+			if (const UDMMaterialLayerObject* PreviousLayerAndStage = Layer->GetPreviousLayer(StageProperty, EDMMaterialLayerStage::Base))
+			{
+				check(PreviousLayerAndStage->GetStage(EDMMaterialLayerStage::Mask));
+
+				if (UDMMaterialStageSource* PreviousStageSource = PreviousLayerAndStage->GetStage(EDMMaterialLayerStage::Mask)->GetSource())
+				{
+					const TArray<FDMMaterialStageConnector>& PreviousStageOutputConnectors = PreviousStageSource->GetOutputConnectors();
+
+					if (PreviousStageOutputConnectors.IsValidIndex(Channel.OutputIndex)
+						&& Throughput->CanInputConnectTo(InInputIdx, PreviousStageOutputConnectors[Channel.OutputIndex], Channel.OutputChannel))
+					{
+						// This is valid
+						bValidConnectionMap = true;
+					}
+				}
+			}
+		}
+		// Check inputs
+		else
+		{
+			const int32 StageInputIdx = Channel.SourceIndex - FDMMaterialStageConnectorChannel::FIRST_STAGE_INPUT;
+
+			if (Inputs.IsValidIndex(StageInputIdx))
+			{
+				const TArray<FDMMaterialStageConnector>& StageInputOutputConnectors = Inputs[StageInputIdx]->GetOutputConnectors();
+
+				if (StageInputOutputConnectors.IsValidIndex(Channel.OutputIndex)
+					&& Throughput->CanInputConnectTo(InInputIdx, StageInputOutputConnectors[Channel.OutputIndex], Channel.OutputChannel))
+				{
+					bValidConnectionMap = true;
+				}
+			}
+		}
+
+		// Invalid source channel
+		if (!bValidConnectionMap)
+		{
+			Channel.SourceIndex = FDMMaterialStageConnectorChannel::NO_SOURCE;
+			bVerified = false;
+		}
+	}
+
+	bool bFoundSource = false;
+
+	for (const FDMMaterialStageConnectorChannel& Channel : InputConnectionMap[InInputIdx].Channels)
+	{
+		if (Channel.SourceIndex != FDMMaterialStageConnectorChannel::NO_SOURCE)
+		{
+			bFoundSource = true;
+			break;
+		}
+	}
+
+	// This is just clean up and does not affect verification.
+	if (!bFoundSource)
+	{
+		InputConnectionMap[InInputIdx].Channels.Empty();
+	}
+
+	return bVerified;
+}
+
+UDMMaterialStageInputValue* UDMMaterialStage::ChangeSource_NewLocalValue(EDMValueType InType)
+{
+	if (!bCanChangeSource)
+	{
+		return nullptr;
+	}
+
+	UDMMaterialLayerObject* Layer = GetLayer();
+	check(Layer);
+
+	UDMMaterialSlot* Slot = Layer->GetSlot();
+	check(Slot);
+
+	UDynamicMaterialModelEditorOnlyData* ModelEditorOnlyData = Slot->GetMaterialModelEditorOnlyData();
+	check(ModelEditorOnlyData);
+
+	UDynamicMaterialModel* MaterialModel = ModelEditorOnlyData->GetMaterialModel();
+	check(MaterialModel);
+
+	// Don't add via the builder, just parent it to the builder.
+	// It won't appear in the global value list.
+	UDMMaterialValue* NewValue = UDMMaterialValue::CreateMaterialValue(MaterialModel, TEXT(""), InType, true);
+	check(NewValue);
+
+	UDMMaterialStageInputValue* InputValue = NewObject<UDMMaterialStageInputValue>(this, NAME_None, RF_Transactional);
+	check(InputValue);
+
+	{
+		const FDMUpdateGuard Guard;
+		InputValue->SetValue(NewValue);
+	}
+
+	SetSource(InputValue);
+
+	return InputValue;
+}
+
+UDMMaterialStageInputValue* UDMMaterialStage::ChangeSource_Value(UDMMaterialValue* InValue)
+{
+	if (!bCanChangeSource)
+	{
+		return nullptr;
+	}
+
+	UDMMaterialLayerObject* Layer = GetLayer();
+	check(Layer);
+
+	UDMMaterialSlot* Slot = Layer->GetSlot();
+	check(Slot);
+
+	UDynamicMaterialModelEditorOnlyData* ModelEditorOnlyData = Slot->GetMaterialModelEditorOnlyData();
+	check(ModelEditorOnlyData);
+
+	UDynamicMaterialModel* MaterialModel = ModelEditorOnlyData->GetMaterialModel();
+	check(MaterialModel);
+
+	UDMMaterialStageInputValue* InputValue = NewObject<UDMMaterialStageInputValue>(this, NAME_None, RF_Transactional);
+	check(InputValue);
+
+	{
+		const FDMUpdateGuard Guard;
+
+		check(InputValue);
+		InputValue->SetValue(InValue);
+	}
+
+	SetSource(InputValue);
+
+	return InputValue;
+}
+
+UDMMaterialStageInputValue* UDMMaterialStage::ChangeSource_NewValue(EDMValueType InType)
+{
+	if (!bCanChangeSource)
+	{
+		return nullptr;
+	}
+
+	UDMMaterialLayerObject* Layer = GetLayer();
+	check(Layer);
+
+	UDMMaterialSlot* Slot = Layer->GetSlot();
+	check(Slot);
+
+	UDynamicMaterialModelEditorOnlyData* ModelEditorOnlyData = Slot->GetMaterialModelEditorOnlyData();
+	check(ModelEditorOnlyData);
+
+	UDynamicMaterialModel* MaterialModel = ModelEditorOnlyData->GetMaterialModel();
+	check(MaterialModel);
+
+	UDMMaterialValue* NewValue = MaterialModel->AddValue(InType);
+	check(NewValue);
+
+	UDMMaterialStageInputValue* InputValue = NewObject<UDMMaterialStageInputValue>(this, NAME_None, RF_Transactional);
+	check(InputValue);
+
+	{
+		const FDMUpdateGuard Guard;
+		InputValue->SetValue(NewValue);
+	}
+
+	SetSource(InputValue);
+
+	return InputValue;
+}
+
+UDMMaterialStageInputSlot* UDMMaterialStage::ChangeSource_Slot(UDMMaterialSlot* InSlot, EDMMaterialPropertyType InProperty)
+{
+	if (!bCanChangeSource)
+	{
+		return nullptr;
+	}
+
+	UDMMaterialLayerObject* Layer = GetLayer();
+	check(Layer);
+
+	UDMMaterialSlot* Slot = Layer->GetSlot();
+	check(Slot);
+
+	check(InSlot);
+	check(Slot != InSlot);
+	check(InSlot->GetMaterialModelEditorOnlyData() == Slot->GetMaterialModelEditorOnlyData());
+
+	UDynamicMaterialModelEditorOnlyData* ModelEditorOnlyData = Slot->GetMaterialModelEditorOnlyData();
+	check(ModelEditorOnlyData);
+
+	TArray<EDMMaterialPropertyType> SlotPropertes = ModelEditorOnlyData->GetMaterialPropertiesForSlot(InSlot);
+	check(SlotPropertes.Contains(InProperty));
+
+	UDMMaterialStageInputSlot* InputSlot = NewObject<UDMMaterialStageInputSlot>(this, NAME_None, RF_Transactional);
+	check(InputSlot);
+
+	{
+		const FDMUpdateGuard Guard;
+		InputSlot->SetSlot(InSlot);
+		InputSlot->SetMaterialProperty(InProperty);
+	}
+
+	SetSource(InputSlot);
+
+	return InputSlot;
+}
+
+UDMMaterialStageExpression* UDMMaterialStage::ChangeSource_Expression(TSubclassOf<UDMMaterialStageExpression> InExpressionClass)
+{
+	if (!bCanChangeSource)
+	{
+		return nullptr;
+	}
+
+	check(InExpressionClass);
+	check(InExpressionClass != UDMMaterialStageExpression::StaticClass());
+
+	UDMMaterialStageExpression* NewExpression = NewObject<UDMMaterialStageExpression>(this, InExpressionClass, NAME_None, RF_Transactional);
+	check(NewExpression);
+
+	SetSource(NewExpression);
+
+	return NewExpression;
+}
+
+UDMMaterialStageBlend* UDMMaterialStage::ChangeSource_Blend(TSubclassOf<UDMMaterialStageBlend> InBlendClass)
+{
+	if (!bCanChangeSource)
+	{
+		return nullptr;
+	}
+
+	check(InBlendClass);
+	check(InBlendClass != UDMMaterialStageBlend::StaticClass());
+
+	UDMMaterialStageBlend* NewBlend = NewObject<UDMMaterialStageBlend>(this, InBlendClass, NAME_None, RF_Transactional);
+	check(NewBlend);
+
+	SetSource(NewBlend);
+	
+	return NewBlend;
+}
+
+UDMMaterialStageGradient* UDMMaterialStage::ChangeSource_Gradient(TSubclassOf<UDMMaterialStageGradient> InGradientClass)
+{
+	if (!bCanChangeSource)
+	{
+		return nullptr;
+	}
+
+	check(InGradientClass);
+	check(InGradientClass != UDMMaterialStageGradient::StaticClass());
+
+	UDMMaterialStageGradient* NewGradient = NewObject<UDMMaterialStageGradient>(this, InGradientClass, NAME_None, RF_Transactional);
+	check(NewGradient);
+
+	SetSource(NewGradient);
+
+	return NewGradient;
+}
+
+UDMMaterialStageInputTextureUV* UDMMaterialStage::ChangeSource_UV(bool bInDoUpdate)
+{
+	if (!bCanChangeSource)
+	{
+		return nullptr;
+	}
+
+	UDMMaterialStageInputTextureUV* InputTextureUV = NewObject<UDMMaterialStageInputTextureUV>(this, NAME_None, RF_Transactional);
+	check(InputTextureUV);
+
+	UDMMaterialLayerObject* Layer = GetLayer();
+	check(Layer);
+
+	UDMMaterialSlot* Slot = Layer->GetSlot();
+	check(Slot);
+
+	UDynamicMaterialModelEditorOnlyData* ModelEditorOnlyData = Slot->GetMaterialModelEditorOnlyData();
+	check(ModelEditorOnlyData);
+
+	UDynamicMaterialModel* MaterialModel = ModelEditorOnlyData->GetMaterialModel();
+	check(MaterialModel);
+
+	InputTextureUV->Init(MaterialModel);
+
+	SetSource(InputTextureUV);
+
+	return InputTextureUV;
+}
+
+UDMMaterialStageFunction* UDMMaterialStage::ChangeSource_MaterialFunction()
+{
+	if (!bCanChangeSource)
+	{
+		return nullptr;
+	}
+
+	UDMMaterialStageFunction* NewMaterialFunction = NewObject<UDMMaterialStageFunction>(this, UDMMaterialStageFunction::StaticClass(), NAME_None, RF_Transactional);
+	check(NewMaterialFunction);
+
+	SetSource(NewMaterialFunction);
+
+	return NewMaterialFunction;
+}
+
+bool UDMMaterialStage::IsCompatibleWithPreviousStage(const UDMMaterialStage* InPreviousStage) const
+{
+	/*
+	 * It is now up to the user to sort this particular problem out because it would do more harm than good
+	 * to force correctness in "transition states" while the user is changing settings.
+	 */
+	return true;
+}
+
+bool UDMMaterialStage::IsCompatibleWithNextStage(const UDMMaterialStage* InNextStage) const
+{
+	if (!InNextStage)
+	{
+		return true;
+	}
+
+	return InNextStage->IsCompatibleWithPreviousStage(this);
+}
+
+void UDMMaterialStage::AddInput(UDMMaterialStageInput* InNewInput)
+{
+	if (!IsComponentValid())
+	{
+		return;
+	}
+
+	check(InNewInput);
+	check(InNewInput->GetStage() == this);
+
+	Inputs.Add(InNewInput);
+	InputConnectionMap.Add(FDMMaterialStageConnection());
+
+	if (HasComponentBeenAdded())
+	{
+		if (GUndo)
+		{
+			InNewInput->Modify();
+		}
+
+		InNewInput->SetComponentState(EDMComponentLifetimeState::Added);
+	}
+}
+
+void UDMMaterialStage::RemoveInput(UDMMaterialStageInput* InInput)
+{
+	if (!IsComponentValid())
+	{
+		return;
+	}
+
+	check(InInput);
+	check(InInput->GetStage() == this);
+
+	bool bFound = false;
+
+	for (int32 InputIdx = 0; InputIdx < Inputs.Num(); ++InputIdx)
+	{
+		if (Inputs[InputIdx] == InInput)
+		{
+			Inputs.RemoveAt(InputIdx);
+			InputConnectionMap.RemoveAt(InputIdx);
+			bFound = true;
+			break;
+		}
+	}
+
+	check(bFound);
+
+	if (GUndo)
+	{
+		InInput->Modify();
+	}
+
+	InInput->SetComponentState(EDMComponentLifetimeState::Removed);
+
+	Update(EDMUpdateType::Structure);
+}
+
+UMaterialInterface* UDMMaterialStage::GetPreviewMaterial()
+{
+	if (!PreviewMaterialBase)
+	{
+		CreatePreviewMaterial();
+
+		if (PreviewMaterialBase)
+		{
+			MarkComponentDirty();
+		}
+	}
+
+	return PreviewMaterialDynamic.Get();
+}
+
+const FDMMaterialStageConnectorChannel* UDMMaterialStage::FindInputChannel(UDMMaterialStageInput* InStageInput)
+{
+	check(InStageInput);
+
+	int32 InputIdx = Inputs.Find(InStageInput);
+	
+	if (InputIdx == INDEX_NONE)
+	{
+		return nullptr;
+	}
+
+	InputIdx += FDMMaterialStageConnectorChannel::FIRST_STAGE_INPUT;
+
+	for (const FDMMaterialStageConnection& Connection : InputConnectionMap)
+	{
+		for (const FDMMaterialStageConnectorChannel& Channel : Connection.Channels)
+		{
+			if (Channel.SourceIndex == InputIdx)
+			{
+				return &Channel;
+			}
+		}
+	}
+
+	return nullptr;
+}
+
+void UDMMaterialStage::DoClean()
+{
+	if (IsComponentValid() && GetLayer())
+	{
+		UpdatePreviewMaterial();
+	}
+
+	Super::DoClean();
+}
+
+void UDMMaterialStage::PostEditUndo()
+{
+	Super::PostEditUndo();
+
+	if (!IsComponentValid())
+	{
+		return;
+	}
+
+	UDMMaterialLayerObject* Layer = GetLayer();
+
+	if (!Layer)
+	{
+		SetComponentState(EDMComponentLifetimeState::Removed);
+		return;
+	}
+
+	MarkComponentDirty();
+
+	Update(EDMUpdateType::Structure);
+}
+
+void UDMMaterialStage::PostLoad()
+{
+	Super::PostLoad();
+
+	AddDelegates();
+}
+
+void UDMMaterialStage::PostEditImport()
+{
+	Super::PostEditImport();
+
+	AddDelegates();
+}
+
+void UDMMaterialStage::CreatePreviewMaterial()
+{
+	if (!IsComponentValid())
+	{
+		return;
+	}
+
+	if (FDynamicMaterialModule::IsMaterialExportEnabled() == false)
+	{
+		UMaterialFactoryNew* MaterialFactory = NewObject<UMaterialFactoryNew>();
+		check(MaterialFactory);
+
+		PreviewMaterialBase = Cast<UMaterial>(MaterialFactory->FactoryCreateNew(
+			UMaterial::StaticClass(),
+			GetTransientPackage(),
+			NAME_None,
+			RF_Transient,
+			nullptr,
+			GWarn
+		));
+
+		PreviewMaterialBase->bIsPreviewMaterial = true;
+	}
+	else
+	{
+		FString MaterialBaseName = GetName() + "-" + FGuid::NewGuid().ToString();
+		const FString FullName = "/Game/DynamicMaterials/" + MaterialBaseName;
+		UPackage* Package = CreatePackage(*FullName);
+
+		UMaterialFactoryNew* MaterialFactory = NewObject<UMaterialFactoryNew>();
+		check(MaterialFactory);
+
+		PreviewMaterialBase = Cast<UMaterial>(MaterialFactory->FactoryCreateNew(
+			UMaterial::StaticClass(),
+			Package,
+			*MaterialBaseName,
+			RF_Standalone | RF_Public,
+			nullptr,
+			GWarn
+		));
+
+		FAssetRegistryModule::AssetCreated(PreviewMaterialBase);
+		Package->FullyLoad();
+	}
+}
+
+void UDMMaterialStage::UpdatePreviewMaterial()
+{
+	if (!IsComponentValid())
+	{
+		return;
+	}
+
+	if (!Source)
+	{
+		return;
+	}
+
+	UE_LOG(LogDynamicMaterialEditor, Display, TEXT("Building Material Designer Stage Preview (%s)..."), *GetName());
+
+	if (!PreviewMaterialBase)
+	{
+		CreatePreviewMaterial();
+
+		if (!PreviewMaterialBase)
+		{
+			return;
+		}
+	}
+
+	UDMMaterialLayerObject* Layer = GetLayer();
+	check(Layer);
+
+	UDMMaterialSlot* Slot = Layer->GetSlot();
+	check(Slot);
+
+	UDynamicMaterialModelEditorOnlyData* ModelEditorOnlyData = Slot->GetMaterialModelEditorOnlyData();
+	check(ModelEditorOnlyData);
+
+	UDynamicMaterialModel* MaterialModel = ModelEditorOnlyData->GetMaterialModel();
+	check(MaterialModel);
+
+	PreviewMaterialBase->GetEditorOnlyData()->EmissiveColor.Expression = nullptr;
+	PreviewMaterialBase->GetEditorOnlyData()->EmissiveColor.OutputIndex = 0;
+
+	const bool bGenerateSuccess = Source->UpdateStagePreviewMaterial(
+		this,
+		PreviewMaterialBase,
+		PreviewMaterialBase->GetEditorOnlyData()->EmissiveColor.Expression,
+		PreviewMaterialBase->GetEditorOnlyData()->EmissiveColor.OutputIndex
+	);
+
+	if (!bGenerateSuccess)
+	{
+		return;
+	}
+
+	PreviewMaterialDynamic = UMaterialInstanceDynamic::Create(PreviewMaterialBase.Get(), GetTransientPackage());
+
+	const bool bUseTexureUVs = UDynamicMaterialEditorSettings::Get()->bPreviewImagesUseTextureUVs;
+
+	TArray<UObject*> Subobjects;
+	GetObjectsWithOuter(MaterialModel, Subobjects, false);
+
+	for (UObject* Subobject : Subobjects)
+	{
+		if (UDMMaterialValue* Value = Cast<UDMMaterialValue>(Subobject))
+		{
+			OnValueUpdated(MaterialModel, Value);
+		}
+		else if (bUseTexureUVs)
+		{
+			if (UDMTextureUV* TextureUV = Cast<UDMTextureUV>(Subobject))
+			{
+				OnTextureUVUpdated(MaterialModel, TextureUV);
+			}
+		}
+	}
+}
+
+#undef LOCTEXT_NAMESPACE
diff --git a/Engine/Plugins/Experimental/DynamicMaterial/Source/DynamicMaterialEditor/Private/Components/DMMaterialStageSource.cpp b/Engine/Plugins/Experimental/DynamicMaterial/Source/DynamicMaterialEditor/Private/Components/DMMaterialStageSource.cpp
new file mode 100644
--- /dev/null
+++ b/Engine/Plugins/Experimental/DynamicMaterial/Source/DynamicMaterialEditor/Private/Components/DMMaterialStageSource.cpp
@@ -0,0 +1,382 @@
+// Copyright Epic Games, Inc. All Rights Reserved.
+
+#include "Components/DMMaterialStageSource.h"
+#include "AssetRegistry/AssetRegistryModule.h"
+#include "Components/DMMaterialLayer.h"
+#include "Components/DMMaterialProperty.h"
+#include "Components/DMMaterialSlot.h"
+#include "Components/DMMaterialStage.h"
+#include "DynamicMaterialEditorModule.h"
+#include "DynamicMaterialModule.h"
+#include "Factories/MaterialFactoryNew.h"
+#include "Materials/Material.h"
+#include "Materials/MaterialExpression.h"
+#include "MaterialValueType.h"
+#include "Model/DMMaterialBuildState.h"
+#include "Model/DMMaterialBuildUtils.h"
+#include "Model/DynamicMaterialModel.h"
+#include "Model/DynamicMaterialModelEditorOnlyData.h"
+#include "UObject/Package.h"
+#include "UObject/UObjectIterator.h"
+
+TArray<TStrongObjectPtr<UClass>> UDMMaterialStageSource::SourceClasses = TArray<TStrongObjectPtr<UClass>>();
+
+UDMMaterialStageSource::UDMMaterialStageSource()
+{
+	PreviewMaterial = nullptr;
+}
+
+UDMMaterialStage* UDMMaterialStageSource::GetStage() const
+{
+	return Cast<UDMMaterialStage>(GetOuterSafe());
+}
+
+void UDMMaterialStageSource::Update(EDMUpdateType InUpdateType)
+{
+	if (!IsComponentValid())
+	{
+		return;
+	}
+
+	if (HasComponentBeenRemoved())
+	{
+		return;
+	}
+
+	if (InUpdateType == EDMUpdateType::Structure)
+	{
+		MarkComponentDirty();
+	}
+
+	UDMMaterialStage* Stage = GetStage();
+	check(Stage);
+
+	Stage->Update(InUpdateType);
+
+	Super::Update(InUpdateType);
+}
+
+void UDMMaterialStageSource::OnComponentAdded()
+{
+	if (!IsComponentValid())
+	{
+		return;
+	}
+
+	Super::OnComponentAdded();
+
+	if (FDMUpdateGuard::CanUpdate())
+	{
+		Update(EDMUpdateType::Structure);
+	}
+}
+
+void UDMMaterialStageSource::GetMaskAlphaBlendNode(const TSharedRef<FDMMaterialBuildState>& InBuildState, UMaterialExpression*& OutExpression, int32& OutOutputIndex, int32& OutOutputChannel) const
+{
+	OutExpression = nullptr;
+}
+
+const TArray<TStrongObjectPtr<UClass>>& UDMMaterialStageSource::GetAvailableSourceClasses()
+{
+	if (SourceClasses.IsEmpty())
+	{
+		GenerateClassList();
+	}
+
+	return SourceClasses;
+}
+
+void UDMMaterialStageSource::GenerateClassList()
+{
+	SourceClasses.Empty();
+
+	for (TObjectIterator<UClass> ClassIT; ClassIT; ++ClassIT)
+	{
+		TSubclassOf<UDMMaterialStageSource> MSEClass = *ClassIT;
+
+		if (!MSEClass.Get())
+		{
+			continue;
+		}
+
+		if (MSEClass->ClassFlags & (CLASS_Abstract | CLASS_Hidden | CLASS_Deprecated | CLASS_NewerVersionExists))
+		{
+			continue;
+		}
+
+		SourceClasses.Add(TStrongObjectPtr<UClass>(MSEClass));
+	}
+}
+
+UMaterial* UDMMaterialStageSource::GetPreviewMaterial()
+{
+	if (!PreviewMaterial)
+	{
+		CreatePreviewMaterial();
+
+		if (PreviewMaterial)
+		{
+			MarkComponentDirty();
+		}
+	}
+
+	return PreviewMaterial;
+}
+
+void UDMMaterialStageSource::UpdateOntoPreviewMaterial(UMaterial* ExternalPreviewMaterial)
+{
+	if (!IsComponentValid())
+	{
+		return;
+	}
+
+	if (!ExternalPreviewMaterial)
+	{
+		return;
+	}
+
+	UpdatePreviewMaterial(ExternalPreviewMaterial);
+}
+
+void UDMMaterialStageSource::CreatePreviewMaterial()
+{
+	if (!IsComponentValid())
+	{
+		return;
+	}
+
+	if (FDynamicMaterialModule::IsMaterialExportEnabled() == false)
+	{
+		UMaterialFactoryNew* MaterialFactory = NewObject<UMaterialFactoryNew>();
+		check(MaterialFactory);
+
+		PreviewMaterial = Cast<UMaterial>(MaterialFactory->FactoryCreateNew(
+			UMaterial::StaticClass(),
+			GetTransientPackage(),
+			NAME_None,
+			RF_Transient,
+			nullptr,
+			GWarn
+		));
+
+		PreviewMaterial->bIsPreviewMaterial = true;
+	}
+	else
+	{
+		FString MaterialBaseName = GetName() + "-" + FGuid::NewGuid().ToString();
+		const FString FullName = "/Game/DynamicMaterials/" + MaterialBaseName;
+		UPackage* Package = CreatePackage(*FullName);
+
+		UMaterialFactoryNew* MaterialFactory = NewObject<UMaterialFactoryNew>();
+		check(MaterialFactory);
+
+		PreviewMaterial = Cast<UMaterial>(MaterialFactory->FactoryCreateNew(
+			UMaterial::StaticClass(),
+			Package,
+			*MaterialBaseName,
+			RF_Standalone | RF_Public,
+			nullptr,
+			GWarn
+		));
+
+		FAssetRegistryModule::AssetCreated(PreviewMaterial);
+		Package->FullyLoad();
+	}
+}
+
+void UDMMaterialStageSource::UpdatePreviewMaterial(UMaterial* InPreviewMaterial)
+{
+	if (!IsComponentValid())
+	{
+		return;
+	}
+
+	if (!InPreviewMaterial)
+	{
+		if (!PreviewMaterial)
+		{
+			CreatePreviewMaterial();
+		}
+
+		InPreviewMaterial = PreviewMaterial;
+
+		if (!PreviewMaterial)
+		{
+			return;
+		}
+	}
+
+	UE_LOG(LogDynamicMaterialEditor, Display, TEXT("Building Material Designer Source Preview (%s)..."), *GetName());
+
+	UDMMaterialStage* Stage = GetStage();
+	check(Stage);
+
+	UDMMaterialLayerObject* Layer = Stage->GetLayer();
+	check(Layer);
+
+	UDMMaterialSlot* Slot = Layer->GetSlot();
+	check(Slot);
+
+	UDynamicMaterialModelEditorOnlyData* ModelEditorOnlyData = Slot->GetMaterialModelEditorOnlyData();
+	check(ModelEditorOnlyData);
+
+	TSharedRef<FDMMaterialBuildState> BuildState = ModelEditorOnlyData->CreateBuildState(InPreviewMaterial);
+	BuildState->SetPreviewMaterial();
+
+	GenerateExpressions(BuildState);
+	UMaterialExpression* StageSourceExpression = BuildState->GetLastStageSourceExpression(this);
+
+	BuildState->GetBuildUtils().UpdatePreviewMaterial(StageSourceExpression, 0, FDMMaterialStageConnectorChannel::WHOLE_CHANNEL, 32);
+}
+
+int32 UDMMaterialStageSource::GetInnateMaskOutput(int32 OutputIndex, int32 OutputChannels) const
+{
+	return INDEX_NONE;
+}
+
+bool UDMMaterialStageSource::UpdateStagePreviewMaterial(UDMMaterialStage* InStage, UMaterial* InPreviewMaterial, 
+	UMaterialExpression*& OutMaterialExpression, int32& OutputIndex)
+{
+	check(InStage);
+	check(InPreviewMaterial);
+
+	UDMMaterialLayerObject* Layer = InStage->GetLayer();
+	check(Layer);
+
+	UDMMaterialSlot* Slot = Layer->GetSlot();
+	check(Slot);
+
+	UDynamicMaterialModelEditorOnlyData* ModelEditorOnlyData = Slot->GetMaterialModelEditorOnlyData();
+	check(ModelEditorOnlyData);
+
+	UDMMaterialProperty* PropertyObj = ModelEditorOnlyData->GetMaterialProperty(EDMMaterialPropertyType::EmissiveColor);
+	check(PropertyObj);
+
+	TSharedRef<FDMMaterialBuildState> BuildState = ModelEditorOnlyData->CreateBuildState(InPreviewMaterial);
+	BuildState->SetPreviewMaterial();
+
+	UDMMaterialStageSource* PreviewSource = InStage->GetSource();
+
+	if (!PreviewSource)
+	{
+		return false;
+	}
+
+	PreviewSource->GenerateExpressions(BuildState);
+	const TArray<UMaterialExpression*>& SourceExpressions = BuildState->GetStageSourceExpressions(PreviewSource);
+
+	if (SourceExpressions.IsEmpty())
+	{
+		return false;
+	}
+
+	UMaterialExpression* LastExpression = SourceExpressions.Last();
+
+	bool bIsMaskStage = false;
+
+	if (Layer->GetStageType(InStage) == EDMMaterialLayerStage::Mask)
+	{
+		bIsMaskStage = true;
+	}
+
+	int32 BestMatch = INDEX_NONE;
+	int32 OutputCount = 0;
+	const int32 FloatsForPropertyType = bIsMaskStage
+		? 1
+		: UDMValueDefinitionLibrary::GetValueDefinition(PropertyObj->GetInputConnectorType()).GetFloatCount();
+
+	for (int32 OutputIdx = 0; OutputIdx < LastExpression->GetOutputs().Num(); ++OutputIdx)
+	{
+		EMaterialValueType CurrentOutputType = static_cast<EMaterialValueType>(LastExpression->GetOutputType(OutputIdx));
+		int32 CurrentOutputCount = 0;
+
+		switch (CurrentOutputType)
+		{
+			case MCT_Float:
+			case MCT_Float1:
+				CurrentOutputCount = 1;
+				break;
+
+			case MCT_Float2:
+				CurrentOutputCount = 2;
+				break;
+
+			case MCT_Float3:
+				CurrentOutputCount = 3;
+				break;
+
+			case MCT_Float4:
+				CurrentOutputCount = 4;
+				break;
+
+			default:
+				continue; // For loop
+		}
+
+		if (CurrentOutputCount > OutputCount)
+		{
+			BestMatch = OutputIdx;
+			OutputCount = CurrentOutputCount;
+
+			if (CurrentOutputCount >= FloatsForPropertyType)
+			{
+				break;
+			}
+		}
+	}
+
+	if (BestMatch != INDEX_NONE)
+	{
+		OutMaterialExpression = LastExpression;
+		OutputIndex = BestMatch;
+		return true;
+	}
+
+	return false;
+}
+
+void UDMMaterialStageSource::NotifyPostChange(const FPropertyChangedEvent& InPropertyChangedEvent, class FEditPropertyChain* InPropertyThatChanged)
+{
+	if (!IsComponentValid())
+	{
+		return;
+	}
+
+	Update(EDMUpdateType::Structure);
+}
+
+void UDMMaterialStageSource::PostEditUndo()
+{
+	Super::PostEditUndo();
+
+	if (!IsComponentValid())
+	{
+		return;
+	}
+
+	MarkComponentDirty();
+	Update(EDMUpdateType::Structure);
+}
+
+void UDMMaterialStageSource::DoClean()
+{
+	if (IsComponentValid())
+	{
+		// Stage Source preview images are currently disabled.
+		//UpdatePreviewMaterial();
+	}
+
+	Super::DoClean();
+}
+
+UDMMaterialComponent* UDMMaterialStageSource::GetParentComponent() const
+{
+	return GetStage();
+}
+
+void UDMMaterialStageSource::PostEditorDuplicate(UDynamicMaterialModel* InMaterialModel, UDMMaterialComponent* InParent)
+{
+	Super::PostEditorDuplicate(InMaterialModel, InParent);
+
+	PreviewMaterial = nullptr;
+}
diff --git a/Engine/Plugins/Experimental/DynamicMaterial/Source/DynamicMaterialEditor/Private/Components/DMMaterialStageThroughput.cpp b/Engine/Plugins/Experimental/DynamicMaterial/Source/DynamicMaterialEditor/Private/Components/DMMaterialStageThroughput.cpp
new file mode 100644
--- /dev/null
+++ b/Engine/Plugins/Experimental/DynamicMaterial/Source/DynamicMaterialEditor/Private/Components/DMMaterialStageThroughput.cpp
@@ -0,0 +1,891 @@
+// Copyright Epic Games, Inc. All Rights Reserved.
+
+#include "Components/DMMaterialStageThroughput.h"
+#include "Components/DMMaterialLayer.h"
+#include "Components/DMMaterialSlot.h"
+#include "Components/DMMaterialStage.h"
+#include "Components/DMMaterialSubStage.h"
+#include "Components/MaterialStageExpressions/DMMSETextureSample.h"
+#include "Components/MaterialStageInputs/DMMSIExpression.h"
+#include "Components/MaterialStageInputs/DMMSITextureUV.h"
+#include "Containers/ArrayView.h"
+#include "DMDefs.h"
+#include "DMMaterialFunctionLibrary.h"
+#include "DMPrivate.h"
+#include "Helpers/DMInputNodeBuilder.h"
+#include "Materials/MaterialExpression.h"
+#include "Materials/MaterialExpressionAppendVector.h"
+#include "Materials/MaterialExpressionComponentMask.h"
+#include "Materials/MaterialExpressionConstant.h"
+#include "Materials/MaterialExpressionConstant2Vector.h"
+#include "Materials/MaterialExpressionConstant3Vector.h"
+#include "Materials/MaterialExpressionMaterialFunctionCall.h"
+#include "Materials/MaterialExpressionTextureBase.h"
+#include "Model/DMMaterialBuildState.h"
+#include "Model/DMMaterialBuildUtils.h"
+#include "Model/DynamicMaterialModelEditorOnlyData.h"
+
+#define LOCTEXT_NAMESPACE "DMMaterialStageThroughput"
+
+TArray<TStrongObjectPtr<UClass>> UDMMaterialStageThroughput::Throughputs = {};
+
+const TArray<TStrongObjectPtr<UClass>>& UDMMaterialStageThroughput::GetAvailableThroughputs()
+{
+	if (Throughputs.IsEmpty())
+	{
+		GenerateThroughputList();
+	}
+
+	return Throughputs;
+}
+
+bool UDMMaterialStageThroughput::CanInputAcceptType(int32 InputIndex, EDMValueType ValueType) const
+{
+	check(InputConnectors.IsValidIndex(InputIndex));
+
+	return InputConnectors[InputIndex].IsCompatibleWith(ValueType);
+}
+
+bool UDMMaterialStageThroughput::CanInputConnectTo(int32 InputIndex, const FDMMaterialStageConnector& OutputConnector, int32 OutputChannel, bool bCheckSingleFloat)
+{
+	if (OutputChannel == FDMMaterialStageConnectorChannel::WHOLE_CHANNEL)
+	{
+		if (CanInputAcceptType(InputIndex, OutputConnector.Type))
+		{
+			return true;
+		}
+
+		if (UDMValueDefinitionLibrary::GetValueDefinition(OutputConnector.Type).IsFloatType() && bCheckSingleFloat)
+		{
+			return CanInputAcceptType(InputIndex, EDMValueType::VT_Float1);
+		}
+
+		return false;
+	}
+
+	if (UDMValueDefinitionLibrary::GetValueDefinition(OutputConnector.Type).IsFloatType())
+	{
+		return CanInputAcceptType(InputIndex, EDMValueType::VT_Float1);
+	}
+
+	// Only float types can have non-whole channels.
+	checkNoEntry();
+	return false;
+}
+
+bool UDMMaterialStageThroughput::CanChangeInput(int32 InputIndex) const
+{
+	return true;
+}
+
+bool UDMMaterialStageThroughput::CanChangeInputType(int32 InputIndex) const
+{
+	check(InputConnectors.IsValidIndex(InputIndex));
+
+	return InputConnectors[InputIndex].Type != EDMValueType::VT_Texture;
+}
+
+bool UDMMaterialStageThroughput::IsInputVisible(int32 InputIndex) const
+{
+	UDMMaterialStage* Stage = GetStage();
+	check(Stage);
+
+	if (UDMMaterialSubStage* SubStage = Cast<UDMMaterialSubStage>(Stage))
+	{
+		Stage = SubStage->GetParentMostStage();
+	}
+
+	UDMMaterialLayerObject* Layer = Stage->GetLayer();
+	check(Layer);
+
+	bool bSupportsUVLink = false;
+
+	if (SupportsLayerMaskTextureUVLink())
+	{
+		if (UDMMaterialStage* BaseStage = Layer->GetStage(EDMMaterialLayerStage::Base))
+		{
+			if (BaseStage->IsEnabled())
+			{
+				if (UDMMaterialStageThroughput* BaseThroughput = Cast<UDMMaterialStageThroughput>(BaseStage->GetSource()))
+				{
+					bSupportsUVLink = BaseThroughput->SupportsLayerMaskTextureUVLink();
+				}
+			}
+		}
+	}
+
+	if (bSupportsUVLink && InputIndex == GetLayerMaskTextureUVLinkInputIndex())
+	{
+		return !Layer->IsTextureUVLinkEnabled() || Layer->GetStageType(Stage) != EDMMaterialLayerStage::Mask;
+	}
+
+	return true;
+}
+
+bool UDMMaterialStageThroughput::ShouldKeepInput(int32 InInputIdx)
+{
+	UDMMaterialStage* Stage = GetStage();
+	check(Stage);
+
+	if (!CanChangeInput(InInputIdx) || !Stage->IsInputMapped(InInputIdx))
+	{
+		return false;
+	}
+
+	UDMMaterialLayerObject* Layer = Stage->GetLayer();
+	check(Layer);
+
+	const EDMMaterialPropertyType StageProperty = Layer->GetMaterialProperty();
+	const TArray<FDMMaterialStageConnection>& InputConnectionMap = Stage->GetInputConnectionMap();
+
+	if (!InputConnectionMap.IsValidIndex(InInputIdx))
+	{
+		return false;
+	}
+
+	if (InputConnectors[InInputIdx].Type == EDMValueType::VT_Float_Any || InputConnectors[InInputIdx].Type == EDMValueType::VT_Texture)
+	{
+		return false;
+	}
+
+	const int RequiredInputCount = UDMValueDefinitionLibrary::GetValueDefinition(InputConnectors[InInputIdx].Type).GetFloatCount();
+
+	int ActualInputCount = 0;
+
+	for (const FDMMaterialStageConnectorChannel& Channel : InputConnectionMap[InInputIdx].Channels)
+	{
+		int32 ThisInputCount = 0;
+
+		switch (Channel.SourceIndex)
+		{
+			case FDMMaterialStageConnectorChannel::NO_SOURCE:
+				return false; // Some sort of error or badly mapped input.
+
+			case FDMMaterialStageConnectorChannel::PREVIOUS_STAGE:
+			{
+				if (const UDMMaterialLayerObject* PreviousLayer = Layer->GetPreviousLayer(StageProperty, EDMMaterialLayerStage::Base))
+				{
+					if (UDMMaterialStage* MaskStage = PreviousLayer->GetStage(EDMMaterialLayerStage::Mask))
+					{
+						check(MaskStage->GetSource());
+
+						const TArray<FDMMaterialStageConnector>& PreviousStageOutputConnectors = MaskStage->GetSource()->GetOutputConnectors();
+						check(PreviousStageOutputConnectors.IsValidIndex(Channel.OutputIndex));
+
+						if (UDMValueDefinitionLibrary::GetValueDefinition(PreviousStageOutputConnectors[Channel.OutputChannel].Type).IsFloatType() == false)
+						{
+							check(InputConnectionMap[InInputIdx].Channels.Num() == 1);
+
+							return InputConnectors[InInputIdx].IsCompatibleWith(PreviousStageOutputConnectors[Channel.OutputChannel]);
+						}
+						else
+						{
+							check(RequiredInputCount > 0);
+						}
+
+						if (Channel.OutputChannel == FDMMaterialStageConnectorChannel::WHOLE_CHANNEL)
+						{
+							ThisInputCount = UDMValueDefinitionLibrary::GetValueDefinition(PreviousStageOutputConnectors[Channel.OutputChannel].Type).GetFloatCount();
+						}
+						else
+						{
+							ThisInputCount =
+								!!(Channel.OutputChannel & FDMMaterialStageConnectorChannel::FIRST_CHANNEL)
+								+ !!(Channel.OutputChannel & FDMMaterialStageConnectorChannel::SECOND_CHANNEL)
+								+ !!(Channel.OutputChannel & FDMMaterialStageConnectorChannel::THIRD_CHANNEL)
+								+ !!(Channel.OutputChannel & FDMMaterialStageConnectorChannel::FOURTH_CHANNEL);
+
+							ThisInputCount = FMath::Min(
+								ThisInputCount,
+								static_cast<int32>(UDMValueDefinitionLibrary::GetValueDefinition(PreviousStageOutputConnectors[Channel.OutputChannel].Type).GetFloatCount())
+							);
+						}
+					}
+				}
+
+				break;
+			}
+
+			// Input
+			default:
+			{
+				const int32 InputObjectIdx = Channel.SourceIndex - FDMMaterialStageConnectorChannel::FIRST_STAGE_INPUT;
+				const TArray<UDMMaterialStageInput*>& StageInputs = Stage->GetInputs();
+				check(StageInputs.IsValidIndex(InputObjectIdx));
+
+				const TArray<FDMMaterialStageConnector>& InputOutputConnectors = StageInputs[InputObjectIdx]->GetOutputConnectors();
+				check(InputOutputConnectors.IsValidIndex(Channel.OutputIndex));
+
+				if (UDMValueDefinitionLibrary::GetValueDefinition(InputOutputConnectors[Channel.OutputChannel].Type).IsFloatType() == false)
+				{
+					check(InputConnectionMap[InInputIdx].Channels.Num() == 1);
+
+					return InputConnectors[InInputIdx].IsCompatibleWith(InputOutputConnectors[Channel.OutputChannel]);
+				}
+				else
+				{
+					check(RequiredInputCount > 0);
+				}
+
+				if (Channel.OutputChannel == FDMMaterialStageConnectorChannel::WHOLE_CHANNEL)
+				{
+					ThisInputCount = UDMValueDefinitionLibrary::GetValueDefinition(InputOutputConnectors[Channel.OutputChannel].Type).GetFloatCount();
+				}
+				else
+				{
+					ThisInputCount =
+						!!(Channel.OutputChannel & FDMMaterialStageConnectorChannel::FIRST_CHANNEL)
+						+ !!(Channel.OutputChannel & FDMMaterialStageConnectorChannel::SECOND_CHANNEL)
+						+ !!(Channel.OutputChannel & FDMMaterialStageConnectorChannel::THIRD_CHANNEL)
+						+ !!(Channel.OutputChannel & FDMMaterialStageConnectorChannel::FOURTH_CHANNEL);
+
+					ThisInputCount = FMath::Min(
+						ThisInputCount,
+						static_cast<int32>(UDMValueDefinitionLibrary::GetValueDefinition(InputOutputConnectors[Channel.OutputChannel].Type).GetFloatCount())
+					);
+				}
+
+				break;
+			}
+		}
+
+		ActualInputCount += ThisInputCount;
+	}
+
+	return (ActualInputCount == RequiredInputCount);
+}
+
+void UDMMaterialStageThroughput::ConnectOutputToInput(const TSharedRef<FDMMaterialBuildState>& InBuildState, int32 InInputIndex,
+	UMaterialExpression* InSourceExpression, int32 InSourceOutputIndex, int32 InSourceOutputChannel)
+{
+	check(InSourceExpression);
+	check(InSourceExpression->GetOutputs().IsValidIndex(InSourceOutputIndex));
+
+	const TArray<UMaterialExpression*>& TargetExpressions = InBuildState->GetStageSourceExpressions(this);
+	check(!TargetExpressions.IsEmpty());
+
+	UMaterialExpression* TargetExpression = GetExpressionForInput(TargetExpressions, InInputIndex);
+	TConstArrayView<FExpressionInput*> TargetExpressionInputs = TargetExpression->GetInputsView();
+	check(TargetExpressionInputs.IsValidIndex(InInputIndex));
+
+	ConnectOutputToInput_Internal(InBuildState, TargetExpression, InInputIndex, InSourceExpression, InSourceOutputIndex, InSourceOutputChannel);
+}
+
+void UDMMaterialStageThroughput::ConnectOutputToInput_Internal(const TSharedRef<FDMMaterialBuildState>& InBuildState, UMaterialExpression* TargetExpression,
+	int32 InputIndex, UMaterialExpression* SourceExpression, int32 SourceOutputIndex, int32 SourceOutputChannel) const
+{
+	check(TargetExpression != SourceExpression);
+
+	FExpressionInput* ExpressionInput = TargetExpression->GetInput(InputIndex);
+	
+	if (SourceOutputChannel == FDMMaterialStageConnectorChannel::WHOLE_CHANNEL)
+	{
+		SourceExpression->ConnectExpression(ExpressionInput, SourceOutputIndex);
+	}
+	else if (SourceExpression->IsA<UMaterialExpressionTextureBase>()
+		&& (SourceOutputChannel == FDMMaterialStageConnectorChannel::FIRST_CHANNEL
+			|| SourceOutputChannel == FDMMaterialStageConnectorChannel::SECOND_CHANNEL
+			|| SourceOutputChannel == FDMMaterialStageConnectorChannel::THIRD_CHANNEL
+			|| SourceOutputChannel == FDMMaterialStageConnectorChannel::FOURTH_CHANNEL))
+	{
+		switch (SourceOutputChannel)
+		{
+			case FDMMaterialStageConnectorChannel::FIRST_CHANNEL:
+				SourceExpression->ConnectExpression(ExpressionInput, 1);
+				break;
+
+			case FDMMaterialStageConnectorChannel::SECOND_CHANNEL:
+				SourceExpression->ConnectExpression(ExpressionInput, 2);
+				break;
+
+			case FDMMaterialStageConnectorChannel::THIRD_CHANNEL:
+				SourceExpression->ConnectExpression(ExpressionInput, 3);
+				break;
+
+			case FDMMaterialStageConnectorChannel::FOURTH_CHANNEL:
+				SourceExpression->ConnectExpression(ExpressionInput, 4);
+				break;
+
+			default:
+				checkNoEntry();
+		}
+	}
+	else
+	{
+		UMaterialExpression* MaskExpression = InBuildState->GetBuildUtils().CreateExpressionBitMask(SourceExpression, SourceOutputIndex, SourceOutputChannel);
+		MaskExpression->ConnectExpression(ExpressionInput, 0);
+	}
+}
+
+int32 UDMMaterialStageThroughput::GetLayerMaskTextureUVLinkInputIndex() const
+{
+	return INDEX_NONE;
+}
+
+FDMExpressionInput UDMMaterialStageThroughput::GetLayerMaskLinkTextureUVInputExpressions(const TSharedRef<FDMMaterialBuildState>& InBuildState) const
+{
+	check(SupportsLayerMaskTextureUVLink());
+
+	int32 TextureUVInputIdx = GetLayerMaskTextureUVLinkInputIndex();
+	check(TextureUVInputIdx != INDEX_NONE);
+
+	FDMMaterialStageConnectorChannel Channel;
+	FDMExpressionInput ExpressionInput;
+
+	ExpressionInput.OutputIndex = ResolveInput(InBuildState, TextureUVInputIdx, Channel, ExpressionInput.OutputExpressions);
+	ExpressionInput.OutputChannel = Channel.OutputChannel;
+
+	return ExpressionInput;
+}
+
+UMaterialExpression* UDMMaterialStageThroughput::GetExpressionForInput(const TArray<UMaterialExpression*>& StageSourceExpressions, int32 InputIdx)
+{
+	check(!StageSourceExpressions.IsEmpty());
+
+	return StageSourceExpressions[0];
+}
+
+void UDMMaterialStageThroughput::OnComponentAdded()
+{
+	if (!IsComponentValid())
+	{
+		return;
+	}
+
+	Super::OnComponentAdded();
+
+	if (IsComponentValid())
+	{
+		const FDMUpdateGuard Guard;
+
+		for (int32 InputIdx = 0; InputIdx < InputConnectors.Num(); ++InputIdx)
+		{
+			if (ShouldKeepInput(InputIdx))
+			{
+				continue;
+			}
+
+			AddDefaultInput(InputIdx);
+		}
+	}
+}
+
+void UDMMaterialStageThroughput::AddDefaultInput(int32 InInputIndex) const
+{
+	if (!IsComponentValid())
+	{
+		return;
+	}
+
+	static const FText UVName = LOCTEXT("UV", "UV");
+	
+	check(InputConnectors.IsValidIndex(InInputIndex));
+
+	UDMMaterialStage* Stage = GetStage();
+	check(Stage);
+
+	switch (InputConnectors[InInputIndex].Type)
+	{
+		default:
+			checkNoEntry();
+			break;
+
+		case EDMValueType::VT_None:
+		case EDMValueType::VT_Float_Any:
+			break;
+
+		case EDMValueType::VT_Float2:
+			if (InputConnectors[InInputIndex].Name.EqualTo(UVName))
+			{
+				Stage->ChangeInput_UV(InInputIndex, FDMMaterialStageConnectorChannel::WHOLE_CHANNEL, FDMMaterialStageConnectorChannel::WHOLE_CHANNEL);
+			}
+			else
+			{
+				Stage->ChangeInput_NewLocalValue(InInputIndex, FDMMaterialStageConnectorChannel::WHOLE_CHANNEL, EDMValueType::VT_Float2, FDMMaterialStageConnectorChannel::WHOLE_CHANNEL);
+			}
+			break;
+
+		case EDMValueType::VT_Float3_RGB:
+			Stage->ChangeInput_Expression(InInputIndex, FDMMaterialStageConnectorChannel::WHOLE_CHANNEL,
+				UDMMaterialStageExpressionTextureSample::StaticClass(), 0, FDMMaterialStageConnectorChannel::WHOLE_CHANNEL);
+			break;
+
+		case EDMValueType::VT_Float4_RGBA:
+			Stage->ChangeInput_Expression(InInputIndex, FDMMaterialStageConnectorChannel::WHOLE_CHANNEL,
+				UDMMaterialStageExpressionTextureSample::StaticClass(), 5, FDMMaterialStageConnectorChannel::WHOLE_CHANNEL);
+			break;
+
+		case EDMValueType::VT_Float1:
+		case EDMValueType::VT_Float3_RPY:
+		case EDMValueType::VT_Float3_XYZ:
+		case EDMValueType::VT_Texture:
+			Stage->ChangeInput_NewLocalValue(InInputIndex, FDMMaterialStageConnectorChannel::WHOLE_CHANNEL, InputConnectors[InInputIndex].Type, FDMMaterialStageConnectorChannel::WHOLE_CHANNEL);
+			break;
+
+		case EDMValueType::VT_ColorAtlas:
+		{
+			if (const UDMMaterialLayerObject* Layer = Stage->GetLayer())
+			{
+				if (Layer->GetStageType(Stage) == EDMMaterialLayerStage::Mask)
+				{
+					Stage->ChangeInput_NewLocalValue(InInputIndex, FDMMaterialStageConnectorChannel::WHOLE_CHANNEL, EDMValueType::VT_ColorAtlas, FDMMaterialStageConnectorChannel::FOURTH_CHANNEL);
+					break;
+				}
+			}
+
+			Stage->ChangeInput_NewLocalValue(InInputIndex, FDMMaterialStageConnectorChannel::WHOLE_CHANNEL, EDMValueType::VT_ColorAtlas, FDMMaterialStageConnectorChannel::THREE_CHANNELS);
+			break;
+		}
+	}
+}
+
+int32 UDMMaterialStageThroughput::ResolveInput(const TSharedRef<FDMMaterialBuildState>& InBuildState, int32 InputIndex, FDMMaterialStageConnectorChannel& OutChannel,
+	TArray<UMaterialExpression*>& OutExpressions) const
+{
+	UDMMaterialStage* Stage = GetStage();
+	check(Stage);
+
+	UDMMaterialStage* ParentMostStage = Stage;
+
+	if (UDMMaterialSubStage* SubStage = Cast<UDMMaterialSubStage>(Stage))
+	{
+		ParentMostStage = SubStage->GetParentMostStage();
+	}
+
+	const UDMMaterialLayerObject* Layer = ParentMostStage->GetLayer();
+	check(Layer);
+
+	if (Layer->IsTextureUVLinkEnabled() 
+		&& Layer->GetStage(EDMMaterialLayerStage::Base, /* Enabled Only */ true)
+		&& Layer->GetStageType(ParentMostStage) == EDMMaterialLayerStage::Mask
+		&& SupportsLayerMaskTextureUVLink() 
+		&& (GetLayerMaskTextureUVLinkInputIndex() == INDEX_NONE || GetLayerMaskTextureUVLinkInputIndex() == InputIndex))
+	{
+		int32 OutputIndex = ResolveLayerMaskTextureUVLinkInput(InBuildState, InputIndex, OutChannel, OutExpressions);
+
+		if (OutputIndex != INDEX_NONE)
+		{
+			return OutChannel.OutputIndex;
+		}
+	}
+
+	const TArray<FDMMaterialStageConnection>& InputConnectionMap = Stage->GetInputConnectionMap();
+
+	if (!InputConnectionMap.IsValidIndex(InputIndex))
+	{
+		return INDEX_NONE;
+	}
+
+	if (InputConnectionMap[InputIndex].Channels.IsEmpty())
+	{
+		return INDEX_NONE;
+	}
+
+	if (InputConnectionMap[InputIndex].Channels.Num() == 1)
+	{
+		OutChannel = InputConnectionMap[InputIndex].Channels[0]; // Full copy in case it is changed by the channel resolve.
+		return ResolveInputChannel(InBuildState, InputIndex, 0, OutChannel, OutExpressions);
+	}
+
+	// Only floats can have sub-channel mapping
+	check(UDMValueDefinitionLibrary::GetValueDefinition(InputConnectors[InputIndex].Type).IsFloatType());
+
+	// Valid scalar/vector is from 1 to 4 floats.
+	check(InputConnectionMap[InputIndex].Channels.Num() <= 4);
+
+	for (const FDMMaterialStageConnectorChannel& Channel : InputConnectionMap[InputIndex].Channels)
+	{
+		check(Channel.OutputChannel != FDMMaterialStageConnectorChannel::WHOLE_CHANNEL);
+	}
+
+	OutChannel.OutputIndex = 0;
+	OutChannel.OutputChannel = FDMMaterialStageConnectorChannel::WHOLE_CHANNEL;
+
+	UMaterialExpression* MakeFloat = nullptr;
+
+	switch (InputConnectionMap[InputIndex].Channels.Num())
+	{
+		case 2:
+			MakeFloat = FDMMaterialFunctionLibrary::Get().GetMakeFloat2(InBuildState->GetDynamicMaterial(), UE_DM_NodeComment_Default);
+			break;
+
+		case 3:
+			MakeFloat = FDMMaterialFunctionLibrary::Get().GetMakeFloat3(InBuildState->GetDynamicMaterial(), UE_DM_NodeComment_Default);
+			break;
+
+		case 4:
+			MakeFloat = FDMMaterialFunctionLibrary::Get().GetMakeFloat3(InBuildState->GetDynamicMaterial(), UE_DM_NodeComment_Default);
+			break;
+
+		default:
+			checkNoEntry();
+			return INDEX_NONE;
+	}
+
+	for (int32 ChannelIdx = 0; ChannelIdx < InputConnectionMap[InputIndex].Channels.Num(); ++ChannelIdx)
+	{
+		FDMMaterialStageConnectorChannel ChannelTemp = InputConnectionMap[InputIndex].Channels[ChannelIdx];
+		TArray<UMaterialExpression*> ChannelExpressions;
+
+		ResolveInputChannel(InBuildState, InputIndex, ChannelIdx, ChannelTemp, ChannelExpressions);
+		check(ChannelExpressions.IsEmpty() == false);
+
+		OutExpressions.Append(ChannelExpressions);
+
+		ChannelExpressions.Last()->ConnectExpression(MakeFloat->GetInput(ChannelIdx), ChannelTemp.OutputIndex);
+	}
+
+	OutExpressions.Add(MakeFloat);
+
+	return 0;
+}
+
+int32 UDMMaterialStageThroughput::ResolveLayerMaskTextureUVLinkInput(const TSharedRef<FDMMaterialBuildState>& InBuildState, int32 InputIndex, 
+	FDMMaterialStageConnectorChannel& OutChannel, TArray<UMaterialExpression*>& OutExpressions) const
+{
+	UDMMaterialStage* Stage = GetStage();
+	check(Stage);
+
+	const UDMMaterialLayerObject* Layer = Stage->GetLayer();
+	check(Layer);
+	check(Layer->GetStage(EDMMaterialLayerStage::Base));
+
+	return ResolveLayerMaskTextureUVLinkInputImpl(InBuildState, Layer->GetStage(EDMMaterialLayerStage::Base)->GetSource(), OutChannel, OutExpressions);
+}
+
+int32 UDMMaterialStageThroughput::ResolveLayerMaskTextureUVLinkInputImpl(const TSharedRef<FDMMaterialBuildState>& InBuildState, 
+	const UDMMaterialStageSource* StageSource, FDMMaterialStageConnectorChannel& OutChannel, TArray<UMaterialExpression*>& OutExpressions)
+{
+	FDMExpressionInput ConnectedInputExpressions;
+
+	if (const UDMMaterialStageInputTextureUV* InputTextureUV = Cast<const UDMMaterialStageInputTextureUV>(StageSource))
+	{
+		InputTextureUV->GenerateExpressions(InBuildState);
+		OutExpressions = InBuildState->GetStageSourceExpressions(InputTextureUV);
+		OutChannel.OutputIndex = 0;
+		OutChannel.OutputChannel = FDMMaterialStageConnectorChannel::WHOLE_CHANNEL;
+
+		return 0;
+	}
+
+	else if (const UDMMaterialStageThroughput* Throughput = Cast<const UDMMaterialStageThroughput>(StageSource))
+	{
+		if (Throughput->SupportsLayerMaskTextureUVLink())
+		{
+			ConnectedInputExpressions = Throughput->GetLayerMaskLinkTextureUVInputExpressions(InBuildState);
+		}
+	}
+
+	else if (const UDMMaterialStageInputThroughput* InputThroughput = Cast<const UDMMaterialStageInputThroughput>(StageSource))
+	{
+		if (InputThroughput->GetMaterialStageThroughput())
+		{
+			UDMMaterialStageThroughput* InputThroughputActual = InputThroughput->GetMaterialStageThroughput();
+			check(InputThroughputActual);
+
+			if (InputThroughputActual->SupportsLayerMaskTextureUVLink())
+			{
+				ConnectedInputExpressions = InputThroughputActual->GetLayerMaskLinkTextureUVInputExpressions(InBuildState);
+			}
+		}
+	}
+
+	if (ConnectedInputExpressions.IsValid())
+	{
+		OutChannel.OutputIndex = ConnectedInputExpressions.OutputIndex;
+		OutChannel.OutputChannel = ConnectedInputExpressions.OutputChannel;
+		OutExpressions = ConnectedInputExpressions.OutputExpressions;
+
+		return ConnectedInputExpressions.OutputIndex;
+	}
+
+	return INDEX_NONE;
+}
+
+int32 UDMMaterialStageThroughput::ResolveInputChannel(const TSharedRef<FDMMaterialBuildState>& InBuildState, int32 InputIndex, int32 ChannelIndex,
+	FDMMaterialStageConnectorChannel& OutChannel, TArray<UMaterialExpression*>& OutExpressions) const
+{
+	if (OutChannel.SourceIndex == FDMMaterialStageConnectorChannel::NO_SOURCE)
+	{
+		return INDEX_NONE;
+	}
+
+	check(InputConnectors.IsValidIndex(InputIndex));
+
+	UDMMaterialStage* Stage = GetStage();
+	check(Stage);
+
+	const TArray<FDMMaterialStageConnector>* InputSourceOutputConnectors = nullptr;
+	int32 InnateMaskOutput = INDEX_NONE;
+	int32 NodeOutputIndex = INDEX_NONE;
+	int32 OutputChannelOverride = INDEX_NONE;
+
+	if (OutChannel.SourceIndex == FDMMaterialStageConnectorChannel::PREVIOUS_STAGE)
+	{
+		const UDMMaterialLayerObject* Layer = Stage->GetLayer();
+		check(Layer);
+
+		if (const UDMMaterialLayerObject* PreviousLayer = Layer->GetPreviousLayer(OutChannel.MaterialProperty, EDMMaterialLayerStage::Base))
+		{
+			PreviousLayer->GenerateExpressions(InBuildState);
+
+			OutExpressions.Append(InBuildState->GetLayerExpressions(PreviousLayer));
+			check(!OutExpressions.IsEmpty());
+
+			UDMMaterialStage* PreviousStage = PreviousLayer->GetLastValidStage(EDMMaterialLayerStage::All);
+			check(PreviousStage);
+
+			UDMMaterialStageSource* PreviousStageSource = PreviousStage->GetSource();
+			check(PreviousStageSource);
+
+			InputSourceOutputConnectors = &(PreviousStageSource->GetOutputConnectors());
+			check(InputSourceOutputConnectors->IsValidIndex(OutChannel.OutputIndex));
+
+			NodeOutputIndex = (*InputSourceOutputConnectors)[OutChannel.OutputIndex].Index;
+			InnateMaskOutput = PreviousStageSource->GetInnateMaskOutput(NodeOutputIndex, OutChannel.OutputChannel);
+		}
+		else
+		{
+			UMaterialExpressionConstant3Vector* Black = InBuildState->GetBuildUtils().CreateExpression<UMaterialExpressionConstant3Vector>(UE_DM_NodeComment_Default);
+			Black->Constant = FLinearColor(0, 0, 0, 1);
+
+			OutExpressions.Add(Black);
+
+			NodeOutputIndex = 0;
+			InnateMaskOutput = INDEX_NONE;
+		}
+	}
+	else
+	{
+		const TArray<UDMMaterialStageInput*> StageInputs = Stage->GetInputs();
+
+		const int32 StageInputIdx = OutChannel.SourceIndex - FDMMaterialStageConnectorChannel::FIRST_STAGE_INPUT;
+		check(StageInputs.IsValidIndex(StageInputIdx));
+
+		UDMMaterialStageInput* InputValue = StageInputs[StageInputIdx];
+		InputValue->GenerateExpressions(InBuildState);
+
+		OutExpressions.Append(InBuildState->GetStageSourceExpressions(InputValue));
+		check(!OutExpressions.IsEmpty());
+
+		InputSourceOutputConnectors = &(InputValue->GetOutputConnectors());
+		check(InputSourceOutputConnectors->IsValidIndex(OutChannel.OutputIndex));
+
+		// Change from output index from Material Designer node to output index of the material expression
+		NodeOutputIndex = (*InputSourceOutputConnectors)[OutChannel.OutputIndex].Index;
+		InnateMaskOutput = InputValue->GetInnateMaskOutput(NodeOutputIndex, OutChannel.OutputChannel);
+		OutputChannelOverride = InputValue->GetOutputChannelOverride(OutChannel.OutputIndex);
+	}
+
+	// If our "Previous Stage" is blank, these need default values.
+	const int32 OutputFloatCount = InputSourceOutputConnectors ? UDMValueDefinitionLibrary::GetValueDefinition((*InputSourceOutputConnectors)[OutChannel.OutputIndex].Type).GetFloatCount() : 3;
+	const bool bOutputIsFloatType = InputSourceOutputConnectors ? UDMValueDefinitionLibrary::GetValueDefinition((*InputSourceOutputConnectors)[OutChannel.OutputIndex].Type).IsFloatType() : true;
+	const int32 InputFloatCount = UDMValueDefinitionLibrary::GetValueDefinition(InputConnectors[InputIndex].Type).GetFloatCount();
+
+	if (OutChannel.OutputChannel == FDMMaterialStageConnectorChannel::WHOLE_CHANNEL
+		&& bOutputIsFloatType && OutputFloatCount > 0 && InputFloatCount > 0)
+	{
+		if (OutputFloatCount > InputFloatCount)
+		{
+			int32 ComponentMask = 0;
+
+			for (int32 FloatIdx = 1; FloatIdx <= InputFloatCount; ++FloatIdx)
+			{
+				ComponentMask += UE::DynamicMaterialEditor::ChannelIndexToChannelBit(FloatIdx);
+			}
+
+			OutChannel.OutputChannel = ComponentMask;
+		}
+	}
+
+	if (OutputChannelOverride != INDEX_NONE)
+	{
+		OutChannel.OutputChannel = OutputChannelOverride;
+	}
+
+	if (InnateMaskOutput != INDEX_NONE)
+	{
+		OutChannel.OutputChannel = FDMMaterialStageConnectorChannel::WHOLE_CHANNEL;
+
+		return InnateMaskOutput;
+	}
+
+	if (OutChannel.OutputChannel == FDMMaterialStageConnectorChannel::WHOLE_CHANNEL
+		&& bOutputIsFloatType && OutputFloatCount > 0 && InputFloatCount > 0 && OutputFloatCount < InputFloatCount)
+	{
+		UMaterialExpressionAppendVector* Append = InBuildState->GetBuildUtils().CreateExpression<UMaterialExpressionAppendVector>(UE_DM_NodeComment_Default);
+		Append->A.Expression = OutExpressions.Last();
+		Append->A.OutputIndex = NodeOutputIndex;
+
+		switch (InputFloatCount - OutputFloatCount)
+		{
+			case 1:
+			{
+				UMaterialExpressionConstant* Constant = InBuildState->GetBuildUtils().CreateExpression<UMaterialExpressionConstant>(UE_DM_NodeComment_Default);
+				OutExpressions.Add(Constant);
+				Constant->R = 0.f;
+				Append->B.Expression = Constant;
+				break;
+			}
+
+			case 2:
+			{
+				UMaterialExpressionConstant2Vector* Constant2 = InBuildState->GetBuildUtils().CreateExpression<UMaterialExpressionConstant2Vector>(UE_DM_NodeComment_Default);
+				OutExpressions.Add(Constant2);
+				Constant2->R = 0.f;
+				Constant2->G = 0.f;
+				Append->B.Expression = Constant2;
+				break;
+			}
+
+			case 3:
+			{
+				UMaterialExpressionConstant3Vector* Constant3 = InBuildState->GetBuildUtils().CreateExpression<UMaterialExpressionConstant3Vector>(UE_DM_NodeComment_Default);
+				OutExpressions.Add(Constant3);
+				Constant3->Constant = FLinearColor::Black;
+				Append->B.Expression = Constant3;
+				break;
+			}
+
+			default:
+				checkNoEntry();
+				break;
+		}
+
+		Append->B.OutputIndex = 0;
+		OutExpressions.Add(Append);
+
+		OutChannel.OutputChannel = FDMMaterialStageConnectorChannel::WHOLE_CHANNEL;
+
+		return 0;
+	}
+
+	return NodeOutputIndex;
+}
+
+void UDMMaterialStageThroughput::UpdatePreviewMaterial(UMaterial* InPreviewMaterial /*= nullptr*/)
+{
+	if (!IsComponentValid())
+	{
+		return;
+	}
+
+	if (!InPreviewMaterial)
+	{
+		if (!PreviewMaterial)
+		{
+			CreatePreviewMaterial();
+		}
+
+		InPreviewMaterial = PreviewMaterial;
+
+		if (!InPreviewMaterial)
+		{
+			return;
+		}
+	}
+
+	UDMMaterialStage* Stage = GetStage();
+	check(Stage);
+
+	UDMMaterialLayerObject* Layer = Stage->GetLayer();
+	check(Layer);
+
+	UDMMaterialSlot* Slot = Layer->GetSlot();
+	check(Slot);
+
+	UDynamicMaterialModelEditorOnlyData* ModelEditorOnlyData = Slot->GetMaterialModelEditorOnlyData();
+	check(ModelEditorOnlyData);
+
+	const TArray<FDMMaterialStageConnection>& InputConnectionMap = Stage->GetInputConnectionMap();
+	const TArray<UDMMaterialStageInput*>& StageInputs = Stage->GetInputs();
+	TArray<UE::DynamicMaterialEditor::Private::FDMInputInputs> Inputs; // There can be multiple inputs per input
+	bool bHasStageInput = false;
+
+	for (int32 InputIdx = 0; InputIdx < InputConnectors.Num(); ++InputIdx)
+	{
+		TArray<UDMMaterialStageInput*> ChannelInputs;
+
+		if (!InputConnectionMap.IsValidIndex(InputIdx))
+		{
+			continue;
+		}
+
+		ChannelInputs.SetNum(InputConnectionMap[InputIdx].Channels.Num());
+		bool bNonStageInput = false;
+
+		for (int32 ChannelIdx = 0; ChannelIdx < InputConnectionMap[InputIdx].Channels.Num(); ++ChannelIdx)
+		{
+			if (InputConnectionMap[InputIdx].Channels[ChannelIdx].SourceIndex == FDMMaterialStageConnectorChannel::PREVIOUS_STAGE)
+			{
+				bHasStageInput = true;
+				ChannelInputs[ChannelIdx] = nullptr;
+				continue;
+			}
+
+			if (InputConnectionMap[InputIdx].Channels[ChannelIdx].SourceIndex >= FDMMaterialStageConnectorChannel::FIRST_STAGE_INPUT)
+			{
+				int32 StageInputIdx = InputConnectionMap[InputIdx].Channels[ChannelIdx].SourceIndex - FDMMaterialStageConnectorChannel::FIRST_STAGE_INPUT;
+				ChannelInputs[ChannelIdx] = StageInputs[StageInputIdx];
+				bNonStageInput = true;
+				continue;
+			}
+		}
+
+		if (ChannelInputs.IsEmpty())
+		{
+			continue;
+		}
+
+		if (bNonStageInput)
+		{
+			Inputs.Add({InputIdx, ChannelInputs});
+		}
+	}
+
+	TSharedRef<FDMMaterialBuildState> BuildState = ModelEditorOnlyData->CreateBuildState(InPreviewMaterial);
+	BuildState->SetPreviewMaterial();
+
+	if (!bHasStageInput || Inputs.IsEmpty())
+	{
+		Stage->GenerateExpressions(BuildState);
+		UMaterialExpression* StageExpression = BuildState->GetLastStageExpression(Stage);
+
+		BuildState->GetBuildUtils().UpdatePreviewMaterial(StageExpression, 0, FDMMaterialStageConnectorChannel::WHOLE_CHANNEL, 32);
+	}
+	else
+	{
+		UE::DynamicMaterialEditor::Private::BuildExpressionInputs(BuildState, InputConnectionMap, Inputs);
+	}
+}
+
+void UDMMaterialStageThroughput::GenerateThroughputList()
+{
+	Throughputs.Empty();
+
+	const TArray<TStrongObjectPtr<UClass>>& SourceList = UDMMaterialStageSource::GetAvailableSourceClasses();
+
+	for (const TStrongObjectPtr<UClass>& SourceClass : SourceList)
+	{
+		UDMMaterialStageThroughput* StageThroughputCDO = Cast<UDMMaterialStageThroughput>(SourceClass->GetDefaultObject(true));
+
+		if (!StageThroughputCDO)
+		{
+			continue;
+		}
+
+		Throughputs.Add(SourceClass);
+	}
+}
+
+UDMMaterialStageThroughput::UDMMaterialStageThroughput()
+	: UDMMaterialStageThroughput(FText::GetEmpty())
+{
+}
+
+UDMMaterialStageThroughput::UDMMaterialStageThroughput(const FText& InName)
+	: Name(InName)
+{
+}
+
+#undef LOCTEXT_NAMESPACE
diff --git a/Engine/Plugins/Experimental/DynamicMaterial/Source/DynamicMaterialEditor/Private/Components/MaterialProperties/DMMPAnisotropy.cpp b/Engine/Plugins/Experimental/DynamicMaterial/Source/DynamicMaterialEditor/Private/Components/MaterialProperties/DMMPAnisotropy.cpp
new file mode 100644
--- /dev/null
+++ b/Engine/Plugins/Experimental/DynamicMaterial/Source/DynamicMaterialEditor/Private/Components/MaterialProperties/DMMPAnisotropy.cpp
@@ -0,0 +1,34 @@
+// Copyright Epic Games, Inc. All Rights Reserved.
+
+#include "Components/MaterialProperties/DMMPAnisotropy.h"
+#include "Model/DynamicMaterialModelEditorOnlyData.h"
+
+UDMMaterialPropertyAnisotropy::UDMMaterialPropertyAnisotropy()
+	: UDMMaterialProperty(
+		EDMMaterialPropertyType(EDMMaterialPropertyType::Anisotropy),
+		EDMValueType::VT_Float1)
+{
+}
+
+bool UDMMaterialPropertyAnisotropy::IsValidForModel(UDynamicMaterialModelEditorOnlyData& InModelEditorOnlyData) const
+{
+	if (InModelEditorOnlyData.GetShadingModel() == EDMMaterialShadingModel::Unlit)
+	{
+		return false;
+	}
+
+	switch (InModelEditorOnlyData.GetBlendMode())
+	{
+		case EBlendMode::BLEND_Opaque:
+		case EBlendMode::BLEND_Masked:
+			return true;
+
+		default:
+			return false;
+	}
+}
+
+UMaterialExpression* UDMMaterialPropertyAnisotropy::GetDefaultInput(const TSharedRef<FDMMaterialBuildState>& InBuildState) const
+{
+	return CreateConstant(InBuildState, 1.f);
+}
diff --git a/Engine/Plugins/Experimental/DynamicMaterial/Source/DynamicMaterialEditor/Private/Components/MaterialProperties/DMMPMetallic.cpp b/Engine/Plugins/Experimental/DynamicMaterial/Source/DynamicMaterialEditor/Private/Components/MaterialProperties/DMMPMetallic.cpp
new file mode 100644
--- /dev/null
+++ b/Engine/Plugins/Experimental/DynamicMaterial/Source/DynamicMaterialEditor/Private/Components/MaterialProperties/DMMPMetallic.cpp
@@ -0,0 +1,34 @@
+// Copyright Epic Games, Inc. All Rights Reserved.
+
+#include "Components/MaterialProperties/DMMPMetallic.h"
+#include "Model/DynamicMaterialModelEditorOnlyData.h"
+
+UDMMaterialPropertyMetallic::UDMMaterialPropertyMetallic()
+	: UDMMaterialProperty(
+		EDMMaterialPropertyType(EDMMaterialPropertyType::Metallic),
+		EDMValueType::VT_Float1)
+{
+}
+
+bool UDMMaterialPropertyMetallic::IsValidForModel(UDynamicMaterialModelEditorOnlyData& InModelEditorOnlyData) const
+{
+	if (InModelEditorOnlyData.GetShadingModel() == EDMMaterialShadingModel::Unlit)
+	{
+		return false;
+	}
+
+	switch (InModelEditorOnlyData.GetBlendMode())
+	{
+		case EBlendMode::BLEND_Opaque:
+		case EBlendMode::BLEND_Masked:
+			return true;
+
+		default:
+			return false;
+	}
+}
+
+UMaterialExpression* UDMMaterialPropertyMetallic::GetDefaultInput(const TSharedRef<FDMMaterialBuildState>& InBuildState) const
+{
+	return CreateConstant(InBuildState, 0.f);
+}
diff --git a/Engine/Plugins/Experimental/DynamicMaterial/Source/DynamicMaterialEditor/Private/Components/MaterialProperties/DMMPNormal.cpp b/Engine/Plugins/Experimental/DynamicMaterial/Source/DynamicMaterialEditor/Private/Components/MaterialProperties/DMMPNormal.cpp
new file mode 100644
--- /dev/null
+++ b/Engine/Plugins/Experimental/DynamicMaterial/Source/DynamicMaterialEditor/Private/Components/MaterialProperties/DMMPNormal.cpp
@@ -0,0 +1,34 @@
+// Copyright Epic Games, Inc. All Rights Reserved.
+
+#include "Components/MaterialProperties/DMMPNormal.h"
+#include "Model/DynamicMaterialModelEditorOnlyData.h"
+
+UDMMaterialPropertyNormal::UDMMaterialPropertyNormal()
+	: UDMMaterialProperty(
+		EDMMaterialPropertyType(EDMMaterialPropertyType::Normal),
+		EDMValueType::VT_Float3_XYZ)
+{
+}
+
+bool UDMMaterialPropertyNormal::IsValidForModel(UDynamicMaterialModelEditorOnlyData& InModelEditorOnlyData) const
+{
+	if (InModelEditorOnlyData.GetShadingModel() == EDMMaterialShadingModel::Unlit)
+	{
+		return false;
+	}
+
+	switch (InModelEditorOnlyData.GetBlendMode())
+	{
+		case EBlendMode::BLEND_Opaque:
+		case EBlendMode::BLEND_Masked:
+			return true;
+
+		default:
+			return false;
+	}
+}
+
+UMaterialExpression* UDMMaterialPropertyNormal::GetDefaultInput(const TSharedRef<FDMMaterialBuildState>& InBuildState) const
+{
+	return CreateConstant(InBuildState, FVector::OneVector);
+}
