[MaterialDesigner]

diff --git a/Engine/Plugins/Experimental/DynamicMaterial/Source/DynamicMaterialEditor/Private/Components/MaterialProperties/DMMPOpacity.cpp b/Engine/Plugins/Experimental/DynamicMaterial/Source/DynamicMaterialEditor/Private/Components/MaterialProperties/DMMPOpacity.cpp
new file mode 100644
--- /dev/null
+++ b/Engine/Plugins/Experimental/DynamicMaterial/Source/DynamicMaterialEditor/Private/Components/MaterialProperties/DMMPOpacity.cpp
@@ -0,0 +1,31 @@
+// Copyright Epic Games, Inc. All Rights Reserved.
+
+#include "Components/MaterialProperties/DMMPOpacity.h"
+#include "Model/DynamicMaterialModelEditorOnlyData.h"
+
+UDMMaterialPropertyOpacity::UDMMaterialPropertyOpacity()
+	: UDMMaterialProperty(
+		EDMMaterialPropertyType(EDMMaterialPropertyType::Opacity),
+		EDMValueType::VT_Float1)
+{
+}
+
+bool UDMMaterialPropertyOpacity::IsValidForModel(UDynamicMaterialModelEditorOnlyData& InModelEditorOnlyData) const
+{
+	switch (InModelEditorOnlyData.GetBlendMode())
+	{
+		case EBlendMode::BLEND_Translucent:
+		case EBlendMode::BLEND_Additive:
+		case EBlendMode::BLEND_AlphaComposite:
+		case EBlendMode::BLEND_AlphaHoldout:
+			return true;
+
+		default:
+			return false;
+	}
+}
+
+UMaterialExpression* UDMMaterialPropertyOpacity::GetDefaultInput(const TSharedRef<FDMMaterialBuildState>& InBuildState) const
+{
+	return CreateConstant(InBuildState, 1.f);
+}
diff --git a/Engine/Plugins/Experimental/DynamicMaterial/Source/DynamicMaterialEditor/Private/Components/MaterialProperties/DMMPPixelDepthOffset.cpp b/Engine/Plugins/Experimental/DynamicMaterial/Source/DynamicMaterialEditor/Private/Components/MaterialProperties/DMMPPixelDepthOffset.cpp
new file mode 100644
--- /dev/null
+++ b/Engine/Plugins/Experimental/DynamicMaterial/Source/DynamicMaterialEditor/Private/Components/MaterialProperties/DMMPPixelDepthOffset.cpp
@@ -0,0 +1,30 @@
+// Copyright Epic Games, Inc. All Rights Reserved.
+
+#include "Components/MaterialProperties/DMMPPixelDepthOffset.h"
+#include "Model/DynamicMaterialModelEditorOnlyData.h"
+
+UDMMaterialPropertyPixelDepthOffset::UDMMaterialPropertyPixelDepthOffset()
+	: UDMMaterialProperty(
+		EDMMaterialPropertyType(EDMMaterialPropertyType::PixelDepthOffset),
+		EDMValueType::VT_Float1)
+{
+}
+
+bool UDMMaterialPropertyPixelDepthOffset::IsValidForModel(UDynamicMaterialModelEditorOnlyData& InModelEditorOnlyData) const
+{
+	switch (InModelEditorOnlyData.GetBlendMode())
+	{
+		case EBlendMode::BLEND_Opaque:
+		case EBlendMode::BLEND_Masked:
+			return true;
+
+		default:
+			return false;
+	}
+}
+
+UMaterialExpression* UDMMaterialPropertyPixelDepthOffset::GetDefaultInput(
+	const TSharedRef<FDMMaterialBuildState>& InBuildState) const
+{
+	return CreateConstant(InBuildState, 0.f);
+}
diff --git a/Engine/Plugins/Experimental/DynamicMaterial/Source/DynamicMaterialEditor/Private/Components/MaterialProperties/DMMPRefraction.cpp b/Engine/Plugins/Experimental/DynamicMaterial/Source/DynamicMaterialEditor/Private/Components/MaterialProperties/DMMPRefraction.cpp
new file mode 100644
--- /dev/null
+++ b/Engine/Plugins/Experimental/DynamicMaterial/Source/DynamicMaterialEditor/Private/Components/MaterialProperties/DMMPRefraction.cpp
@@ -0,0 +1,31 @@
+// Copyright Epic Games, Inc. All Rights Reserved.
+
+#include "Components/MaterialProperties/DMMPRefraction.h"
+#include "Model/DynamicMaterialModelEditorOnlyData.h"
+
+UDMMaterialPropertyRefraction::UDMMaterialPropertyRefraction()
+	: UDMMaterialProperty(
+		EDMMaterialPropertyType(EDMMaterialPropertyType::Refraction),
+		EDMValueType::VT_Float1)
+{
+}
+
+bool UDMMaterialPropertyRefraction::IsValidForModel(UDynamicMaterialModelEditorOnlyData& InModelEditorOnlyData) const
+{
+	switch (InModelEditorOnlyData.GetBlendMode())
+	{
+		case EBlendMode::BLEND_Translucent:
+		case EBlendMode::BLEND_Additive:
+		case EBlendMode::BLEND_AlphaComposite:
+			return true;
+
+		default:
+			return false;
+	}
+}
+
+UMaterialExpression* UDMMaterialPropertyRefraction::GetDefaultInput(
+	const TSharedRef<FDMMaterialBuildState>& InBuildState) const
+{
+	return CreateConstant(InBuildState, 0.f);
+}
diff --git a/Engine/Plugins/Experimental/DynamicMaterial/Source/DynamicMaterialEditor/Private/Components/MaterialProperties/DMMPRoughness.cpp b/Engine/Plugins/Experimental/DynamicMaterial/Source/DynamicMaterialEditor/Private/Components/MaterialProperties/DMMPRoughness.cpp
new file mode 100644
--- /dev/null
+++ b/Engine/Plugins/Experimental/DynamicMaterial/Source/DynamicMaterialEditor/Private/Components/MaterialProperties/DMMPRoughness.cpp
@@ -0,0 +1,34 @@
+// Copyright Epic Games, Inc. All Rights Reserved.
+
+#include "Components/MaterialProperties/DMMPRoughness.h"
+#include "Model/DynamicMaterialModelEditorOnlyData.h"
+
+UDMMaterialPropertyRoughness::UDMMaterialPropertyRoughness()
+	: UDMMaterialProperty(
+		EDMMaterialPropertyType(EDMMaterialPropertyType::Roughness),
+		EDMValueType::VT_Float1)
+{
+}
+
+bool UDMMaterialPropertyRoughness::IsValidForModel(UDynamicMaterialModelEditorOnlyData& InModelEditorOnlyData) const
+{
+	if (InModelEditorOnlyData.GetShadingModel() == EDMMaterialShadingModel::Unlit)
+	{
+		return false;
+	}
+
+	switch (InModelEditorOnlyData.GetBlendMode())
+	{
+		case EBlendMode::BLEND_Opaque:
+		case EBlendMode::BLEND_Masked:
+			return true;
+
+		default:
+			return false;
+	}
+}
+
+UMaterialExpression* UDMMaterialPropertyRoughness::GetDefaultInput(const TSharedRef<FDMMaterialBuildState>& InBuildState) const
+{
+	return CreateConstant(InBuildState, 0.f);
+}
diff --git a/Engine/Plugins/Experimental/DynamicMaterial/Source/DynamicMaterialEditor/Private/Components/MaterialProperties/DMMPSpecular.cpp b/Engine/Plugins/Experimental/DynamicMaterial/Source/DynamicMaterialEditor/Private/Components/MaterialProperties/DMMPSpecular.cpp
new file mode 100644
--- /dev/null
+++ b/Engine/Plugins/Experimental/DynamicMaterial/Source/DynamicMaterialEditor/Private/Components/MaterialProperties/DMMPSpecular.cpp
@@ -0,0 +1,34 @@
+// Copyright Epic Games, Inc. All Rights Reserved.
+
+#include "Components/MaterialProperties/DMMPSpecular.h"
+#include "Model/DynamicMaterialModelEditorOnlyData.h"
+
+UDMMaterialPropertySpecular::UDMMaterialPropertySpecular()
+	: UDMMaterialProperty(
+		EDMMaterialPropertyType(EDMMaterialPropertyType::Specular),
+		EDMValueType::VT_Float3_RGB)
+{
+}
+
+bool UDMMaterialPropertySpecular::IsValidForModel(UDynamicMaterialModelEditorOnlyData& InModelEditorOnlyData) const
+{
+	if (InModelEditorOnlyData.GetShadingModel() == EDMMaterialShadingModel::Unlit)
+	{
+		return false;
+	}
+
+	switch (InModelEditorOnlyData.GetBlendMode())
+	{
+		case EBlendMode::BLEND_Opaque:
+		case EBlendMode::BLEND_Masked:
+			return true;
+
+		default:
+			return false;
+	}
+}
+
+UMaterialExpression* UDMMaterialPropertySpecular::GetDefaultInput(const TSharedRef<FDMMaterialBuildState>& InBuildState) const
+{
+	return CreateConstant(InBuildState, FVector::OneVector);
+}
diff --git a/Engine/Plugins/Experimental/DynamicMaterial/Source/DynamicMaterialEditor/Private/Components/MaterialProperties/DMMPTangent.cpp b/Engine/Plugins/Experimental/DynamicMaterial/Source/DynamicMaterialEditor/Private/Components/MaterialProperties/DMMPTangent.cpp
new file mode 100644
--- /dev/null
+++ b/Engine/Plugins/Experimental/DynamicMaterial/Source/DynamicMaterialEditor/Private/Components/MaterialProperties/DMMPTangent.cpp
@@ -0,0 +1,34 @@
+// Copyright Epic Games, Inc. All Rights Reserved.
+
+#include "Components/MaterialProperties/DMMPTangent.h"
+#include "Model/DynamicMaterialModelEditorOnlyData.h"
+
+UDMMaterialPropertyTangent::UDMMaterialPropertyTangent()
+	: UDMMaterialProperty(
+		EDMMaterialPropertyType(EDMMaterialPropertyType::Tangent),
+		EDMValueType::VT_Float3_XYZ)
+{
+}
+
+bool UDMMaterialPropertyTangent::IsValidForModel(UDynamicMaterialModelEditorOnlyData& InModelEditorOnlyData) const
+{
+	if (InModelEditorOnlyData.GetShadingModel() == EDMMaterialShadingModel::Unlit)
+	{
+		return false;
+	}
+
+	switch (InModelEditorOnlyData.GetBlendMode())
+	{
+		case EBlendMode::BLEND_Opaque:
+		case EBlendMode::BLEND_Masked:
+			return true;
+
+		default:
+			return false;
+	}
+}
+
+UMaterialExpression* UDMMaterialPropertyTangent::GetDefaultInput(const TSharedRef<FDMMaterialBuildState>& InBuildState) const
+{
+	return CreateConstant(InBuildState, FVector::OneVector);
+}
diff --git a/Engine/Plugins/Experimental/DynamicMaterial/Source/DynamicMaterialEditor/Private/Components/MaterialStageBlends/DMMSBContrastBase.cpp b/Engine/Plugins/Experimental/DynamicMaterial/Source/DynamicMaterialEditor/Private/Components/MaterialStageBlends/DMMSBContrastBase.cpp
new file mode 100644
--- /dev/null
+++ b/Engine/Plugins/Experimental/DynamicMaterial/Source/DynamicMaterialEditor/Private/Components/MaterialStageBlends/DMMSBContrastBase.cpp
@@ -0,0 +1,17 @@
+// Copyright Epic Games, Inc. All Rights Reserved.
+
+#include "Components/MaterialStageBlends/DMMSBContrastBase.h"
+
+#define LOCTEXT_NAMESPACE "DMMaterialStageBlendContrastBase"
+
+UDMMaterialStageBlendContrastBase::UDMMaterialStageBlendContrastBase()
+	: UDMMaterialStageBlendContrastBase(LOCTEXT("BlendContrastBase", "Contrast Base"))
+{
+}
+
+UDMMaterialStageBlendContrastBase::UDMMaterialStageBlendContrastBase(const FText& InName)
+	: UDMMaterialStageBlend(InName)
+{
+}
+
+#undef LOCTEXT_NAMESPACE
diff --git a/Engine/Plugins/Experimental/DynamicMaterial/Source/DynamicMaterialEditor/Private/Components/MaterialStageExpressions/DMMSETextureSampleEdgeColor.cpp b/Engine/Plugins/Experimental/DynamicMaterial/Source/DynamicMaterialEditor/Private/Components/MaterialStageExpressions/DMMSETextureSampleEdgeColor.cpp
new file mode 100644
--- /dev/null
+++ b/Engine/Plugins/Experimental/DynamicMaterial/Source/DynamicMaterialEditor/Private/Components/MaterialStageExpressions/DMMSETextureSampleEdgeColor.cpp
@@ -0,0 +1,166 @@
+// Copyright Epic Games, Inc. All Rights Reserved.
+
+#include "Components/MaterialStageExpressions/DMMSETextureSampleEdgeColor.h"
+#include "Components/DMMaterialStage.h"
+#include "Components/DMMaterialSubStage.h"
+#include "Components/MaterialStageExpressions/DMMSETextureSample.h"
+#include "Components/MaterialStageInputs/DMMSIExpression.h"
+#include "Components/MaterialStageInputs/DMMSIValue.h"
+#include "Components/MaterialValues/DMMaterialValueFloat2.h"
+#include "DMPrivate.h"
+#include "Materials/MaterialExpressionTextureSample.h"
+
+#define LOCTEXT_NAMESPACE "DMMaterialStageExpressionTextureSample"
+
+namespace UE::DynamicMaterialEditor::TextureSampleEdgeColor::Private
+{
+	constexpr float EdgeGap = 0.025;
+}
+
+UDMMaterialStageExpressionTextureSampleEdgeColor::UDMMaterialStageExpressionTextureSampleEdgeColor()
+	: UDMMaterialStageExpression(
+		LOCTEXT("EdgeColor", "Edge Color"),
+		UMaterialExpressionTextureSample::StaticClass()
+	)
+{
+	bInputRequired = true;
+	bAllowNestedInputs = true;
+
+	InputConnectors.Add({1, LOCTEXT("Texture", "Texture"), EDMValueType::VT_Texture});
+	InputConnectors.Add({0, LOCTEXT("Location", "Location"), EDMValueType::VT_Float2});
+
+	using namespace UE::DynamicMaterialEditor::TextureSampleEdgeColor::Private;
+
+	EditableProperties.Add(GET_MEMBER_NAME_CHECKED(UDMMaterialStageExpressionTextureSampleEdgeColor, EdgeLocation));
+	EdgeLocation = EDMEdgeLocation::TopLeft;
+	
+	OutputConnectors.Add({0, LOCTEXT("ColorRGB", "Color (RGB)"), EDMValueType::VT_Float3_RGB});
+}
+
+bool UDMMaterialStageExpressionTextureSampleEdgeColor::CanChangeInputType(int32 InInputIndex) const
+{
+	return false;
+}
+
+bool UDMMaterialStageExpressionTextureSampleEdgeColor::IsInputVisible(int32 InInputIndex) const
+{
+	if (InInputIndex == 1 && EdgeLocation != EDMEdgeLocation::Custom)
+	{
+		return false;
+	}
+
+	return Super::IsInputVisible(InInputIndex);
+}
+
+void UDMMaterialStageExpressionTextureSampleEdgeColor::AddDefaultInput(int32 InInputIndex) const
+{
+	UDMMaterialStage* Stage = GetStage();
+	check(Stage);
+
+	if (InInputIndex == 1)
+	{
+		using namespace UE::DynamicMaterialEditor::TextureSampleEdgeColor::Private;
+
+		UDMMaterialStageInputValue* InputValue = Stage->ChangeInput_NewLocalValue(1, FDMMaterialStageConnectorChannel::WHOLE_CHANNEL, EDMValueType::VT_Float2, FDMMaterialStageConnectorChannel::WHOLE_CHANNEL);
+		check(InputValue);
+
+		UDMMaterialValueFloat2* InputFloat2 = Cast<UDMMaterialValueFloat2>(InputValue->GetValue());
+		check(InputFloat2);
+
+		InputFloat2->SetDefaultValue(FVector2D(EdgeGap, EdgeGap));
+		InputFloat2->ApplyDefaultValue();
+		return;
+	}
+
+	return Super::AddDefaultInput(InInputIndex);
+}
+
+void UDMMaterialStageExpressionTextureSampleEdgeColor::PostEditChangeProperty(struct FPropertyChangedEvent& PropertyChangedEvent)
+{
+	Super::PostEditChangeProperty(PropertyChangedEvent);
+
+	static const FName EdgeLocationName = GET_MEMBER_NAME_CHECKED(UDMMaterialStageExpressionTextureSampleEdgeColor, EdgeLocation);
+
+	if (PropertyChangedEvent.MemberProperty && PropertyChangedEvent.MemberProperty->GetFName() == EdgeLocationName)
+	{
+		OnEdgeLocationChanged();
+	}
+}
+
+void UDMMaterialStageExpressionTextureSampleEdgeColor::OnEdgeLocationChanged()
+{
+	if (EdgeLocation == EDMEdgeLocation::Custom)
+	{
+		return;
+	}
+
+	using namespace UE::DynamicMaterialEditor::TextureSampleEdgeColor::Private;
+
+	UDMMaterialStage* Stage = GetStage();
+	check(Stage);
+	check(Stage->GetInputs().IsValidIndex(1));
+
+	UDMMaterialStageInputValue* InputValue = Cast<UDMMaterialStageInputValue>(Stage->GetInputs()[1]);
+	check(InputValue);
+	
+	UDMMaterialValueFloat2* Float2 = Cast<UDMMaterialValueFloat2>(InputValue->GetValue());
+	check(Float2);
+
+	FVector2D Float2Value = FVector2D::ZeroVector;
+	
+	switch (EdgeLocation)
+	{
+		case EDMEdgeLocation::TopLeft:
+			Float2Value.X = EdgeGap;
+			Float2Value.Y = EdgeGap;
+			break;
+
+		case EDMEdgeLocation::Top:
+			Float2Value.X = 0.5;
+			Float2Value.Y = EdgeGap;
+			break;
+
+		case EDMEdgeLocation::TopRight:
+			Float2Value.X = 1.f - EdgeGap;
+			Float2Value.Y = EdgeGap;
+			break;
+
+		case EDMEdgeLocation::Left:
+			Float2Value.X = EdgeGap;
+			Float2Value.Y = 0.5;
+			break;
+
+		case EDMEdgeLocation::Center:
+			Float2Value.X = 0.5;
+			Float2Value.Y = 0.5;
+			break;
+
+		case EDMEdgeLocation::Right:
+			Float2Value.X = 1.f - EdgeGap;
+			Float2Value.Y = 0.5;
+			break;
+
+		case EDMEdgeLocation::BottomLeft:
+			Float2Value.X = EdgeGap;
+			Float2Value.Y = 1.f - EdgeGap;
+			break;
+
+		case EDMEdgeLocation::Bottom:
+			Float2Value.X = 0.5;
+			Float2Value.Y = 1.f - EdgeGap;
+			break;
+
+		case EDMEdgeLocation::BottomRight:
+			Float2Value.X = 1.f - EdgeGap;
+			Float2Value.Y = 1.f - EdgeGap;
+			break;
+
+		default:
+			checkNoEntry();
+			break;
+	}
+
+	Float2->SetValue(Float2Value);
+}
+
+#undef LOCTEXT_NAMESPACE
diff --git a/Engine/Plugins/Experimental/DynamicMaterial/Source/DynamicMaterialEditor/Private/Components/MaterialStageGradients/DMMSGLinear.cpp b/Engine/Plugins/Experimental/DynamicMaterial/Source/DynamicMaterialEditor/Private/Components/MaterialStageGradients/DMMSGLinear.cpp
new file mode 100644
--- /dev/null
+++ b/Engine/Plugins/Experimental/DynamicMaterial/Source/DynamicMaterialEditor/Private/Components/MaterialStageGradients/DMMSGLinear.cpp
@@ -0,0 +1,122 @@
+// Copyright Epic Games, Inc. All Rights Reserved.
+
+#include "Components/MaterialStageGradients/DMMSGLinear.h"
+#include "DMDefs.h"
+#include "DMMaterialFunctionLibrary.h"
+#include "Materials/MaterialExpressionComponentMask.h"
+#include "Materials/MaterialExpressionConstant.h"
+#include "Materials/MaterialExpressionFmod.h"
+#include "Materials/MaterialExpressionIf.h"
+#include "Materials/MaterialExpressionMaterialFunctionCall.h"
+#include "Materials/MaterialExpressionSubtract.h"
+#include "Model/DMMaterialBuildState.h"
+#include "Model/DMMaterialBuildUtils.h"
+
+#define LOCTEXT_NAMESPACE "DMMaterialStageGradientLinear"
+
+UDMMaterialStageGradientLinear::UDMMaterialStageGradientLinear()
+	: UDMMaterialStageGradient(LOCTEXT("GradientLinear", "Linear Gradient"))
+{
+	EditableProperties.Add(GET_MEMBER_NAME_CHECKED(UDMMaterialStageGradientLinear, Tiling));
+}
+
+void UDMMaterialStageGradientLinear::GenerateExpressions(const TSharedRef<FDMMaterialBuildState>& InBuildState) const
+{
+	if (!IsComponentValid() || !IsComponentAdded())
+	{
+		return;
+	}
+
+	if (InBuildState->HasStageSource(this))
+	{
+		return;
+	}
+
+	UMaterialExpressionComponentMask* Mask = InBuildState->GetBuildUtils().CreateExpression<UMaterialExpressionComponentMask>(UE_DM_NodeComment_Default);
+
+	Mask->R = 1;
+	Mask->G = 0;
+	Mask->B = 0;
+	Mask->A = 0;
+
+	TArray<UMaterialExpression*> OutputExpressions;
+	OutputExpressions.Add(Mask);
+
+	UMaterialExpression* Output = nullptr;
+
+	switch (Tiling)
+	{
+		case ELinearGradientTileType::NoTile:
+			Output = Mask;
+			break;
+
+		case ELinearGradientTileType::Tile:
+		{
+			UMaterialExpressionFmod* Fmod = InBuildState->GetBuildUtils().CreateExpression<UMaterialExpressionFmod>(UE_DM_NodeComment_Default);
+			UMaterialExpressionConstant* Divisor = InBuildState->GetBuildUtils().CreateExpression<UMaterialExpressionConstant>(UE_DM_NodeComment_Default);
+
+			Mask->ConnectExpression(&Fmod->A, 0);
+
+			Divisor->R = 1.f;
+			Divisor->ConnectExpression(&Fmod->B, 0);
+
+			Output = Fmod;
+			OutputExpressions.Append({Divisor, Fmod});
+			break;
+		}
+
+		case ELinearGradientTileType::TileAndMirror:
+		{
+			UMaterialExpressionFmod* Fmod = InBuildState->GetBuildUtils().CreateExpression<UMaterialExpressionFmod>(UE_DM_NodeComment_Default);
+			UMaterialExpressionConstant* Divisor = InBuildState->GetBuildUtils().CreateExpression<UMaterialExpressionConstant>(UE_DM_NodeComment_Default);
+			UMaterialExpressionSubtract* Subtract = InBuildState->GetBuildUtils().CreateExpression<UMaterialExpressionSubtract>(UE_DM_NodeComment_Default);
+			UMaterialExpressionIf* IfExpr = InBuildState->GetBuildUtils().CreateExpression<UMaterialExpressionIf>(UE_DM_NodeComment_Default);
+
+			Mask->ConnectExpression(&Fmod->A, 0);
+
+			Divisor->R = 2.f;
+			Divisor->ConnectExpression(&Fmod->B, 0);
+
+			Fmod->ConnectExpression(&Subtract->B, 0);
+			Fmod->ConnectExpression(&IfExpr->A, 0);
+			Fmod->ConnectExpression(&IfExpr->AEqualsB, 0);
+			Fmod->ConnectExpression(&IfExpr->ALessThanB, 0);
+
+			Subtract->ConstA = 2.f;
+			Subtract->ConnectExpression(&IfExpr->AGreaterThanB, 0);
+
+			IfExpr->ConstB = 1.f;
+			IfExpr->EqualsThreshold = 0.f;
+
+			Output = IfExpr;
+			OutputExpressions.Append({Divisor, Fmod, Subtract, IfExpr});
+			break;
+		}
+
+		default:
+			checkNoEntry();
+			break;
+	}
+
+	UMaterialExpression* MakeFloat = FDMMaterialFunctionLibrary::Get().GetMakeFloat3(InBuildState->GetDynamicMaterial(), UE_DM_NodeComment_Default);
+	Output->ConnectExpression(MakeFloat->GetInput(0), 0);
+	Output->ConnectExpression(MakeFloat->GetInput(1), 0);
+	Output->ConnectExpression(MakeFloat->GetInput(2), 0);
+
+	OutputExpressions.Add(MakeFloat);	
+	InBuildState->AddStageSourceExpressions(this, OutputExpressions);
+}
+
+void UDMMaterialStageGradientLinear::SetTilingType(ELinearGradientTileType InType)
+{
+	if (Tiling == InType)
+	{
+		return;
+	}
+
+	Tiling = InType;
+
+	Update(EDMUpdateType::Structure);
+}
+
+#undef LOCTEXT_NAMESPACE
diff --git a/Engine/Plugins/Experimental/DynamicMaterial/Source/DynamicMaterialEditor/Private/Components/MaterialStageGradients/DMMSGRadial.cpp b/Engine/Plugins/Experimental/DynamicMaterial/Source/DynamicMaterialEditor/Private/Components/MaterialStageGradients/DMMSGRadial.cpp
new file mode 100644
--- /dev/null
+++ b/Engine/Plugins/Experimental/DynamicMaterial/Source/DynamicMaterialEditor/Private/Components/MaterialStageGradients/DMMSGRadial.cpp
@@ -0,0 +1,38 @@
+// Copyright Epic Games, Inc. All Rights Reserved.
+
+#include "Components/MaterialStageGradients/DMMSGRadial.h"
+#include "DMMaterialFunctionLibrary.h"
+#include "Materials/MaterialExpressionMaterialFunctionCall.h"
+#include "Model/DMMaterialBuildState.h"
+#include "Model/DMMaterialBuildUtils.h"
+
+#define LOCTEXT_NAMESPACE "DMMaterialStageGradientRadial"
+
+UDMMaterialStageGradientRadial::UDMMaterialStageGradientRadial()
+	: UDMMaterialStageGradient(LOCTEXT("GradientRadial", "Radial Gradient"))
+{
+}
+
+void UDMMaterialStageGradientRadial::GenerateExpressions(const TSharedRef<FDMMaterialBuildState>& InBuildState) const
+{
+	if (!IsComponentValid() || !IsComponentAdded())
+	{
+		return;
+	}
+
+	if (InBuildState->HasStageSource(this))
+	{
+		return;
+	}
+
+	UMaterialExpressionMaterialFunctionCall* Radial = FDMMaterialFunctionLibrary::Get().GetRadialGradientExponential(InBuildState->GetDynamicMaterial(), UE_DM_NodeComment_Default);
+	UMaterialExpression* MakeFloat = FDMMaterialFunctionLibrary::Get().GetMakeFloat3(InBuildState->GetDynamicMaterial(), UE_DM_NodeComment_Default);
+
+	Radial->ConnectExpression(MakeFloat->GetInput(0), 0);
+	Radial->ConnectExpression(MakeFloat->GetInput(1), 0);
+	Radial->ConnectExpression(MakeFloat->GetInput(2), 0);
+
+	InBuildState->AddStageSourceExpressions(this, {Radial, MakeFloat});
+}
+
+#undef LOCTEXT_NAMESPACE
diff --git a/Engine/Plugins/Experimental/DynamicMaterial/Source/DynamicMaterialEditor/Private/Customizations/DMStageEditDetailsCustomization.h b/Engine/Plugins/Experimental/DynamicMaterial/Source/DynamicMaterialEditor/Private/Customizations/DMStageEditDetailsCustomization.h
new file mode 100644
--- /dev/null
+++ b/Engine/Plugins/Experimental/DynamicMaterial/Source/DynamicMaterialEditor/Private/Customizations/DMStageEditDetailsCustomization.h
@@ -0,0 +1,26 @@
+// Copyright Epic Games, Inc. All Rights Reserved.
+
+#pragma once
+
+#include "PropertyEditorModule.h"
+#include "IPropertyTypeCustomization.h"
+
+class FDetailWidgetRow;
+class IDetailChildrenBuilder;
+
+/**
+ * Material Designer Stage Edit Details Customization
+ */
+class FDMStageEditDetailsCustomization : public IPropertyTypeCustomization
+{
+public:
+	static TSharedRef<IPropertyTypeCustomization> MakeInstance();
+
+	//~ Begin IPropertyTypeCustomization
+	virtual void CustomizeHeader(TSharedRef<IPropertyHandle> InStructPropertyHandle, FDetailWidgetRow& HeaderRow, IPropertyTypeCustomizationUtils& StructCustomizationUtils) override;
+	virtual void CustomizeChildren(TSharedRef<IPropertyHandle> InStructPropertyHandle, IDetailChildrenBuilder& StructBuilder, IPropertyTypeCustomizationUtils& StructCustomizationUtils) override {}
+	//~ End IPropertyTypeCustomization
+
+private:
+	TSharedPtr<IPropertyHandle> StructPropertyHandle;
+};
diff --git a/Engine/Plugins/Experimental/DynamicMaterial/Source/DynamicMaterialEditor/Private/DMBlueprintFunctionLibrary.cpp b/Engine/Plugins/Experimental/DynamicMaterial/Source/DynamicMaterialEditor/Private/DMBlueprintFunctionLibrary.cpp
new file mode 100644
--- /dev/null
+++ b/Engine/Plugins/Experimental/DynamicMaterial/Source/DynamicMaterialEditor/Private/DMBlueprintFunctionLibrary.cpp
@@ -0,0 +1,481 @@
+// Copyright Epic Games, Inc. All Rights Reserved.
+
+#include "DMBlueprintFunctionLibrary.h"
+#include "AssetRegistry/AssetRegistryModule.h"
+#include "Components/DMMaterialLayer.h"
+#include "Components/DMMaterialSlot.h"
+#include "Components/DMMaterialStage.h"
+#include "Components/DMMaterialStageBlend.h"
+#include "Components/DMMaterialStageFunction.h"
+#include "Components/DMMaterialStageGradient.h"
+#include "Components/DMMaterialStageSource.h"
+#include "Components/DMMaterialStageThroughput.h"
+#include "Components/DMMaterialStageThroughputLayerBlend.h"
+#include "Components/DMMaterialSubStage.h"
+#include "Components/MaterialStageBlends/DMMSBNormal.h"
+#include "Components/MaterialStageExpressions/DMMSETextureSample.h"
+#include "Components/MaterialStageInputs/DMMSIExpression.h"
+#include "Components/MaterialStageInputs/DMMSITextureUV.h"
+#include "Components/MaterialStageInputs/DMMSIValue.h"
+#include "Components/MaterialValues/DMMaterialValueTexture.h"
+#include "Components/PrimitiveComponent.h"
+#include "DMEDefs.h"
+#include "DMObjectMaterialProperty.h"
+#include "DMWorldSubsystem.h"
+#include "DynamicMaterialEditorModule.h"
+#include "Engine/Texture.h"
+#include "Engine/World.h"
+#include "GameFramework/Actor.h"
+#include "Material/DynamicMaterialInstance.h"
+#include "Material/DynamicMaterialInstanceFactory.h"
+#include "Materials/Material.h"
+#include "Model/DynamicMaterialModel.h"
+#include "Model/DynamicMaterialModelEditorOnlyData.h"
+#include "Model/IDynamicMaterialModelEditorOnlyDataInterface.h"
+#include "UObject/Package.h"
+#include "UObject/UObjectGlobals.h"
+
+UDMMaterialStageInputValue* UDMBlueprintFunctionLibrary::FindDefaultStageOpacityInputValue(UDMMaterialStage* InStage)
+{
+	if (!IsValid(InStage))
+	{
+		return nullptr;
+	}
+
+	int32 OpacityInputIndex = -1;
+
+	if (UDMMaterialStageThroughput* const ThroughputSource = Cast<UDMMaterialStageThroughput>(InStage->GetSource()))
+	{
+		const TArray<FDMMaterialStageConnector>& Connectors = ThroughputSource->GetInputConnectors();
+		for (const FDMMaterialStageConnector& Connector : Connectors)
+		{
+			if (Connector.Name.ToString() == "Opacity")
+			{
+				OpacityInputIndex = Connector.Index;
+			}
+		}
+	}
+
+	const TArray<UDMMaterialStageInput*>& StageInputs = InStage->GetInputs();
+	const TArray<FDMMaterialStageConnection>& InputConnectionMap = InStage->GetInputConnectionMap();
+
+	if (InputConnectionMap.IsValidIndex(OpacityInputIndex) && InputConnectionMap[OpacityInputIndex].Channels.Num() > 0)
+	{
+		if (InputConnectionMap[OpacityInputIndex].Channels[0].SourceIndex >= FDMMaterialStageConnectorChannel::FIRST_STAGE_INPUT)
+		{
+			const int32 StageInputIndex = InputConnectionMap[OpacityInputIndex].Channels[0].SourceIndex - FDMMaterialStageConnectorChannel::FIRST_STAGE_INPUT;
+
+			return Cast<UDMMaterialStageInputValue>(StageInputs[StageInputIndex]);
+		}
+	}
+
+	return nullptr;
+}
+
+void UDMBlueprintFunctionLibrary::SetDefaultStageSourceTexture(UDMMaterialStage* InStage, UTexture* InTexture)
+{
+	if (!IsValid(InStage) || !IsValid(InTexture))
+	{
+		return;
+	}
+
+	UDMMaterialStageSource* StageSource = InStage->GetSource();
+	if (!StageSource)
+	{
+		return;
+	}
+
+	auto SetTextureValue = [InTexture](UDMMaterialStageInputValue* NewInputValue)
+	{
+		if (UDMMaterialValueTexture* TextureValue = Cast<UDMMaterialValueTexture>(NewInputValue->GetValue()))
+		{
+			TextureValue->SetValue(InTexture);
+		}
+	};
+
+	if (GUndo)
+	{
+		InStage->Modify();
+	}
+
+	if (UDMMaterialStageBlend* const Blend = Cast<UDMMaterialStageBlend>(StageSource))
+	{
+		UDMMaterialStageInputExpression* NewInput = InStage->ChangeInput_Expression(UDMMaterialStageBlend::InputB, FDMMaterialStageConnectorChannel::WHOLE_CHANNEL,
+			UDMMaterialStageExpressionTextureSample::StaticClass(), 0, FDMMaterialStageConnectorChannel::WHOLE_CHANNEL);
+
+		UDMMaterialSubStage* SubStage = NewInput->GetSubStage();
+		check(SubStage);
+		UDMMaterialStageInputValue* NewInputValue = SubStage->ChangeInput_NewLocalValue(0, FDMMaterialStageConnectorChannel::WHOLE_CHANNEL, EDMValueType::VT_Texture, FDMMaterialStageConnectorChannel::WHOLE_CHANNEL);
+		SetTextureValue(NewInputValue);
+
+		return;
+	}
+
+	if (UDMMaterialStageThroughputLayerBlend* const LayerBlend = Cast<UDMMaterialStageThroughputLayerBlend>(StageSource))
+	{
+		const bool bHasAlpha = UE::DynamicMaterial::Private::HasAlpha(InTexture);
+
+		UDMMaterialStageInputExpression* NewInput = InStage->ChangeInput_Expression(2, FDMMaterialStageConnectorChannel::WHOLE_CHANNEL,
+			UDMMaterialStageExpressionTextureSample::StaticClass(), bHasAlpha ? 1 : 0, FDMMaterialStageConnectorChannel::WHOLE_CHANNEL);
+
+		UDMMaterialSubStage* SubStage = NewInput->GetSubStage();
+		check(IsValid(SubStage));
+		UDMMaterialStageInputValue* NewInputValue = SubStage->ChangeInput_NewLocalValue(0, FDMMaterialStageConnectorChannel::WHOLE_CHANNEL, EDMValueType::VT_Texture, FDMMaterialStageConnectorChannel::WHOLE_CHANNEL);
+		SetTextureValue(NewInputValue);
+
+		return;
+	}
+
+	UDMMaterialStageExpression* NewExpression = InStage->ChangeSource_Expression(UDMMaterialStageExpressionTextureSample::StaticClass());
+	UDMMaterialStageInputValue* NewInputValue = InStage->ChangeInput_NewLocalValue(0, FDMMaterialStageConnectorChannel::WHOLE_CHANNEL, EDMValueType::VT_Texture, FDMMaterialStageConnectorChannel::WHOLE_CHANNEL);
+	SetTextureValue(NewInputValue);
+}
+
+TArray<FDMObjectMaterialProperty> UDMBlueprintFunctionLibrary::GetActorDynamicMaterialSlots(AActor* InActor, const bool bInIncludeEmptySlots)
+{
+	TArray<FDMObjectMaterialProperty> OutSlots;
+
+	if (!IsValid(InActor))
+	{
+		return OutSlots;
+	}
+
+	InActor->ForEachComponent<UPrimitiveComponent>(false, [bInIncludeEmptySlots, &OutSlots](UPrimitiveComponent* InPrimComp)
+		{
+			const int32 PrimitiveMaterialCount = InPrimComp->GetNumMaterials();
+
+			for (int32 MaterialIndex = 0; MaterialIndex < PrimitiveMaterialCount; ++MaterialIndex)
+			{
+				if (UDynamicMaterialInstance* const MDI = Cast<UDynamicMaterialInstance>(InPrimComp->GetMaterial(MaterialIndex)))
+				{
+					if (UDynamicMaterialModel* const ThisMaterialModel = MDI->GetMaterialModel())
+					{
+						OutSlots.Add(FDMObjectMaterialProperty(InPrimComp, MaterialIndex));
+						continue;
+					}
+				}
+
+				if (bInIncludeEmptySlots)
+				{
+					OutSlots.Add(FDMObjectMaterialProperty(InPrimComp, MaterialIndex));
+				}
+			}
+		});
+
+	return OutSlots;
+}
+
+TArray<FDMObjectMaterialProperty> UDMBlueprintFunctionLibrary::GetActorMaterialProperties(AActor* InActor)
+{
+	TArray<FDMObjectMaterialProperty> ActorProperties;
+
+	if (!IsValid(InActor))
+	{
+		return ActorProperties;
+	}
+
+	FDMGetObjectMaterialPropertiesDelegate PropertyGenerator = FDynamicMaterialEditorModule::GetCustomMaterialPropertyGenerator(InActor->GetClass());
+
+	if (PropertyGenerator.IsBound())
+	{
+		ActorProperties = PropertyGenerator.Execute(InActor);
+
+		if (!ActorProperties.IsEmpty())
+		{
+			return ActorProperties;
+		}
+	}
+
+	InActor->ForEachComponent<UPrimitiveComponent>(false, [&ActorProperties](UPrimitiveComponent* InComp)
+		{
+			for (int32 MaterialIdx = 0; MaterialIdx < InComp->GetNumMaterials(); ++MaterialIdx)
+			{
+				ActorProperties.Add({InComp, MaterialIdx});
+			}
+		});
+
+	auto IterateProperies = [&ActorProperties](UObject* Outer)
+	{
+		if (!Outer || !Outer->GetClass())
+		{
+			return;
+		}
+
+		for (FProperty* OuterProperty : TFieldRange<FProperty>(Outer->GetClass()))
+		{
+			const EPropertyFlags Flags = OuterProperty->GetPropertyFlags();
+			const bool bEditable = !!(Flags & CPF_Edit);
+			const bool bEditConst = !!(Flags & CPF_EditConst);
+			const bool bNoEditInstance = !!(Flags & CPF_DisableEditOnInstance);
+
+			if (!bEditable || bEditConst || bNoEditInstance)
+			{
+				continue;;
+			}
+
+			if (FObjectPropertyBase* ObjectProperty = CastField<FObjectPropertyBase>(OuterProperty))
+			{
+				if (ObjectProperty->PropertyClass && ObjectProperty->PropertyClass->IsChildOf(UMaterialInterface::StaticClass()))
+				{
+					ActorProperties.Add(FDMObjectMaterialProperty(Outer, ObjectProperty));
+				}
+			}
+			else if (FArrayProperty* ArrayProperty = CastField<FArrayProperty>(OuterProperty))
+			{
+				if (FObjectPropertyBase* ArrayObjectProperty = CastField<FObjectPropertyBase>(ArrayProperty->Inner))
+				{
+					if (ArrayObjectProperty->PropertyClass && ArrayObjectProperty->PropertyClass->IsChildOf(UMaterialInterface::StaticClass()))
+					{
+						FScriptArrayHelper ArrayHelper(ArrayProperty, ArrayProperty->ContainerPtrToValuePtr<void>(Outer));
+
+						for (int Idx = 0, Count = ArrayHelper.Num(); Idx < Count; ++Idx)
+						{
+							ActorProperties.Add(FDMObjectMaterialProperty(Outer, OuterProperty, Idx));
+						}
+					}
+				}
+			}
+		}
+	};
+
+	IterateProperies(InActor);
+
+	InActor->ForEachComponent<UActorComponent>(false, [&IterateProperies](UActorComponent* InComp)
+		{
+			IterateProperies(InComp);
+		});
+
+	return ActorProperties;
+}
+
+UDynamicMaterialModel* UDMBlueprintFunctionLibrary::FindFirstValidActorMaterialSlot(const TArray<AActor*>& InActors)
+{
+	// Loop through selected actors to find first one with a valid material model on a primitive component.
+	for (AActor* Actor : InActors)
+	{
+		if (!IsValid(Actor))
+		{
+			continue;
+		}
+
+		const TArray<FDMObjectMaterialProperty> ActorMaterialProperties = GetActorMaterialProperties(Actor);
+		for (const FDMObjectMaterialProperty& ActorMaterialProperty : ActorMaterialProperties)
+		{
+			if (UDynamicMaterialModel* const ThisMaterialModel = ActorMaterialProperty.GetMaterialModel())
+			{
+				return ThisMaterialModel;
+			}
+		}
+	}
+
+	return nullptr;
+}
+
+FDMObjectMaterialProperty UDMBlueprintFunctionLibrary::MakePrimitiveComponentSlot(UPrimitiveComponent* InComponent, int32 InSlot)
+{
+	return FDMObjectMaterialProperty(InComponent, InSlot);
+}
+
+FDMObjectMaterialProperty UDMBlueprintFunctionLibrary::MakeObjectMaterialProperty(UObject* InObject, FName PropertyName)
+{
+	FProperty* Property = nullptr;
+
+	if (IsValid(InObject))
+	{
+		Property = InObject->GetClass()->FindPropertyByName(PropertyName);
+	}
+
+	return FDMObjectMaterialProperty(InObject, Property, INDEX_NONE);
+}
+
+FDMObjectMaterialProperty UDMBlueprintFunctionLibrary::MakeObjectMaterialPropertyArray(UObject* InObject, FName PropertyName, int32 ArrayIndex)
+{
+	FProperty* Property = nullptr;
+
+	if (IsValid(InObject))
+	{
+		Property = InObject->GetClass()->FindPropertyByName(PropertyName);
+	}
+
+	return FDMObjectMaterialProperty(InObject, Property, ArrayIndex);
+}
+
+UDynamicMaterialModel* UDMBlueprintFunctionLibrary::CreateDynamicMaterialInObject(FDMObjectMaterialProperty& InMaterialProperty)
+{
+	if (!InMaterialProperty.IsValid())
+	{
+		return nullptr;
+	}
+
+	UObject* const Outer = InMaterialProperty.OuterWeak.Get();
+
+	UDynamicMaterialInstanceFactory* const InstanceFactory = NewObject<UDynamicMaterialInstanceFactory>();
+	check(InstanceFactory);
+
+	UDynamicMaterialInstance* const NewInstance = Cast<UDynamicMaterialInstance>(InstanceFactory->FactoryCreateNew(UDynamicMaterialInstance::StaticClass(),
+		Outer, NAME_None, RF_Transactional, nullptr, GWarn));
+	check(NewInstance);
+
+	bool bSubsystemTakenOver = false;
+
+	if (const UWorld* const World = Outer->GetWorld())
+	{
+		if (IsValid(World))
+		{
+			UDMWorldSubsystem* const WorldSubsystem = World->GetSubsystem<UDMWorldSubsystem>();
+
+			if (IsValid(WorldSubsystem))
+			{
+				if (WorldSubsystem->GetMaterialValueSetterDelegate().IsBound())
+				{
+					bSubsystemTakenOver = WorldSubsystem->GetMaterialValueSetterDelegate().Execute(InMaterialProperty, NewInstance);
+				}
+			}
+		}
+	}
+
+	if (!bSubsystemTakenOver)
+	{
+		InMaterialProperty.SetMaterial(NewInstance);
+	}
+
+	return NewInstance->GetMaterialModel();
+}
+
+UDynamicMaterialModel* UDMBlueprintFunctionLibrary::GetObjectPropertyMaterialModel(const FDMObjectMaterialProperty& InMaterialProperty) const
+{
+	return InMaterialProperty.GetMaterialModel();
+}
+
+UDynamicMaterialInstance* UDMBlueprintFunctionLibrary::GetObjectPropertyMaterial(const FDMObjectMaterialProperty& InMaterialProperty) const
+{
+	return InMaterialProperty.GetMaterial();
+}
+
+void UDMBlueprintFunctionLibrary::SetObjectPropertyMaterial(FDMObjectMaterialProperty& InMaterialProperty, UDynamicMaterialInstance* InDynamicMaterial)
+{
+	InMaterialProperty.SetMaterial(InDynamicMaterial);
+}
+
+bool UDMBlueprintFunctionLibrary::IsObjectPropertyValid(const FDMObjectMaterialProperty& InMaterialProperty) const
+{
+	return InMaterialProperty.IsValid();
+}
+
+FText UDMBlueprintFunctionLibrary::GetObjectPropertyName(const FDMObjectMaterialProperty& InMaterialProperty, bool bInIgnoreNewStatus) const
+{
+	return InMaterialProperty.GetPropertyName(bInIgnoreNewStatus);
+}
+
+bool UDMBlueprintFunctionLibrary::ExportMaterialInstance(UDynamicMaterialModel* InMaterialModel, const FString& InSavePath)
+{
+	if (!IsValid(InMaterialModel))
+	{
+		UE_LOG(LogDynamicMaterialEditor, Error, TEXT("Invalid material to export."));
+		return false;
+	}
+
+	if (InSavePath.IsEmpty())
+	{
+		UE_LOG(LogDynamicMaterialEditor, Error, TEXT("Invalid material save path to export."));
+		return false;
+	}
+
+	UDynamicMaterialInstance* Instance = InMaterialModel->GetDynamicMaterialInstance();
+	if (!Instance)
+	{
+		UE_LOG(LogDynamicMaterialEditor, Error, TEXT("Failed to find a Material Designer Instance to export."));
+		return false;
+	}
+
+	const FString PackagePath = FPaths::GetBaseFilename(*InSavePath, false);
+
+	UPackage* Package = CreatePackage(*PackagePath);
+	if (!Package)
+	{
+		UE_LOG(LogDynamicMaterialEditor, Error, TEXT("Failed to create package for Material Designer Instance (%s)."), *PackagePath);
+		return false;
+	}
+
+	const FString AssetName = FPaths::GetBaseFilename(*InSavePath, true);
+
+	FObjectDuplicationParameters Params = InitStaticDuplicateObjectParams(Instance, Package, *AssetName,
+		RF_Public | RF_Standalone, nullptr, EDuplicateMode::Normal, EInternalObjectFlags::None);
+
+	UObject* NewAsset = StaticDuplicateObjectEx(Params);
+	if (!NewAsset)
+	{
+		UE_LOG(LogDynamicMaterialEditor, Error, TEXT("Failed to create new Material Designer Instance asset."));
+		return false;
+	}
+
+	// Not sure why these need to be set again!
+	NewAsset->SetFlags(RF_Public | RF_Standalone);
+
+	if (UDynamicMaterialInstance* NewInstance = Cast<UDynamicMaterialInstance>(NewAsset))
+	{
+		if (UDynamicMaterialModel* NewModel = NewInstance->GetMaterialModel())
+		{
+			if (IDynamicMaterialModelEditorOnlyDataInterface* ModelEditorOnlyData = NewModel->GetEditorOnlyData())
+			{
+				ModelEditorOnlyData->RequestMaterialBuild();
+			}
+		}
+	}
+
+	FAssetRegistryModule::AssetCreated(NewAsset);
+	Package->FullyLoad();
+
+	return true;
+}
+
+bool UDMBlueprintFunctionLibrary::ExportGeneratedMaterial(UDynamicMaterialModel* InMaterialModel, const FString& InSavePath)
+{
+	if (!IsValid(InMaterialModel))
+	{
+		UE_LOG(LogDynamicMaterialEditor, Error, TEXT("Invalid material to export."));
+		return false;
+	}
+
+	if (InSavePath.IsEmpty())
+	{
+		UE_LOG(LogDynamicMaterialEditor, Error, TEXT("Invalid material save path to export."));
+		return false;
+	}
+
+	UMaterial* GeneratedMaterial = InMaterialModel->GetGeneratedMaterial();
+	if (!GeneratedMaterial)
+	{
+		UE_LOG(LogDynamicMaterialEditor, Error, TEXT("Failed to find a generated material to export."));
+		return false;
+	}
+
+	const FString PackagePath = FPaths::GetBaseFilename(*InSavePath, false);
+
+	UPackage* Package = CreatePackage(*PackagePath);
+	if (!Package)
+	{
+		UE_LOG(LogDynamicMaterialEditor, Error, TEXT("Failed to create package for exported material (%s)."), *PackagePath);
+		return false;
+	}
+
+	const FString AssetName = FPaths::GetBaseFilename(*InSavePath, true);
+
+	FObjectDuplicationParameters Params = InitStaticDuplicateObjectParams(GeneratedMaterial, Package, *AssetName,
+		RF_Public | RF_Standalone, nullptr, EDuplicateMode::Normal, EInternalObjectFlags::None);
+
+	UObject* NewAsset = StaticDuplicateObjectEx(Params);
+	if (!NewAsset)
+	{
+		UE_LOG(LogDynamicMaterialEditor, Error, TEXT("Failed to create new material asset."));
+		return false;
+	}
+
+	// Not sure why these need to be set again!
+	NewAsset->SetFlags(RF_Public | RF_Standalone);
+
+	FAssetRegistryModule::AssetCreated(NewAsset);
+	Package->FullyLoad();
+
+	return true;
+}
diff --git a/Engine/Plugins/Experimental/DynamicMaterial/Source/DynamicMaterialEditor/Private/DMEDefs.cpp b/Engine/Plugins/Experimental/DynamicMaterial/Source/DynamicMaterialEditor/Private/DMEDefs.cpp
new file mode 100644
--- /dev/null
+++ b/Engine/Plugins/Experimental/DynamicMaterial/Source/DynamicMaterialEditor/Private/DMEDefs.cpp
@@ -0,0 +1,61 @@
+// Copyright Epic Games, Inc. All Rights Reserved.
+
+#include "DMEDefs.h"
+
+#define LOCTEXT_NAMESPACE "MaterialDesignerEditorDefs"
+
+bool UE::DynamicMaterialEditor::IsCustomMaterialProperty(EDMMaterialPropertyType InMaterialProperty)
+{
+	return (InMaterialProperty >= EDMMaterialPropertyType::Custom1 && InMaterialProperty <= EDMMaterialPropertyType::Custom4);
+}
+
+int32 UE::DynamicMaterialEditor::ChannelIndexToChannelBit(int32 InChannelIndex)
+{
+	switch (InChannelIndex)
+	{
+		case 0:
+			return FDMMaterialStageConnectorChannel::WHOLE_CHANNEL;
+
+		case 1:
+			return FDMMaterialStageConnectorChannel::FIRST_CHANNEL;
+
+		case 2:
+			return FDMMaterialStageConnectorChannel::SECOND_CHANNEL;
+
+		case 3:
+			return FDMMaterialStageConnectorChannel::THIRD_CHANNEL;
+
+		case 4:
+			return FDMMaterialStageConnectorChannel::FOURTH_CHANNEL;
+
+		default:
+			checkNoEntry();
+			return 0;
+	}
+}
+
+int32 UE::DynamicMaterialEditor::ChannelBitToChannelIndex(int32 InChannelBit)
+{
+	switch (InChannelBit)
+	{
+		case FDMMaterialStageConnectorChannel::WHOLE_CHANNEL:
+			return 0;
+
+		case FDMMaterialStageConnectorChannel::FIRST_CHANNEL:
+			return 1;
+
+		case FDMMaterialStageConnectorChannel::SECOND_CHANNEL:
+			return 2;
+
+		case FDMMaterialStageConnectorChannel::THIRD_CHANNEL:
+			return 3;
+
+		case FDMMaterialStageConnectorChannel::FOURTH_CHANNEL:
+			return 4;
+
+		default:
+			return 0;
+	}
+}
+
+#undef LOCTEXT_NAMESPACE
diff --git a/Engine/Plugins/Experimental/DynamicMaterial/Source/DynamicMaterialEditor/Private/DMWorldSubsystem.cpp b/Engine/Plugins/Experimental/DynamicMaterial/Source/DynamicMaterialEditor/Private/DMWorldSubsystem.cpp
new file mode 100644
--- /dev/null
+++ b/Engine/Plugins/Experimental/DynamicMaterial/Source/DynamicMaterialEditor/Private/DMWorldSubsystem.cpp
@@ -0,0 +1,17 @@
+// Copyright Epic Games, Inc. All Rights Reserved.
+ 
+#include "DMWorldSubsystem.h"
+#include "IDetailKeyframeHandler.h"
+#include "LevelEditor/DMLevelEditorIntegration.h"
+
+UDMWorldSubsystem::UDMWorldSubsystem()
+	: KeyframeHandler(nullptr)
+{
+	// Default fallback implementation
+	InvokeTabDelegate.BindWeakLambda(this, [this]() { FDMLevelEditorIntegration::InvokeTabForWorld(GetWorld()); });
+}
+ 
+void UDMWorldSubsystem::SetKeyframeHandler(const TSharedPtr<IDetailKeyframeHandler>& InKeyframeHandler)
+{
+	KeyframeHandler = InKeyframeHandler;
+}
diff --git a/Engine/Plugins/Experimental/DynamicMaterial/Source/DynamicMaterialEditor/Private/DetailsPanel/Slate/SDMDetailsPanelMaterialInterfaceWidget.cpp b/Engine/Plugins/Experimental/DynamicMaterial/Source/DynamicMaterialEditor/Private/DetailsPanel/Slate/SDMDetailsPanelMaterialInterfaceWidget.cpp
new file mode 100644
--- /dev/null
+++ b/Engine/Plugins/Experimental/DynamicMaterial/Source/DynamicMaterialEditor/Private/DetailsPanel/Slate/SDMDetailsPanelMaterialInterfaceWidget.cpp
@@ -0,0 +1,207 @@
+// Copyright Epic Games, Inc. All Rights Reserved.
+
+#include "DetailsPanel/Slate/SDMDetailsPanelMaterialInterfaceWidget.h"
+#include "AssetRegistry/AssetRegistryModule.h"
+#include "AssetToolsModule.h"
+#include "DetailLayoutBuilder.h"
+#include "IAssetTools.h"
+#include "Material/DynamicMaterialInstance.h"
+#include "Material/DynamicMaterialInstanceFactory.h"
+#include "PropertyCustomizationHelpers.h"
+#include "PropertyHandle.h"
+#include "Widgets/Input/SButton.h"
+#include "Widgets/Layout/SWrapBox.h"
+#include "Widgets/Text/STextBlock.h"
+
+#define LOCTEXT_NAMESPACE "SDMDetailsPanelMaterialInterfaceWidget"
+
+void SDMDetailsPanelMaterialInterfaceWidget::Construct(const FArguments& InArgs, const TSharedRef<IPropertyHandle>& InPropertyHandle)
+{
+	PropertyHandle = InPropertyHandle;
+
+	TArray<UObject*> OuterObjects;
+	InPropertyHandle->GetOuterObjects(OuterObjects);
+
+	if (OuterObjects.IsEmpty())
+	{
+		return;
+	}
+
+	UObject* Value = nullptr;
+	InPropertyHandle->GetValue(Value);
+
+	UDynamicMaterialInstance* Instance = GetDynamicMaterialInstance();
+
+	FObjectPropertyBase* const ObjectProperty = CastField<FObjectPropertyBase>(PropertyHandle->GetProperty());
+
+	UClass* const ObjectClass = ObjectProperty
+		? ToRawPtr(ObjectProperty->PropertyClass)
+		: ToRawPtr(UMaterialInterface::StaticClass());
+
+	// @formatter:off
+	ChildSlot
+	[
+		SNew(SVerticalBox)
+		+ SVerticalBox::Slot()
+		.AutoHeight()
+		.HAlign(EHorizontalAlignment::HAlign_Left)
+		.VAlign(EVerticalAlignment::VAlign_Center)
+		.Padding(10.f, 5.f, 10.f, 5.f)
+		[
+			SNew(SObjectPropertyEntryBox)
+			.AllowClear(true)
+			.AllowedClass(ObjectClass)
+			.DisplayBrowse(true)
+			.DisplayCompactSize(false)
+			.DisplayThumbnail(true)
+			.DisplayUseSelected(true)
+			.EnableContentPicker(true)
+			.PropertyHandle(PropertyHandle)
+			.ThumbnailPool(InArgs._ThumbnailPool)
+		]
+		+ SVerticalBox::Slot()
+		.HAlign(EHorizontalAlignment::HAlign_Left)
+		.VAlign(EVerticalAlignment::VAlign_Center)
+		.Padding(10.f, 5.f, 10.f, 5.f)
+		.AutoHeight()
+		[
+			SNew(SButton)
+			.OnClicked(this, &SDMDetailsPanelMaterialInterfaceWidget::OnButtonClicked)
+			.IsEnabled(InPropertyHandle, &IPropertyHandle::IsEditable)
+			.Content()
+			[
+				SNew(STextBlock)
+				.Text(this, &SDMDetailsPanelMaterialInterfaceWidget::GetButtonText)
+				.Font(IDetailLayoutBuilder::GetDetailFontBold())
+			]
+		]
+	];
+	// @formatter:on
+}
+
+UObject* SDMDetailsPanelMaterialInterfaceWidget::GetAsset() const
+{
+	TArray<UObject*> OuterObjects;
+	PropertyHandle->GetOuterObjects(OuterObjects);
+
+	if (OuterObjects.IsEmpty())
+	{
+		return nullptr;
+	}
+
+	UObject* Value = nullptr;
+	PropertyHandle->GetValue(Value);
+
+	return Value;
+}
+
+UDynamicMaterialInstance* SDMDetailsPanelMaterialInterfaceWidget::GetDynamicMaterialInstance() const
+{
+	return Cast<UDynamicMaterialInstance>(GetAsset());
+}
+
+void SDMDetailsPanelMaterialInterfaceWidget::SetAsset(UObject* NewAsset)
+{
+	TArray<UObject*> OuterObjects;
+	PropertyHandle->GetOuterObjects(OuterObjects);
+
+	if (OuterObjects.IsEmpty())
+	{
+		return;
+	}
+
+	PropertyHandle->SetValueFromFormattedString(NewAsset ? NewAsset->GetPathName() : "");
+}
+
+void SDMDetailsPanelMaterialInterfaceWidget::SetDynamicMaterialInstance(UDynamicMaterialInstance* NewInstance)
+{
+	SetAsset(NewInstance);
+}
+
+FText SDMDetailsPanelMaterialInterfaceWidget::GetButtonText() const
+{
+	if (GetDynamicMaterialInstance())
+	{
+		return LOCTEXT("OpenMaterialDesignerModel", "Edit with Material Designer");
+	}
+
+	return LOCTEXT("CreateMaterialDesignerModel", "Create with Material Designer");
+}
+
+FReply SDMDetailsPanelMaterialInterfaceWidget::OnButtonClicked()
+{
+	if (GetDynamicMaterialInstance())
+	{
+		return OpenDynamicMaterialInstanceTab();
+	}
+
+	return CreateDynamicMaterialInstance();
+}
+
+FReply SDMDetailsPanelMaterialInterfaceWidget::CreateDynamicMaterialInstance()
+{
+	UDynamicMaterialInstance* Instance = GetDynamicMaterialInstance();
+
+	// We already have an instance, so we don't need to create one
+	if (Instance)
+	{
+		return FReply::Unhandled();
+	}
+
+	TArray<UObject*> OuterObjects;
+	PropertyHandle->GetOuterObjects(OuterObjects);
+
+	if (OuterObjects.IsEmpty())
+	{
+		return FReply::Unhandled();
+	}
+
+	UDynamicMaterialInstanceFactory* DynamicMaterialInstanceFactory = NewObject<UDynamicMaterialInstanceFactory>();
+	check(DynamicMaterialInstanceFactory);
+
+	UDynamicMaterialInstance* NewInstance = Cast<UDynamicMaterialInstance>(DynamicMaterialInstanceFactory->FactoryCreateNew(
+		UDynamicMaterialInstance::StaticClass(),
+		OuterObjects[0],
+		"DynamicMaterialInstance",
+		RF_NoFlags,
+		nullptr,
+		GWarn
+	));
+
+	PropertyHandle->SetValueFromFormattedString(NewInstance->GetPathName());
+
+	return OpenDynamicMaterialInstanceTab();
+}
+
+FReply SDMDetailsPanelMaterialInterfaceWidget::ClearDynamicMaterialInstance()
+{
+	UDynamicMaterialInstance* Instance = GetDynamicMaterialInstance();
+
+	// We don't have an instance, so we don't need to clear it (and don't clear non-MDIs)
+	if (!Instance)
+	{
+		return FReply::Unhandled();
+	}
+
+	SetDynamicMaterialInstance(nullptr);
+
+	return FReply::Handled();
+}
+
+FReply SDMDetailsPanelMaterialInterfaceWidget::OpenDynamicMaterialInstanceTab()
+{
+	UDynamicMaterialInstance* Instance = GetDynamicMaterialInstance();
+
+	// We don't have a MDI, so don't try to open it.
+	if (!Instance)
+	{
+		return FReply::Unhandled();
+	}
+
+	IAssetTools& AssetTools = FModuleManager::LoadModuleChecked<FAssetToolsModule>("AssetTools").Get();
+	AssetTools.OpenEditorForAssets({Instance});
+
+	return FReply::Handled();
+}
+
+#undef LOCTEXT_NAMESPACE
diff --git a/Engine/Plugins/Experimental/DynamicMaterial/Source/DynamicMaterialEditor/Private/DetailsPanel/Slate/SDMDetailsPanelTabSpawner.cpp b/Engine/Plugins/Experimental/DynamicMaterial/Source/DynamicMaterialEditor/Private/DetailsPanel/Slate/SDMDetailsPanelTabSpawner.cpp
new file mode 100644
--- /dev/null
+++ b/Engine/Plugins/Experimental/DynamicMaterial/Source/DynamicMaterialEditor/Private/DetailsPanel/Slate/SDMDetailsPanelTabSpawner.cpp
@@ -0,0 +1,218 @@
+// Copyright Epic Games, Inc. All Rights Reserved.
+
+#include "SDMDetailsPanelTabSpawner.h"
+#include "AssetRegistry/AssetRegistryModule.h"
+#include "AssetToolsModule.h"
+#include "DetailLayoutBuilder.h"
+#include "IAssetTools.h"
+#include "Model/DynamicMaterialModel.h"
+#include "Model/DynamicMaterialModelFactory.h"
+#include "PropertyCustomizationHelpers.h"
+#include "PropertyHandle.h"
+#include "Widgets/Input/SButton.h"
+#include "Widgets/Layout/SWrapBox.h"
+#include "Widgets/Text/STextBlock.h"
+
+#define LOCTEXT_NAMESPACE "SDMDetailsPanelTabSpawner"
+
+void SDMDetailsPanelTabSpawner::Construct(const FArguments& InArgs, const TSharedRef<IPropertyHandle>& InPropertyHandle)
+{
+	PropertyHandle = InPropertyHandle;
+
+	TArray<UObject*> OuterObjects;
+	InPropertyHandle->GetOuterObjects(OuterObjects);
+
+	if (OuterObjects.IsEmpty())
+	{
+		return;
+	}
+
+	UObject* Value = nullptr;
+	InPropertyHandle->GetValue(Value);
+
+	UDynamicMaterialModel* MaterialModel = GetMaterialModel();
+
+	// @formatter:off
+	ChildSlot
+	[
+		SNew(SVerticalBox)
+		+ SVerticalBox::Slot()
+		.AutoHeight()
+		.HAlign(EHorizontalAlignment::HAlign_Left)
+		.VAlign(EVerticalAlignment::VAlign_Center)
+		.Padding(10.f, 5.f, 10.f, 5.f)
+		[
+			SNew(SObjectPropertyEntryBox)
+			.AllowClear(true)
+			.AllowedClass(UDynamicMaterialModel::StaticClass())
+			.DisplayBrowse(true)
+			.DisplayThumbnail(false)
+			.DisplayCompactSize(true)
+			.DisplayUseSelected(true)
+			.EnableContentPicker(true)
+			.ObjectPath(this, &SDMDetailsPanelTabSpawner::GetEditorPath)
+			.OnObjectChanged(this, &SDMDetailsPanelTabSpawner::OnEditorChanged)
+		]
+		+ SVerticalBox::Slot()
+		.HAlign(EHorizontalAlignment::HAlign_Left)
+		.VAlign(EVerticalAlignment::VAlign_Center)
+		.Padding(10.f, 5.f, 10.f, 5.f)
+		.AutoHeight()
+		[
+			SNew(SButton)
+			.OnClicked(this, &SDMDetailsPanelTabSpawner::OnButtonClicked)
+			.Content()
+			[
+				SNew(STextBlock)
+				.Text(this, &SDMDetailsPanelTabSpawner::GetButtonText)
+				.Font(IDetailLayoutBuilder::GetDetailFontBold())
+			]
+		]
+	];
+	// @formatter:on
+}
+
+UDynamicMaterialModel* SDMDetailsPanelTabSpawner::GetMaterialModel() const
+{
+	TArray<UObject*> OuterObjects;
+	PropertyHandle->GetOuterObjects(OuterObjects);
+
+	if (OuterObjects.IsEmpty())
+	{
+		return nullptr;
+	}
+
+	UObject* Value = nullptr;
+	PropertyHandle->GetValue(Value);
+
+	return Cast<UDynamicMaterialModel>(Value);
+}
+
+void SDMDetailsPanelTabSpawner::SetMaterialModel(UDynamicMaterialModel* InNewModel)
+{
+	TArray<UObject*> OuterObjects;
+	PropertyHandle->GetOuterObjects(OuterObjects);
+
+	if (OuterObjects.IsEmpty())
+	{
+		return;
+	}
+
+	PropertyHandle->SetValueFromFormattedString(InNewModel ? InNewModel->GetPathName() : "");
+}
+
+FText SDMDetailsPanelTabSpawner::GetButtonText() const
+{
+	if (GetMaterialModel())
+	{
+		return LOCTEXT("OpenMaterialDesignerModel", "Edit with Material Designer");
+	}
+
+	return LOCTEXT("CreateMaterialDesignerModel", "Create with Material Designer");
+}
+
+FReply SDMDetailsPanelTabSpawner::OnButtonClicked()
+{
+	if (GetMaterialModel())
+	{
+		return OpenDynamicMaterialModelTab();
+	}
+
+	return CreateDynamicMaterialModel();
+}
+
+FReply SDMDetailsPanelTabSpawner::CreateDynamicMaterialModel()
+{
+	UDynamicMaterialModel* MaterialModel = GetMaterialModel();
+
+	// We already have a builder, so we don't need to create one
+	if (MaterialModel)
+	{
+		return FReply::Unhandled();
+	}
+
+	TArray<UObject*> OuterObjects;
+	PropertyHandle->GetOuterObjects(OuterObjects);
+
+	if (OuterObjects.IsEmpty())
+	{
+		return FReply::Unhandled();
+	}
+
+	FString PackageName, AssetName;
+	IAssetTools& AssetTools = FModuleManager::LoadModuleChecked<FAssetToolsModule>("AssetTools").Get();
+	AssetTools.CreateUniqueAssetName(
+		UDynamicMaterialModelFactory::BaseDirectory / UDynamicMaterialModelFactory::BaseName + FGuid::NewGuid().ToString(),
+		"",
+		PackageName,
+		AssetName
+	);
+
+	UPackage* Package = CreatePackage(*PackageName);
+	check(Package);
+
+	UDynamicMaterialModelFactory* DynamicMaterialModelFactory = NewObject<UDynamicMaterialModelFactory>();
+	check(DynamicMaterialModelFactory);
+
+	UDynamicMaterialModel* NewModel = Cast<UDynamicMaterialModel>(DynamicMaterialModelFactory->FactoryCreateNew(
+		UDynamicMaterialModel::StaticClass(),
+		Package,
+		*AssetName,
+		RF_Standalone | RF_Public,
+		nullptr,
+		GWarn
+	));
+
+	FAssetRegistryModule::AssetCreated(NewModel);
+	Package->FullyLoad();
+
+	PropertyHandle->SetValueFromFormattedString(NewModel->GetPathName());
+
+	return OpenDynamicMaterialModelTab();
+}
+
+FReply SDMDetailsPanelTabSpawner::ClearDynamicMaterialModel()
+{
+	UDynamicMaterialModel* MaterialModel = GetMaterialModel();
+
+	// We don't have a builder, so we don't need to clear it
+	if (!MaterialModel)
+	{
+		return FReply::Unhandled();
+	}
+
+	MaterialModel->ResetData();
+	SetMaterialModel(nullptr);
+
+	return FReply::Handled();
+}
+
+FReply SDMDetailsPanelTabSpawner::OpenDynamicMaterialModelTab()
+{
+	UDynamicMaterialModel* MaterialModel = GetMaterialModel();
+
+	// We don't have a builder, so we can't open it
+	if (!MaterialModel)
+	{
+		return FReply::Unhandled();
+	}
+
+	IAssetTools& AssetTools = FModuleManager::LoadModuleChecked<FAssetToolsModule>("AssetTools").Get();
+	AssetTools.OpenEditorForAssets({MaterialModel});
+
+	return FReply::Handled();
+}
+
+FString SDMDetailsPanelTabSpawner::GetEditorPath() const
+{
+	UDynamicMaterialModel* MaterialModel = GetMaterialModel();
+
+	return MaterialModel ? MaterialModel->GetPathName() : "";
+}
+
+void SDMDetailsPanelTabSpawner::OnEditorChanged(const FAssetData& InAssetData)
+{
+	SetMaterialModel(Cast<UDynamicMaterialModel>(InAssetData.GetAsset()));
+}
+
+#undef LOCTEXT_NAMESPACE
diff --git a/Engine/Plugins/Experimental/DynamicMaterial/Source/DynamicMaterialEditor/Private/DetailsPanel/Slate/SDMMaterialListExtensionWidget.cpp b/Engine/Plugins/Experimental/DynamicMaterial/Source/DynamicMaterialEditor/Private/DetailsPanel/Slate/SDMMaterialListExtensionWidget.cpp
new file mode 100644
--- /dev/null
+++ b/Engine/Plugins/Experimental/DynamicMaterial/Source/DynamicMaterialEditor/Private/DetailsPanel/Slate/SDMMaterialListExtensionWidget.cpp
@@ -0,0 +1,209 @@
+// Copyright Epic Games, Inc. All Rights Reserved.
+
+#include "SDMMaterialListExtensionWidget.h"
+#include "AssetToolsModule.h"
+#include "Components/PrimitiveComponent.h"
+#include "DMWorldSubsystem.h"
+#include "DetailLayoutBuilder.h"
+#include "Engine/World.h"
+#include "IAssetTools.h"
+#include "Material/DynamicMaterialInstance.h"
+#include "Material/DynamicMaterialInstanceFactory.h"
+#include "MaterialList.h"
+#include "Widgets/Input/SButton.h"
+#include "Widgets/Layout/SBox.h"
+#include "Widgets/Text/STextBlock.h"
+
+#define LOCTEXT_NAMESPACE "SDMMaterialListExtensionWidget"
+
+void SDMMaterialListExtensionWidget::Construct(const FArguments& InArgs, const TSharedRef<FMaterialItemView>& InMaterialItemView, 
+	UPrimitiveComponent* InCurrentComponent, IDetailLayoutBuilder& InDetailBuilder)
+{
+	MaterialItemViewWeak = InMaterialItemView;
+	CurrentComponentWeak = InCurrentComponent;
+
+	if (!CurrentComponentWeak.IsValid())
+	{
+		return;
+	}
+
+	// @formatter:off
+	ChildSlot
+	[
+		SNew(SBox)
+		.HAlign(EHorizontalAlignment::HAlign_Left)
+		[
+			SNew(SButton)
+			.OnClicked(this, &SDMMaterialListExtensionWidget::OnButtonClicked)
+			.Content()
+			[
+				SNew(STextBlock)
+				.Text(this, &SDMMaterialListExtensionWidget::GetButtonText)
+				.Font(IDetailLayoutBuilder::GetDetailFontBold())
+			]
+		]
+	];
+	// @formatter:on
+}
+
+UObject* SDMMaterialListExtensionWidget::GetAsset() const
+{
+	UPrimitiveComponent* CurrentComponent = CurrentComponentWeak.Get();
+	if (!ensure(CurrentComponent))
+	{
+		return nullptr;
+	}
+
+	const TSharedPtr<FMaterialItemView> MaterialItemView = MaterialItemViewWeak.Pin();
+	if (!ensure(MaterialItemView.IsValid()))
+	{
+		return nullptr;
+	}
+
+	UMaterialInterface* MaterialInterface = Cast<UMaterialInterface>(MaterialItemView->GetMaterialListItem().Material.Get());
+	if (!MaterialInterface)
+	{
+		return nullptr;
+	}
+
+	return MaterialInterface;
+}
+
+UDynamicMaterialInstance* SDMMaterialListExtensionWidget::GetDynamicMaterialInstance() const
+{
+	return Cast<UDynamicMaterialInstance>(GetAsset());
+}
+
+void SDMMaterialListExtensionWidget::SetAsset(UObject* NewAsset)
+{
+	UPrimitiveComponent* CurrentComponent = CurrentComponentWeak.Get();
+	if (!ensure(CurrentComponent))
+	{
+		return;
+	}
+
+	const TSharedPtr<FMaterialItemView> MaterialItemView = MaterialItemViewWeak.Pin();
+	if (!ensure(MaterialItemView.IsValid()))
+	{
+		return;
+	}
+
+	UMaterialInterface* MaterialInterface = Cast<UMaterialInterface>(NewAsset);
+	if (!MaterialInterface)
+	{
+		return;
+	}
+
+	MaterialItemView->ReplaceMaterial(MaterialInterface);
+}
+
+void SDMMaterialListExtensionWidget::SetDynamicMaterialInstance(UDynamicMaterialInstance* NewInstance)
+{
+	if (MaterialItemViewWeak.IsValid())
+	{
+		if (UPrimitiveComponent* CurrentComponent = CurrentComponentWeak.Get())
+		{
+			if (UWorld* World = CurrentComponent->GetWorld())
+			{
+				if (UDMWorldSubsystem* WorldSubsystem = World->GetSubsystem<UDMWorldSubsystem>())
+				{
+					if (WorldSubsystem->GetInvokeTabDelegate().IsBound())
+					{
+						if (WorldSubsystem->GetMaterialValueSetterDelegate().IsBound())
+						{
+							const FMaterialListItem& ListItem = MaterialItemViewWeak.Pin()->GetMaterialListItem();
+							const FDMObjectMaterialProperty MaterialProperty(CurrentComponent, ListItem.SlotIndex);
+							WorldSubsystem->GetMaterialValueSetterDelegate().Execute(MaterialProperty, NewInstance);
+						}
+
+						return;
+					}
+				}
+			}
+		}
+	}
+
+	SetAsset(NewInstance);
+}
+
+FText SDMMaterialListExtensionWidget::GetButtonText() const
+{
+	if (GetDynamicMaterialInstance())
+	{
+		return LOCTEXT("OpenMaterialDesignerModel", "Edit with Material Designer");
+	}
+
+	return LOCTEXT("CreateMaterialDesignerModel", "Create with Material Designer");
+}
+
+FReply SDMMaterialListExtensionWidget::OnButtonClicked()
+{
+	if (GetDynamicMaterialInstance())
+	{
+		return OpenDynamicMaterialInstanceTab();
+	}
+
+	return CreateDynamicMaterialInstance();
+}
+
+FReply SDMMaterialListExtensionWidget::CreateDynamicMaterialInstance()
+{
+	UDynamicMaterialInstance* Instance = GetDynamicMaterialInstance();
+
+	// We already have an instance, so we don't need to create one
+	if (Instance)
+	{
+		return FReply::Unhandled();
+	}
+
+	FName InstanceName = MakeUniqueObjectName(CurrentComponentWeak.Get(), UDynamicMaterialInstance::StaticClass(), "DynamicMaterialInstance");
+
+	UDynamicMaterialInstanceFactory* DynamicMaterialInstanceFactory = NewObject<UDynamicMaterialInstanceFactory>();
+	check(DynamicMaterialInstanceFactory);
+
+	UDynamicMaterialInstance* NewInstance = Cast<UDynamicMaterialInstance>(DynamicMaterialInstanceFactory->FactoryCreateNew(
+		UDynamicMaterialInstance::StaticClass(),
+		CurrentComponentWeak.Get(),
+		InstanceName,
+		RF_NoFlags,
+		nullptr,
+		GWarn
+	));
+
+	SetDynamicMaterialInstance(NewInstance);
+
+	return OpenDynamicMaterialInstanceTab();
+}
+
+FReply SDMMaterialListExtensionWidget::ClearDynamicMaterialInstance()
+{
+	UDynamicMaterialInstance* Instance = GetDynamicMaterialInstance();
+
+	// We don't have an instance, so we don't need to clear it (or any other asset in its place)
+	if (!Instance)
+	{
+		return FReply::Unhandled();
+	}
+
+	SetDynamicMaterialInstance(nullptr);
+
+	return FReply::Handled();
+}
+
+FReply SDMMaterialListExtensionWidget::OpenDynamicMaterialInstanceTab()
+{
+	UDynamicMaterialInstance* Instance = GetDynamicMaterialInstance();
+
+	// We don't have an instance, so we can't open it
+	if (!Instance)
+	{
+		return FReply::Unhandled();
+	}
+
+	IAssetTools& AssetTools = FModuleManager::LoadModuleChecked<FAssetToolsModule>("AssetTools").Get();
+	AssetTools.OpenEditorForAssets({Instance});
+
+	return FReply::Handled();
+}
+
+#undef LOCTEXT_NAMESPACE
diff --git a/Engine/Plugins/Experimental/DynamicMaterial/Source/DynamicMaterialEditor/Private/DynamicMaterialEditorModule.cpp b/Engine/Plugins/Experimental/DynamicMaterial/Source/DynamicMaterialEditor/Private/DynamicMaterialEditorModule.cpp
new file mode 100644
--- /dev/null
+++ b/Engine/Plugins/Experimental/DynamicMaterial/Source/DynamicMaterialEditor/Private/DynamicMaterialEditorModule.cpp
@@ -0,0 +1,500 @@
+// Copyright Epic Games, Inc. All Rights Reserved.
+
+#include "DynamicMaterialEditorModule.h"
+#include "AssetToolsModule.h"
+#include "Components/DMMaterialComponent.h"
+#include "Components/DMMaterialEffectFunction.h"
+#include "Components/DMMaterialStageSource.h"
+#include "Components/DMMaterialStageThroughput.h"
+#include "Components/DMTextureUV.h"
+#include "Components/MaterialStageInputs/DMMSIThroughput.h"
+#include "Components/MaterialValues/DMMaterialValueBool.h"
+#include "Components/MaterialValues/DMMaterialValueFloat1.h"
+#include "Components/MaterialValues/DMMaterialValueFloat2.h"
+#include "Components/MaterialValues/DMMaterialValueFloat3RGB.h"
+#include "Components/MaterialValues/DMMaterialValueFloat3RPY.h"
+#include "Components/MaterialValues/DMMaterialValueFloat3XYZ.h"
+#include "Components/MaterialValues/DMMaterialValueFloat4.h"
+#include "Components/MaterialValues/DMMaterialValueTexture.h"
+#include "Components/PrimitiveComponent.h"
+#include "DetailLayoutBuilder.h"
+#include "DetailsPanel/DMPropertyTypeCustomizer.h"
+#include "DetailsPanel/Slate/SDMMaterialListExtensionWidget.h"
+#include "DMMaterialFunctionLibrary.h"
+#include "DMWorldSubsystem.h"
+#include "DynamicMaterialEditorCommands.h"
+#include "DynamicMaterialEditorSettings.h"
+#include "DynamicMaterialEditorStyle.h"
+#include "DynamicMaterialModule.h"
+#include "Engine/World.h"
+#include "LevelEditor/DMLevelEditorIntegration.h"
+#include "Material/DynamicMaterialInstance.h"
+#include "MaterialList.h"
+#include "Model/DynamicMaterialModel.h"
+#include "Model/DynamicMaterialModelEditorOnlyData.h"
+#include "Modules/ModuleManager.h"
+#include "PropertyEditorModule.h"
+#include "Slate/Properties/Editors/SDMPropertyEditBoolValue.h"
+#include "Slate/Properties/Editors/SDMPropertyEditFloat1Value.h"
+#include "Slate/Properties/Editors/SDMPropertyEditFloat2Value.h"
+#include "Slate/Properties/Editors/SDMPropertyEditFloat3RGBValue.h"
+#include "Slate/Properties/Editors/SDMPropertyEditFloat3RPYValue.h"
+#include "Slate/Properties/Editors/SDMPropertyEditFloat3XYZValue.h"
+#include "Slate/Properties/Editors/SDMPropertyEditFloat4Value.h"
+#include "Slate/Properties/Editors/SDMPropertyEditTextureValue.h"
+#include "Slate/Properties/Generators/DMComponentPropertyRowGenerator.h"
+#include "Slate/Properties/Generators/DMInputThroughputPropertyRowGenerator.h"
+#include "Slate/Properties/Generators/DMMaterialEffectFunctionPropertyRowGenerator.h"
+#include "Slate/Properties/Generators/DMMaterialValuePropertyRowGenerator.h"
+#include "Slate/Properties/Generators/DMStagePropertyRowGenerator.h"
+#include "Slate/Properties/Generators/DMTextureUVPropertyRowGenerator.h"
+#include "Slate/Properties/Generators/DMThroughputPropertyRowGenerator.h"
+#include "Slate/SDMEditor.h"
+
+DEFINE_LOG_CATEGORY(LogDynamicMaterialEditor);
+
+namespace UE::DynamicMaterialEditor::Private
+{
+	FDelegateHandle MateriaListWidgetsDelegate;
+
+	void AddMaterialListWidgets(const TSharedRef<FMaterialItemView>& InMaterialItemView,
+		UActorComponent* InCurrentComponent,
+		IDetailLayoutBuilder& InDetailBuilder,
+		TArray<TSharedPtr<SWidget>>& OutExtensions)
+	{
+		if (UPrimitiveComponent* PrimitiveComponent = Cast<UPrimitiveComponent>(InCurrentComponent))
+		{
+			OutExtensions.Add(SNew(SDMMaterialListExtensionWidget,
+				InMaterialItemView,
+				PrimitiveComponent,
+				InDetailBuilder
+			));
+		}
+	}
+}
+
+const FName FDynamicMaterialEditorModule::TabId = TEXT("MaterialDesigner");
+TMap<UClass*, FDMCreateValueEditWidgetDelegate> FDynamicMaterialEditorModule::ValueEditWidgetDelegates;
+TMap<UClass*, FDMComponentPropertyRowGeneratorDelegate> FDynamicMaterialEditorModule::ComponentPropertyRowGenerators;
+TMap<UClass*, FDMGetObjectMaterialPropertiesDelegate> FDynamicMaterialEditorModule::CustomMaterialPropertyGenerators;
+FDMOnUIValueUpdate FDynamicMaterialEditorModule::OnUIValueUpdate;
+
+void FDynamicMaterialEditorModule::RegisterValueEditWidgetDelegate(UClass* InClass, FDMCreateValueEditWidgetDelegate ValueEditBodyDelegate)
+{
+	ValueEditWidgetDelegates.Emplace(InClass, ValueEditBodyDelegate);
+}
+
+FDMCreateValueEditWidgetDelegate FDynamicMaterialEditorModule::GetValueEditWidgetDelegate(UClass* InClass)
+{
+	FDMCreateValueEditWidgetDelegate* Delegate = ValueEditWidgetDelegates.Find(InClass);
+
+	if (Delegate)
+	{
+		return *Delegate;
+	}
+
+	return FDMCreateValueEditWidgetDelegate::CreateLambda([](TSharedPtr<SDMComponentEdit>, UDMMaterialValue*)->TSharedPtr<SWidget> { return nullptr; });
+}
+
+TSharedPtr<SWidget> FDynamicMaterialEditorModule::CreateEditWidgetForValue(const TSharedPtr<SDMComponentEdit>& InComponentEditWidget, UDMMaterialValue* InValue)
+{
+	if (!IsValid(InValue))
+	{
+		return SNullWidget::NullWidget;
+	}
+
+	FDMCreateValueEditWidgetDelegate ValueDelegate = GetValueEditWidgetDelegate(InValue->GetClass());
+
+	if (ValueDelegate.IsBound())
+	{
+		return ValueDelegate.Execute(InComponentEditWidget, InValue);
+	}
+
+	return SNullWidget::NullWidget;
+}
+
+void FDynamicMaterialEditorModule::RegisterComponentPropertyRowGeneratorDelegate(UClass* InClass, FDMComponentPropertyRowGeneratorDelegate InComponentPropertyRowGeneratorDelegate)
+{
+	ComponentPropertyRowGenerators.Emplace(InClass, InComponentPropertyRowGeneratorDelegate);
+}
+
+FDMComponentPropertyRowGeneratorDelegate FDynamicMaterialEditorModule::GetComponentPropertyRowGeneratorDelegate(UClass* InClass)
+{
+	UClass* FoundClass = nullptr;
+
+	for (const TPair<UClass*, FDMComponentPropertyRowGeneratorDelegate>& Pair : ComponentPropertyRowGenerators)
+	{
+		// If we have an exact match, return the delegate immediately.
+		if (Pair.Key == InClass)
+		{
+			return Pair.Value;
+		}
+
+		if (InClass->IsChildOf(Pair.Key))
+		{
+			// If we haven't got anything or the current class is a more specific generator.
+			if (FoundClass == nullptr || Pair.Key->IsChildOf(FoundClass))
+			{
+				FoundClass = Pair.Key;
+			}
+		}
+	}
+
+	if (!FoundClass)
+	{
+		// Return an invalid lambda
+		static FDMComponentPropertyRowGeneratorDelegate NullLambda = FDMComponentPropertyRowGeneratorDelegate::CreateLambda(
+			[](const TSharedRef<SDMComponentEdit>& InComponentEditWidget, UDMMaterialComponent* InComponent, TArray<FDMPropertyHandle>& InOutPropertyRows,
+				TSet<UDMMaterialComponent*>& InOutProcessedObjects)
+			{
+			});
+
+		return NullLambda;
+	}
+
+	return ComponentPropertyRowGenerators[FoundClass];
+}
+
+void FDynamicMaterialEditorModule::GeneratorComponentPropertyRows(const TSharedRef<SDMComponentEdit>& InComponentEditWidget, UDMMaterialComponent* InComponent,
+	TArray<FDMPropertyHandle>& InOutPropertyRows, TSet<UDMMaterialComponent*>& InOutProcessedObjects)
+{
+	if (!IsValid(InComponent))
+	{
+		return;
+	}
+
+	if (InOutProcessedObjects.Contains(InComponent))
+	{
+		return;
+	}
+
+	FDMComponentPropertyRowGeneratorDelegate RowGenerator = GetComponentPropertyRowGeneratorDelegate(InComponent->GetClass());
+	RowGenerator.ExecuteIfBound(InComponentEditWidget, InComponent, InOutPropertyRows, InOutProcessedObjects);
+}
+
+void FDynamicMaterialEditorModule::RegisterCustomMaterialPropertyGenerator(UClass* InClass, FDMGetObjectMaterialPropertiesDelegate InGenerator)
+{
+	if (!InClass || !InGenerator.IsBound())
+	{
+		return;
+	}
+
+	CustomMaterialPropertyGenerators.FindOrAdd(InClass) = InGenerator;
+}
+
+FDMGetObjectMaterialPropertiesDelegate FDynamicMaterialEditorModule::GetCustomMaterialPropertyGenerator(UClass* InClass)
+{
+	if (InClass)
+	{
+		if (CustomMaterialPropertyGenerators.Contains(InClass))
+		{
+			return CustomMaterialPropertyGenerators[InClass];
+		}
+	}
+
+	return FDMGetObjectMaterialPropertiesDelegate();
+}
+
+FDynamicMaterialEditorModule& FDynamicMaterialEditorModule::Get()
+{
+	return FModuleManager::LoadModuleChecked<FDynamicMaterialEditorModule>(TEXT("DynamicMaterialEditor"));
+}
+
+FDynamicMaterialEditorModule::FDynamicMaterialEditorModule()
+	: CommandList(MakeShared<FUICommandList>())
+{
+}
+
+void FDynamicMaterialEditorModule::StartupModule()
+{
+	FDynamicMaterialEditorStyle::Initialize();
+	FDynamicMaterialEditorCommands::Register();
+	MapCommands();
+
+	FunctionLibrary = MakeShared<FDMMaterialFunctionLibrary>();
+
+	FPropertyEditorModule& PropertyModule = FModuleManager::GetModuleChecked<FPropertyEditorModule>("PropertyEditor");
+
+	PropertyModule.RegisterCustomPropertyTypeLayout(UDynamicMaterialModelEditorOnlyData::StaticClass()->GetFName()
+		, FOnGetPropertyTypeCustomizationInstance::CreateStatic(&FDMPropertyTypeCustomizer::MakeInstance));
+
+	using namespace UE::DynamicMaterialEditor::Private;
+
+	MateriaListWidgetsDelegate = FMaterialList::OnAddMaterialItemViewExtraBottomWidget.AddStatic(&AddMaterialListWidgets);
+
+	FDMLevelEditorIntegration::Initialize();
+
+	RegisterValueEditWidgetDelegate<UDMMaterialValueBool,      SDMPropertyEditBoolValue>();
+	RegisterValueEditWidgetDelegate<UDMMaterialValueFloat1,    SDMPropertyEditFloat1Value>();
+	RegisterValueEditWidgetDelegate<UDMMaterialValueFloat2,    SDMPropertyEditFloat2Value>();
+	RegisterValueEditWidgetDelegate<UDMMaterialValueFloat3RGB, SDMPropertyEditFloat3RGBValue>();
+	RegisterValueEditWidgetDelegate<UDMMaterialValueFloat3RPY, SDMPropertyEditFloat3RPYValue>();
+	RegisterValueEditWidgetDelegate<UDMMaterialValueFloat3XYZ, SDMPropertyEditFloat3XYZValue>();
+	RegisterValueEditWidgetDelegate<UDMMaterialValueFloat4,    SDMPropertyEditFloat4Value>();
+	RegisterValueEditWidgetDelegate<UDMMaterialValueTexture,   SDMPropertyEditTextureValue>();
+
+	RegisterComponentPropertyRowGeneratorDelegate<UDMMaterialComponent,            FDMComponentPropertyRowGenerator>();
+	RegisterComponentPropertyRowGeneratorDelegate<UDMMaterialStage,                FDMStagePropertyRowGenerator>();
+	RegisterComponentPropertyRowGeneratorDelegate<UDMMaterialValue,                FDMMaterialValuePropertyRowGenerator>();
+	RegisterComponentPropertyRowGeneratorDelegate<UDMTextureUV,                    FDMTextureUVPropertyRowGenerator>();
+	RegisterComponentPropertyRowGeneratorDelegate<UDMMaterialStageThroughput,      FDMThroughputPropertyRowGenerator>();
+	RegisterComponentPropertyRowGeneratorDelegate<UDMMaterialStageInputThroughput, FDMInputThroughputPropertyRowGenerator>();
+	RegisterComponentPropertyRowGeneratorDelegate<UDMMaterialEffectFunction,       FDMMaterialEffectFunctionPropertyRowGenerator>();
+
+	FDynamicMaterialModule::GetCreateEditorOnlyDataDelegate().BindLambda([](UDynamicMaterialModel* InMaterialModel) -> TScriptInterface<IDynamicMaterialModelEditorOnlyDataInterface>
+		{
+			return TScriptInterface<IDynamicMaterialModelEditorOnlyDataInterface>(NewObject<UDynamicMaterialModelEditorOnlyData>(InMaterialModel));
+		});
+
+	BuildRequestList.Empty();
+
+	UDMMaterialValueTexture::GetDefaultRGBTexture.BindLambda([]()
+		{
+			return UDynamicMaterialEditorSettings::Get()->DefaultRGBTexture.LoadSynchronous();
+		});
+}
+
+void FDynamicMaterialEditorModule::ShutdownModule()
+{
+	FDynamicMaterialEditorCommands::Unregister();
+	FDynamicMaterialEditorStyle::Shutdown();
+
+	FunctionLibrary.Reset();
+
+	if (UObjectInitialized() && !IsEngineExitRequested() && FModuleManager::Get().IsModuleLoaded("PropertyEditor"))
+	{
+		FPropertyEditorModule& PropertyModule = FModuleManager::GetModuleChecked<FPropertyEditorModule>("PropertyEditor");
+		PropertyModule.UnregisterCustomPropertyTypeLayout(UDynamicMaterialModelEditorOnlyData::StaticClass()->GetFName());
+	}
+
+	using namespace UE::DynamicMaterialEditor::Private;
+	FMaterialList::OnAddMaterialItemViewExtraBottomWidget.Remove(MateriaListWidgetsDelegate);
+
+	FDMLevelEditorIntegration::Shutdown();
+
+	BuildRequestList.Empty();
+
+	UDMMaterialValueTexture::GetDefaultRGBTexture.Unbind();
+}
+
+const TSharedPtr<FDMMaterialFunctionLibrary>& FDynamicMaterialEditorModule::GetFunctionLibrary()
+{
+	return FunctionLibrary;
+}
+
+void FDynamicMaterialEditorModule::SetDynamicMaterialModel(UDynamicMaterialModel* InMaterialModel, UWorld* InWorld, bool bInInvokeTab)
+{
+	using namespace UE::DynamicMaterialEditor::Private;
+
+	if (TSharedPtr<SDMEditor> Editor = FDMLevelEditorIntegration::GetEditorForWorld(InWorld))
+	{
+		if (bInInvokeTab)
+		{
+			FDMLevelEditorIntegration::InvokeTabForWorld(InWorld);
+		}
+
+		Editor->SetMaterialModel(InMaterialModel);
+	}
+	else if (IsValid(InWorld))
+	{
+		if (UDMWorldSubsystem* DMWorldSubsystem = InWorld->GetSubsystem<UDMWorldSubsystem>())
+		{
+			if (bInInvokeTab)
+			{
+				DMWorldSubsystem->GetInvokeTabDelegate().ExecuteIfBound();
+			}
+
+			DMWorldSubsystem->GetSetCustomEditorModelDelegate().ExecuteIfBound(InMaterialModel);
+		}
+	}
+}
+
+void FDynamicMaterialEditorModule::SetDynamicMaterialObjectProperty(const FDMObjectMaterialProperty& InObjectProperty,
+	UWorld* InWorld, bool bInInvokeTab)
+{
+	using namespace UE::DynamicMaterialEditor::Private;
+
+	if (TSharedPtr<SDMEditor> Editor = FDMLevelEditorIntegration::GetEditorForWorld(InWorld))
+	{
+		if (bInInvokeTab)
+		{
+			FDMLevelEditorIntegration::InvokeTabForWorld(InWorld);
+		}
+
+		Editor->SetMaterialObjectProperty(InObjectProperty);
+	}
+	else if (IsValid(InWorld))
+	{
+		if (UDMWorldSubsystem* DMWorldSubsystem = InWorld->GetSubsystem<UDMWorldSubsystem>())
+		{
+			if (bInInvokeTab)
+			{
+				DMWorldSubsystem->GetInvokeTabDelegate().ExecuteIfBound();
+			}
+
+			DMWorldSubsystem->GetCustomObjectPropertyEditorDelegate().ExecuteIfBound(InObjectProperty);
+		}
+	}
+}
+
+void FDynamicMaterialEditorModule::SetDynamicMaterialInstance(UDynamicMaterialInstance* InInstance, UWorld* InWorld, 
+	bool bInInvokeTab)
+{
+	if (IsValid(InInstance))
+	{
+		if (UDynamicMaterialModel* InstanceModel = InInstance->GetMaterialModel())
+		{
+			SetDynamicMaterialModel(InstanceModel, InWorld, bInInvokeTab);
+		}
+	}
+}
+
+void FDynamicMaterialEditorModule::SetDynamicMaterialActor(AActor* InActor, UWorld* InWorld, bool bInInvokeTab)
+{
+	using namespace UE::DynamicMaterialEditor::Private;
+
+	if (TSharedPtr<SDMEditor> Editor = FDMLevelEditorIntegration::GetEditorForWorld(InWorld))
+	{
+		if (bInInvokeTab)
+		{
+			FDMLevelEditorIntegration::InvokeTabForWorld(InWorld);
+		}
+
+		Editor->SetMaterialActor(InActor);
+	}
+	else if (IsValid(InWorld))
+	{
+		if (UDMWorldSubsystem* DMWorldSubsystem = InWorld->GetSubsystem<UDMWorldSubsystem>())
+		{
+			if (bInInvokeTab)
+			{
+				DMWorldSubsystem->GetInvokeTabDelegate().ExecuteIfBound();
+			}
+
+			DMWorldSubsystem->GetSetCustomEditorActorDelegate().ExecuteIfBound(InActor);
+		}
+	}
+}
+
+void FDynamicMaterialEditorModule::ClearDynamicMaterialModel(UWorld* InWorld)
+{
+	SetDynamicMaterialModel(nullptr, InWorld, /* Invoke tab */ false);
+}
+
+TSharedRef<SWidget> FDynamicMaterialEditorModule::CreateEditor(UDynamicMaterialModel* InMaterialModel, UWorld* InAssetEditorWorld)
+{
+	TSharedRef<SDMEditor> NewEditor = SNew(SDMEditor, InMaterialModel);
+
+	if (IsValid(InAssetEditorWorld))
+	{
+		UDMWorldSubsystem* WorldSubsystem = InAssetEditorWorld->GetSubsystem<UDMWorldSubsystem>();
+
+		if (IsValid(WorldSubsystem))
+		{
+			WorldSubsystem->GetSetCustomEditorModelDelegate().BindSP(NewEditor, &SDMEditor::SetMaterialModel);
+			WorldSubsystem->GetCustomObjectPropertyEditorDelegate().BindSP(NewEditor, &SDMEditor::SetMaterialObjectProperty);
+			WorldSubsystem->GetSetCustomEditorActorDelegate().BindSP(NewEditor, &SDMEditor::SetMaterialActor);
+		}
+	}
+
+	return NewEditor;
+}
+
+TSharedRef<SWidget> FDynamicMaterialEditorModule::CreateEmptyTabContent()
+{
+	return SDMEditor::GetEmptyContent();
+}
+
+void FDynamicMaterialEditorModule::Tick(float DeltaTime)
+{
+	if (UDMMaterialComponent::CanClean() == false)
+	{
+		return;
+	}
+
+	for (const FDMBuildRequestEntry& ToBuild : BuildRequestList)
+	{
+		if (UObject* Object = FindObject<UObject>(nullptr, *ToBuild.AssetPath, false))
+		{
+			ProcessBuildRequest(Object, ToBuild.bDirtyAssets);
+		}
+	}
+
+	BuildRequestList.Empty();
+}
+
+TStatId FDynamicMaterialEditorModule::GetStatId() const
+{
+	RETURN_QUICK_DECLARE_CYCLE_STAT(FDynamicMaterialEditorModule, STATGROUP_Tickables);
+}
+
+void FDynamicMaterialEditorModule::AddBuildRequest(UObject* InToBuild, bool bInDirtyAssets)
+{
+	if (!IsValid(InToBuild))
+	{
+		return;
+	}
+
+	FDynamicMaterialEditorModule::Get().BuildRequestList.Add({InToBuild->GetPathName(), bInDirtyAssets});
+	
+	static const double VeryShortTime = 0.0001;
+
+	// Make sure we don't spam updates on a single tick.
+	UDMMaterialComponent::PreventClean(VeryShortTime);
+}
+
+void FDynamicMaterialEditorModule::OpenEditor(UWorld* InWorld)
+{
+	if (!IsValid(InWorld))
+	{
+		FDMLevelEditorIntegration::InvokeTabForWorld(InWorld);
+	}
+	else if (UDMWorldSubsystem* DMWorldSubsystem = InWorld->GetSubsystem<UDMWorldSubsystem>())
+	{
+		DMWorldSubsystem->GetInvokeTabDelegate().ExecuteIfBound();
+	}
+}
+
+void FDynamicMaterialEditorModule::ProcessBuildRequest(UObject* InToBuild, bool bInDirtyAssets)
+{
+	if (!IsValid(InToBuild))
+	{
+		return;
+	}
+
+	if (InToBuild->GetClass()->ImplementsInterface(UDMBuildable::StaticClass()) == false)
+	{
+		return;
+	}
+
+	IDMBuildable::Execute_DoBuild(InToBuild, bInDirtyAssets);
+}
+
+void FDynamicMaterialEditorModule::MapCommands()
+{
+	const FDynamicMaterialEditorCommands& DMEditorCommands = FDynamicMaterialEditorCommands::Get();
+
+	UDynamicMaterialEditorSettings* Settings = UDynamicMaterialEditorSettings::Get();
+
+	CommandList->MapAction(
+		DMEditorCommands.OpenEditorSettingsWindow,
+		FExecuteAction::CreateUObject(Settings, &UDynamicMaterialEditorSettings::OpenEditorSettingsWindow)
+	);
+}
+
+void FDynamicMaterialEditorModule::UnmapCommands()
+{
+	const FDynamicMaterialEditorCommands& DMEditorCommands = FDynamicMaterialEditorCommands::Get();
+
+	auto UnmapAction = [this](const TSharedPtr<const FUICommandInfo>& InCommandInfo)
+	{
+		if (CommandList->IsActionMapped(InCommandInfo))
+		{
+			CommandList->UnmapAction(InCommandInfo);
+		}
+	};
+
+	UnmapAction(DMEditorCommands.OpenEditorSettingsWindow);
+}
+
+IMPLEMENT_MODULE(FDynamicMaterialEditorModule, DynamicMaterialEditor)
diff --git a/Engine/Plugins/Experimental/DynamicMaterial/Source/DynamicMaterialEditor/Private/DynamicMaterialEditorSettings.cpp b/Engine/Plugins/Experimental/DynamicMaterial/Source/DynamicMaterialEditor/Private/DynamicMaterialEditorSettings.cpp
new file mode 100644
--- /dev/null
+++ b/Engine/Plugins/Experimental/DynamicMaterial/Source/DynamicMaterialEditor/Private/DynamicMaterialEditorSettings.cpp
@@ -0,0 +1,233 @@
+// Copyright Epic Games, Inc. All Rights Reserved.
+
+#include "DynamicMaterialEditorSettings.h"
+#include "AssetRegistry/AssetData.h"
+#include "Components/DMMaterialComponent.h"
+#include "Engine/AssetManager.h"
+#include "GenericPlatform/GenericApplication.h"
+#include "Interfaces/IPluginManager.h"
+#include "ISettingsModule.h"
+#include "Materials/MaterialFunctionInterface.h"
+#include "Misc/Paths.h"
+#include "Modules/ModuleManager.h"
+#include "UObject/UObjectIterator.h"
+
+#define LOCTEXT_NAMESPACE "MaterialDesignerSettings"
+
+UDynamicMaterialEditorSettings::UDynamicMaterialEditorSettings()
+{
+	CategoryName = TEXT("Avalanche");
+	SectionName = TEXT("Material Designer");
+
+	bFollowSelection = true;
+
+	ResetAllLayoutSettings();
+
+	DefaultRGBTexture = TSoftObjectPtr<UTexture>(FSoftObjectPath(TEXT("Texture2D'/DynamicMaterial/T_Default_Texture.T_Default_Texture'")));
+	DefaultOpaqueTexture = TSoftObjectPtr<UTexture>(FSoftObjectPath(TEXT("Texture2D'/Engine/EngineResources/WhiteSquareTexture.WhiteSquareTexture'")));
+	DefaultPreviewCanvasTexture = nullptr;
+	DefaultOpacitySlotMask = TSoftObjectPtr<UTexture>(FSoftObjectPath(TEXT("/Script/Engine.Texture2D'/DynamicMaterial/Textures/T_DM_HorizontalGradient.T_DM_HorizontalGradient'")));
+}
+
+UDynamicMaterialEditorSettings* UDynamicMaterialEditorSettings::Get()
+{
+	UDynamicMaterialEditorSettings* DefaultSettings = GetMutableDefault<UDynamicMaterialEditorSettings>();
+	static bool bInitialized = false;
+	if (!bInitialized)
+	{
+		bInitialized = true;
+		DefaultSettings->SetFlags(RF_Transactional);
+	}
+	return DefaultSettings;
+}
+
+void UDynamicMaterialEditorSettings::PostInitProperties()
+{
+	Super::PostInitProperties();
+
+}
+
+void UDynamicMaterialEditorSettings::PostEditChangeProperty(FPropertyChangedEvent& InPropertyChangedEvent)
+{
+	Super::PostEditChangeProperty(InPropertyChangedEvent);
+
+	const FName PropertyName = InPropertyChangedEvent.GetMemberPropertyName();
+
+	if (PropertyName == GET_MEMBER_NAME_CHECKED(UDynamicMaterialEditorSettings, SpinBoxValueMultiplier_Default)
+		|| PropertyName == GET_MEMBER_NAME_CHECKED(UDynamicMaterialEditorSettings, SpinBoxValueMultiplier_Shift)
+		|| PropertyName == GET_MEMBER_NAME_CHECKED(UDynamicMaterialEditorSettings, SpinBoxValueMultiplier_AltShift)
+		|| PropertyName == GET_MEMBER_NAME_CHECKED(UDynamicMaterialEditorSettings, SpinBoxValueMultiplier_Cmd)
+		|| PropertyName == GET_MEMBER_NAME_CHECKED(UDynamicMaterialEditorSettings, SpinBoxValueMultiplier_AltCmd)
+		|| PropertyName == GET_MEMBER_NAME_CHECKED(UDynamicMaterialEditorSettings, SplitterLocation)
+		|| PropertyName == GET_MEMBER_NAME_CHECKED(UDynamicMaterialEditorSettings, LayerPreviewSize)
+		|| PropertyName == GET_MEMBER_NAME_CHECKED(UDynamicMaterialEditorSettings, SlotPreviewSize)
+		|| PropertyName == GET_MEMBER_NAME_CHECKED(UDynamicMaterialEditorSettings, DetailsPreviewSize)
+		|| PropertyName == GET_MEMBER_NAME_CHECKED(UDynamicMaterialEditorSettings, bShowTooltipPreview)
+		|| PropertyName == GET_MEMBER_NAME_CHECKED(UDynamicMaterialEditorSettings, TooltipTextureSize)
+		|| PropertyName == GET_MEMBER_NAME_CHECKED(UDynamicMaterialEditorSettings, bPreviewImagesUseTextureUVs))
+	{
+		OnSettingsChanged.Broadcast(InPropertyChangedEvent);
+	}
+
+	if (PropertyName == GET_MEMBER_NAME_CHECKED(UDynamicMaterialEditorSettings, bPreviewImagesUseTextureUVs))
+	{
+		for (UDMMaterialComponent* Component : TObjectRange<UDMMaterialComponent>())
+		{
+			Component->MarkComponentDirty();
+		}
+	}
+}
+
+void UDynamicMaterialEditorSettings::OpenEditorSettingsWindow() const
+{
+	static ISettingsModule& SettingsModule = FModuleManager::LoadModuleChecked<ISettingsModule>("Settings");
+	SettingsModule.ShowViewer(GetContainerName(), CategoryName, SectionName);
+}
+
+void UDynamicMaterialEditorSettings::ResetAllLayoutSettings()
+{
+	SpinBoxValueMultiplier_AltShift = 1.0f; // Normal Sensitivity
+	SpinBoxValueMultiplier_Shift = 0.25f;   //   5x Default
+	SpinBoxValueMultiplier_Default = 0.05f; //  20x Sensitivity
+	SpinBoxValueMultiplier_Cmd = 0.01f;		//  1/5 Default
+	SpinBoxValueMultiplier_AltCmd = 0.001f; // 1/50 Default, 1/10 Cmd
+	SplitterLocation = 240;
+	LayerPreviewSize = 40;
+	SlotPreviewSize = 96;
+	DetailsPreviewSize = 96;
+	bShowTooltipPreview = false;
+	TooltipTextureSize = 512;
+	MaxFloatSliderWidth = 200.f;
+	bPreviewImagesUseTextureUVs = true;
+}
+
+float UDynamicMaterialEditorSettings::GetSpinboxValueChangeMultiplier(const FModifierKeysState& InModifierKeys) const
+{
+	const bool bIsShiftDown = InModifierKeys.IsLeftShiftDown() || InModifierKeys.IsRightShiftDown();
+	const bool bIsAltDown = InModifierKeys.IsLeftAltDown() || InModifierKeys.IsRightAltDown();
+	const bool bIsCmdDown = InModifierKeys.IsLeftControlDown() || InModifierKeys.IsLeftCommandDown() || InModifierKeys.IsRightControlDown() || InModifierKeys.IsRightCommandDown();
+	
+	if (InModifierKeys.IsShiftDown())
+	{
+		return InModifierKeys.IsAltDown() ? SpinBoxValueMultiplier_AltShift : SpinBoxValueMultiplier_Shift;
+	}
+
+	if (InModifierKeys.IsCommandDown() || InModifierKeys.IsControlDown())
+	{
+		return InModifierKeys.IsAltDown() ? SpinBoxValueMultiplier_AltCmd : SpinBoxValueMultiplier_Cmd;
+	}
+
+	return SpinBoxValueMultiplier_Default;
+}
+
+TArray<FDMMaterialEffectList> UDynamicMaterialEditorSettings::GetEffectList() const
+{
+	TArray<FDMMaterialEffectList> Effects = {};
+
+	IAssetRegistry* AssetRegistry = IAssetRegistry::Get();
+
+	if (!AssetRegistry)
+	{
+		return Effects;
+	}
+
+	TSharedPtr<IPlugin> Plugin = IPluginManager::Get().FindPlugin(UE_PLUGIN_NAME);
+
+	if (!Plugin.IsValid())
+	{
+		return Effects;
+	}
+
+	const FString PluginEffectPath = Plugin->GetMountedAssetPath() / "MaterialFunctions" / "Effects";
+
+	TArray<FName> AssetPaths = {*PluginEffectPath};
+	AssetPaths.Append(UDynamicMaterialEditorSettings::Get()->CustomEffectsFolders);
+
+	TArray<FAssetData> Assets;
+	AssetRegistry->GetAssetsByPaths(AssetPaths, Assets, /* bRecursive */ true, /* Only Assets on Disk */ true);
+
+	TArray<FString> AssetPathStrings;
+	AssetPathStrings.Reserve(AssetPaths.Num());
+	Algo::Transform(AssetPaths, AssetPathStrings, [](const FName& InElement) { return InElement.ToString(); });
+
+	for (FString& AssetPathString : AssetPathStrings)
+	{
+		if (AssetPathString.EndsWith(TEXT("/")) || AssetPathString.EndsWith(TEXT("\\")))
+		{
+			AssetPathString = AssetPathString.LeftChop(1);
+		}
+	}
+
+	auto FindBasePath = [&AssetPathStrings](const FString& InPath)
+		{
+			for (const FString& AssetPathString : AssetPathStrings)
+			{
+				if (InPath.StartsWith(AssetPathString))
+				{
+					return AssetPathString;
+				}
+			}
+
+			return FString("");
+		};
+
+	for (const FAssetData& Asset : Assets)
+	{
+		UClass* AssetClass = Asset.GetClass(EResolveClass::Yes);
+
+		if (!AssetClass || !AssetClass->IsChildOf(UMaterialFunctionInterface::StaticClass()))
+		{
+			continue;
+		}
+
+		const FString AssetPath = Asset.GetObjectPathString();
+		const FString AssetBasePath = FindBasePath(AssetPath);
+
+		FString Path = AssetPath;
+		FString Category = "";
+
+		// Reduce from /BasePath/Category/Random/OtherPaths/Asset.Asset to /BasePath and Category
+		while (true)
+		{
+			const FString ParentPath = FPaths::GetPath(Path);
+
+			if (ParentPath == AssetBasePath)
+			{
+				Category = Path.Mid(AssetBasePath.Len() + 1);
+				break;
+			}
+			else if (ParentPath.IsEmpty())
+			{
+				break;
+			}
+
+			Path = ParentPath;
+		}
+
+		if (Category.IsEmpty())
+		{
+			continue;
+		}
+
+		FDMMaterialEffectList* EffectList = Effects.FindByPredicate(
+			[Category](const FDMMaterialEffectList& InElement)
+			{
+				return InElement.Name == Category;
+			});
+
+		if (!EffectList)
+		{
+			Effects.Add({Category, {}});
+			EffectList = &Effects.Last();
+		}
+
+		TSoftObjectPtr<UMaterialFunctionInterface> SoftPtr;
+		SoftPtr = Asset.GetSoftObjectPath();
+
+		EffectList->Effects.Add(SoftPtr);
+	}
+
+	return Effects;
+}
+
+#undef LOCTEXT_NAMESPACE
