[MaterialDesigner]

diff --git a/Engine/Plugins/Experimental/DynamicMaterial/Source/DynamicMaterialEditor/Private/DynamicMaterialEditorStyle.cpp b/Engine/Plugins/Experimental/DynamicMaterial/Source/DynamicMaterialEditor/Private/DynamicMaterialEditorStyle.cpp
new file mode 100644
--- /dev/null
+++ b/Engine/Plugins/Experimental/DynamicMaterial/Source/DynamicMaterialEditor/Private/DynamicMaterialEditorStyle.cpp
@@ -0,0 +1,576 @@
+// Copyright Epic Games, Inc. All Rights Reserved.
+
+#include "DynamicMaterialEditorStyle.h"
+#include "Brushes/SlateImageBrush.h"
+#include "DetailLayoutBuilder.h"
+#include "DynamicMaterialEditorModule.h"
+#include "Engine/Texture2D.h"
+#include "Interfaces/IPluginManager.h"
+#include "Misc/Paths.h"
+#include "Styling/SlateStyleRegistry.h"
+#include "Styling/StyleColors.h"
+
+TSharedPtr<FSlateStyleSet> FDynamicMaterialEditorStyle::StyleInstance = nullptr;
+
+void FDynamicMaterialEditorStyle::Initialize()
+{
+	if (!StyleInstance.IsValid())
+	{
+		StyleInstance = Create();
+		FSlateStyleRegistry::RegisterSlateStyle(*StyleInstance);
+	}
+}
+
+void FDynamicMaterialEditorStyle::Shutdown()
+{
+	if (StyleInstance.IsValid())
+	{
+		FSlateStyleRegistry::UnRegisterSlateStyle(*StyleInstance);
+		ensure(StyleInstance.IsUnique());
+		StyleInstance.Reset();
+	}
+}
+
+FName FDynamicMaterialEditorStyle::GetStyleSetName()
+{
+	static FName StyleSetName(TEXT("DynamicMaterialEditorStyle"));
+	return StyleSetName;
+}
+
+const ISlateStyle& FDynamicMaterialEditorStyle::Get()
+{
+	if (!StyleInstance.IsValid())
+	{
+		Initialize();
+	}
+	return *StyleInstance;
+}
+
+FLinearColor FDynamicMaterialEditorStyle::GetColor(const FName& InName)
+{
+	return FDynamicMaterialEditorStyle::Get().GetColor(InName);
+}
+
+const FSlateBrush* FDynamicMaterialEditorStyle::GetBrush(const FName& InName)
+{
+	return FDynamicMaterialEditorStyle::Get().GetBrush(InName);
+}
+
+namespace UE::DynamicMaterialEditor::Private
+{
+	const TMap<EDMValueType, FName> TypeIcons = {
+		{EDMValueType::VT_None,        TEXT("Icons.Type.None")},
+		{EDMValueType::VT_Bool,        TEXT("Icons.Type.Bool")},
+		{EDMValueType::VT_Float1,      TEXT("Icons.Type.Float1")},
+		{EDMValueType::VT_Float2,      TEXT("Icons.Type.Float2")},
+		{EDMValueType::VT_Float3_RPY,  TEXT("Icons.Type.Float3_RPY")},
+		{EDMValueType::VT_Float3_RGB,  TEXT("Icons.Type.Float3_RGB")},
+		{EDMValueType::VT_Float3_XYZ,  TEXT("Icons.Type.Float3_XYZ")},
+		{EDMValueType::VT_Float4_RGBA, TEXT("Icons.Type.Float4_RGBA")},
+		{EDMValueType::VT_Float_Any,   TEXT("Icons.Type.Float_Any")},
+		{EDMValueType::VT_Texture,     TEXT("Icons.Type.Texture")},
+	};
+}
+
+FName FDynamicMaterialEditorStyle::GetBrushNameForType(EDMValueType InType)
+{
+	if (UE::DynamicMaterialEditor::Private::TypeIcons.Contains(InType))
+	{
+		return UE::DynamicMaterialEditor::Private::TypeIcons[InType];
+	}
+
+	checkNoEntry();
+
+	static FName NoType = TEXT("Icons.Type.None");
+	return NoType;
+}
+
+
+#define IMAGE_BRUSH(RelativePath, ... ) FSlateImageBrush(Style->RootToContentDir(RelativePath, TEXT(".png")), __VA_ARGS__)
+#define BOX_BRUSH(RelativePath, ... ) FSlateBoxBrush(Style->RootToContentDir(RelativePath, TEXT(".png")), __VA_ARGS__)
+#define BORDER_BRUSH(RelativePath, ... ) FSlateBorderBrush(Style->RootToContentDir(RelativePath, TEXT(".png")), __VA_ARGS__)
+
+#define IMAGE_BRUSH_SVG(RelativePath, ...) FSlateVectorImageBrush(Style->RootToContentDir(RelativePath, TEXT(".svg")), __VA_ARGS__)
+#define BOX_BRUSH_SVG(RelativePath, ...) FSlateVectorBoxBrush(Style->RootToContentDir(RelativePath, TEXT(".svg")), __VA_ARGS__)
+#define BORDER_BRUSH_SVG(RelativePath, ...) FSlateVectorBorderBrush(Style->RootToContentDir(RelativePath, TEXT(".svg")), __VA_ARGS__)
+
+#define DEFAULT_FONT(...) FCoreStyle::GetDefaultFontStyle(__VA_ARGS__)
+
+#define IMAGE_PLUGIN_BRUSH(RelativePath, ...) FSlateImageBrush(FDynamicMaterialEditorStyle::InContent(RelativePath, ".png"), __VA_ARGS__)
+#define IMAGE_PLUGIN_BRUSH_SVG(RelativePath, ...) FSlateVectorImageBrush(FDynamicMaterialEditorStyle::InContent(RelativePath, ".svg"), __VA_ARGS__)
+
+const FVector2D Icon8x8(8.f, 8.f);
+const FVector2D Icon12x12(12.f, 12.f);
+const FVector2D Icon16x16(16.f, 16.f);
+const FVector2D Icon20x20(20.f, 20.f);
+const FVector2D Icon24x24(24.f, 24.f);
+const FVector2D Icon32x32(32.f, 32.f);
+const FVector2D Icon40x40(40.f, 40.f);
+
+FLinearColor ReplaceColorAlpha(const FLinearColor& InColor, const float InNewAlpha)
+{
+	FLinearColor OutColor(InColor);
+	OutColor.A = InNewAlpha;
+	return OutColor;
+}
+TSharedRef<FSlateStyleSet> FDynamicMaterialEditorStyle::Create()
+{
+	TSharedRef<FSlateStyleSet> Style = MakeShared<FSlateStyleSet>(TEXT("DynamicMaterial"));
+
+	TSharedPtr<IPlugin> Plugin = IPluginManager::Get().FindPlugin(TEXT("DynamicMaterial"));
+	check(Plugin.IsValid());
+
+	Style->SetContentRoot(FPaths::Combine(Plugin->GetBaseDir(), TEXT("Resources")));
+
+	////////////////////////////////////////////////////////////////////////////////////////////////////
+	// Color Styles
+	const FLinearColor EngineSelectColor = FStyleColors::Primary.GetSpecifiedColor();
+	const FLinearColor EngineSelectHoverColor = FStyleColors::PrimaryHover.GetSpecifiedColor();
+	const FLinearColor EngineSelectPressColor = FStyleColors::PrimaryPress.GetSpecifiedColor();
+
+	const FLinearColor SelectColor = ReplaceColorAlpha(FStyleColors::Select.GetSpecifiedColor(), 0.9f);
+	const FLinearColor SelectHoverColor = FStyleColors::Select.GetSpecifiedColor();
+	const FLinearColor SelectPressColor = EngineSelectPressColor;
+
+	Style->Set("Color.Select", SelectColor);
+	Style->Set("Color.Select.Hover", SelectHoverColor);
+	Style->Set("Color.Select.Press", SelectPressColor);
+
+	////////////////////////////////////////////////////////////////////////////////////////////////////
+	// Brush Styles
+	Style->Set("Icons.Menu.Dropdown", new IMAGE_BRUSH_SVG("Icons/MenuDropdown", Icon16x16));
+
+	Style->Set("Icons.Type.None", new IMAGE_BRUSH("Icons/ValueTypes/None", Icon12x12));
+	Style->Set("Icons.Type.Bool", new IMAGE_BRUSH("Icons/ValueTypes/Bool", Icon12x12));
+	Style->Set("Icons.Type.Float1", new IMAGE_BRUSH("Icons/ValueTypes/Float1", Icon12x12));
+	Style->Set("Icons.Type.Float2", new IMAGE_BRUSH("Icons/ValueTypes/Float2", Icon12x12));
+	Style->Set("Icons.Type.Float3_RPY", new IMAGE_BRUSH("Icons/ValueTypes/Float3_RPY", Icon12x12));
+	Style->Set("Icons.Type.Float3_RGB", new IMAGE_BRUSH("Icons/ValueTypes/Float3_RGB", Icon12x12));
+	Style->Set("Icons.Type.Float3_XYZ", new IMAGE_BRUSH("Icons/ValueTypes/Float3_XYZ", Icon12x12));
+	Style->Set("Icons.Type.Float4_RGBA", new IMAGE_BRUSH("Icons/ValueTypes/Float4_RGBA", Icon12x12));
+	Style->Set("Icons.Type.Float_Any", new IMAGE_BRUSH("Icons/ValueTypes/Float_Any", Icon12x12));
+	Style->Set("Icons.Type.Texture", new IMAGE_BRUSH("Icons/ValueTypes/Texture", Icon12x12));
+
+	Style->Set("Icons.Material.DefaultLit", new IMAGE_BRUSH("Icons/EditorIcons/MaterialTypeDefaultLit", Icon32x32));
+	Style->Set("Icons.Material.Unlit", new IMAGE_BRUSH("Icons/EditorIcons/MaterialTypeUnlit", Icon32x32));
+
+	Style->Set("Icons.Lock", new IMAGE_BRUSH_SVG("Icons/EditorIcons/Lock", Icon16x16));
+	Style->Set("Icons.Unlock", new IMAGE_BRUSH_SVG("Icons/EditorIcons/Unlock", Icon16x16));
+
+	Style->Set("Icons.Remove", new IMAGE_BRUSH("Icons/EditorIcons/Remove_16px", Icon16x16));
+
+	Style->Set("Icons.Normalize", new IMAGE_BRUSH("Icons/EditorIcons/Normalize", Icon16x16));
+
+	Style->Set("Icons.Stage.EnabledButton", new IMAGE_BRUSH("Icons/EditorIcons/WhiteBall", Icon8x8));
+	Style->Set("Icons.Stage.BaseToggle", new IMAGE_BRUSH("Icons/EditorIcons/BaseToggle_16x", Icon16x16));
+	Style->Set("Icons.Stage.MaskToggle", new IMAGE_BRUSH("Icons/EditorIcons/MaskToggle_16x", Icon16x16));
+	Style->Set("Icons.Stage.Enabled", new IMAGE_BRUSH_SVG("Icons/EditorIcons/Enable", Icon24x24));
+	Style->Set("Icons.Stage.Disabled", new IMAGE_BRUSH_SVG("Icons/EditorIcons/Disable", Icon24x24));
+
+	Style->Set("Icons.Stage.ChainLinked", new IMAGE_BRUSH_SVG("Icons/EditorIcons/ChainLinked", Icon16x16));
+	Style->Set("Icons.Stage.ChainUnlinked", new IMAGE_BRUSH_SVG("Icons/EditorIcons/ChainUnlinked", Icon16x16));
+	Style->Set("Icons.Stage.ChainLinked.Horizontal", new IMAGE_BRUSH_SVG("Icons/EditorIcons/ChainLinked_Horizontal", Icon24x24));
+	Style->Set("Icons.Stage.ChainUnlinked.Horizontal", new IMAGE_BRUSH_SVG("Icons/EditorIcons/ChainUnlinked_Horizontal", Icon24x24));
+	Style->Set("Icons.Stage.ChainLinked.Vertical", new IMAGE_BRUSH_SVG("Icons/EditorIcons/ChainLinked_Vertical", Icon24x24));
+	Style->Set("Icons.Stage.ChainUnlinked.Vertical", new IMAGE_BRUSH_SVG("Icons/EditorIcons/ChainUnlinked_Vertical", Icon24x24));
+
+	Style->Set("ImageBorder", new FSlateRoundedBoxBrush(
+		FLinearColor::Transparent, 10.0f, 
+		FStyleColors::Panel.GetSpecifiedColor(), 2.0f));
+
+	Style->Set("Border.SinglePixel", new BORDER_BRUSH(TEXT("Images/Borders/Border_SinglePixel"), FMargin(1.0f / 4.0f)));
+	Style->Set("Border.LeftTopRight", new BORDER_BRUSH(TEXT("Images/Borders/Border_LeftTopRight"), FMargin(1.0f / 4.0f, 1.0f / 2.0f)));
+	Style->Set("Border.LeftBottomRight", new BORDER_BRUSH(TEXT("Images/Borders/Border_LeftBottomRight"), FMargin(1.0f / 4.0f, 1.0f / 2.0f)));
+	Style->Set("Border.TopLeftBottom", new BORDER_BRUSH(TEXT("Images/Borders/Border_TopLeftBottom"), FMargin(1.0f / 2.0f, 1.0f / 4.0f)));
+	Style->Set("Border.TopRightBottom", new BORDER_BRUSH(TEXT("Images/Borders/Border_TopRightBottom"), FMargin(1.0f / 2.0f, 1.0f / 4.0f)));
+	Style->Set("Border.Top", new BORDER_BRUSH(TEXT("Images/Borders/Border_Top"), FMargin(0.0f, 1.0f / 2.0f)));
+	Style->Set("Border.Bottom", new BORDER_BRUSH(TEXT("Images/Borders/Border_Bottom"), FMargin(0.0f, 1.0f / 2.0f)));
+	Style->Set("Border.Left", new BORDER_BRUSH(TEXT("Images/Borders/Border_Left"), FMargin(1.0f / 2.0f, 0.0f)));
+	Style->Set("Border.Right", new BORDER_BRUSH(TEXT("Images/Borders/Border_Right"), FMargin(1.0f / 2.0f, 0.0f)));
+
+	////////////////////////////////////////////////////////////////////////////////////////////////////
+	// Button Styles
+	Style->Set("HoverHintOnly", FButtonStyle()
+		.SetNormal(FSlateNoResource())
+		.SetHovered(FSlateRoundedBoxBrush(FLinearColor(1, 1, 1, 0.15f), 4.0f))
+		.SetPressed(FSlateRoundedBoxBrush(FLinearColor(1, 1, 1, 0.25f), 4.0f))
+		.SetNormalPadding(FMargin(0, 0, 0, 1))
+		.SetPressedPadding(FMargin(0, 1, 0, 0)));
+
+	Style->Set("HoverHintOnly.Bordered", FButtonStyle()
+		.SetNormal(FSlateRoundedBoxBrush(FLinearColor::Transparent, 4.0f, FLinearColor(1, 1, 1, 0.25f), 1.0f))
+		.SetHovered(FSlateRoundedBoxBrush(FLinearColor(1, 1, 1, 0.15f), 4.0f, FLinearColor(1, 1, 1, 0.4f), 1.0f))
+		.SetPressed(FSlateRoundedBoxBrush(FLinearColor(1, 1, 1, 0.25f), 4.0f, FLinearColor(1, 1, 1, 0.5f), 1.0f))
+		.SetNormalPadding(FMargin(0, 0, 0, 1))
+		.SetPressedPadding(FMargin(0, 1, 0, 0)));
+
+	Style->Set("HoverHintOnly.Bordered.Dark", FButtonStyle()
+		.SetNormal(FSlateRoundedBoxBrush(FLinearColor::Transparent, 4.0f, FStyleColors::InputOutline.GetSpecifiedColor(), 1.0f))
+		.SetHovered(FSlateRoundedBoxBrush(FLinearColor(1, 1, 1, 0.15f), 4.0f, FLinearColor(1, 1, 1, 0.4f), 1.0f))
+		.SetPressed(FSlateRoundedBoxBrush(FLinearColor(1, 1, 1, 0.25f), 4.0f, FLinearColor(1, 1, 1, 0.5f), 1.0f))
+		.SetNormalPadding(FMargin(0, 0, 0, 1))
+		.SetPressedPadding(FMargin(0, 1, 0, 0)));
+
+	SetupStageStyles(Style);
+	SetupLayerViewStyles(Style);
+	SetupLayerViewItemHandleStyles(Style);
+	SetupEffectsViewStyles(Style);
+	SetupTextStyles(Style);
+
+	//////////////////////////////////////////////////////
+	/// Editable TextBox Style
+	FEditableTextBoxStyle InlineEditableTextBoxStyle;
+	InlineEditableTextBoxStyle.SetPadding(FMargin(0));
+	InlineEditableTextBoxStyle.SetBackgroundColor(FSlateColor(FLinearColor::Transparent));
+
+	Style->Set("InlineEditableTextBoxStyle", InlineEditableTextBoxStyle);
+
+	return Style;
+}
+
+void FDynamicMaterialEditorStyle::SetupStageStyles(const TSharedRef<FSlateStyleSet>& Style)
+{
+	constexpr FLinearColor StageEnabledColor = FLinearColor(0.0f, 1.0f, 0.0f, 0.7f);
+	constexpr FLinearColor StageDisabledColor = FLinearColor(1.0f, 0.0f, 0.0f, 0.7f);
+
+	Style->Set("Color.Stage.Enabled", StageEnabledColor);
+	Style->Set("Color.Stage.Disabled", StageDisabledColor);
+
+	const float StageCornerRadius = 6.0f;
+	const float StageBorderWidth = 2.0f;
+
+	Style->Set("Stage.Inactive", new FSlateRoundedBoxBrush(
+		FLinearColor::Transparent, StageCornerRadius,
+		FStyleColors::Panel.GetSpecifiedColor(), 1.0f));
+	Style->Set("Stage.Inactive.Hover", new FSlateRoundedBoxBrush(
+		FLinearColor::Transparent, StageCornerRadius,
+		ReplaceColorAlpha(FStyleColors::Foreground.GetSpecifiedColor(), 0.2f), 1.0f));
+	Style->Set("Stage.Inactive.Select", new FSlateRoundedBoxBrush(
+		FLinearColor::Transparent, StageCornerRadius,
+		Style->GetColor("Color.Select"), 2.0f));
+	Style->Set("Stage.Inactive.Select.Hover", new FSlateRoundedBoxBrush(
+		FLinearColor::Transparent, StageCornerRadius,
+		Style->GetColor("Color.Select.Hover"), 2.0f));
+
+	Style->Set("Stage.Disabled", new FSlateRoundedBoxBrush(
+		FLinearColor::Transparent, StageCornerRadius,
+		ReplaceColorAlpha(StageDisabledColor, 0.8f), 1.0f));
+	Style->Set("Stage.Disabled.Hover", new FSlateRoundedBoxBrush(
+		FLinearColor::Transparent, StageCornerRadius,
+		StageDisabledColor, 1.0f));
+	Style->Set("Stage.Disabled.Select", new FSlateRoundedBoxBrush(
+		FLinearColor::Transparent, StageCornerRadius,
+		ReplaceColorAlpha(FStyleColors::Select.GetSpecifiedColor(), 0.9f), 2.0f));
+	Style->Set("Stage.Disabled.Select.Hover", new FSlateRoundedBoxBrush(
+		FLinearColor::Transparent, StageCornerRadius,
+		FStyleColors::Select.GetSpecifiedColor(), 2.0f));
+
+	Style->Set("Stage.Outline", new FSlateRoundedBoxBrush(
+		FStyleColors::InputOutline.GetSpecifiedColor(), 1.0f));
+	Style->Set("Stage.Drag", new FSlateRoundedBoxBrush(
+		FLinearColor::Transparent, StageCornerRadius,
+		FStyleColors::AccentYellow.GetSpecifiedColor(), 2.0f));
+
+	Style->Set("Stage.TextBlock", new FSlateRoundedBoxBrush(
+		ReplaceColorAlpha(FStyleColors::Panel.GetSpecifiedColor(), 0.7f), 0.0f,
+		ReplaceColorAlpha(FStyleColors::InputOutline.GetSpecifiedColor(), 0.4f), 1.0f));
+}
+
+void FDynamicMaterialEditorStyle::SetupLayerViewStyles(const TSharedRef<FSlateStyleSet>& Style)
+{
+	Style->Set("LayerView.Background", new FSlateRoundedBoxBrush(
+		FStyleColors::Panel.GetSpecifiedColor()/*FLinearColor(0, 0, 0, 0.25f)*/, 6.0f,
+		FStyleColors::Header.GetSpecifiedColor()/*FLinearColor(1, 1, 1, 0.2f)*/, 1.0f));
+
+	Style->Set("LayerView.Details.Background", new FSlateRoundedBoxBrush(
+		FStyleColors::Recessed.GetSpecifiedColor()/*FLinearColor(0, 0, 0, 0.25f)*/, 6.0f,
+		FStyleColors::Header.GetSpecifiedColor()/*FLinearColor(1, 1, 1, 0.2f)*/, 1.0f));
+
+	/**
+	 * SListView and and FTableViewStyle have no support for adding padding between the background brush
+	 * and the SListView widget, so we are not using this style for the SDMSlotLayerView. Instead, we add a
+	 * SBorder around the SDMBLayerView and style that.
+	 */
+	Style->Set("LayerView", FTableViewStyle()
+		.SetBackgroundBrush(*Style->GetBrush("LayerView.Background"))
+	);
+
+	const float LayerViewItemCornerRadius = 5.0f;
+	const float LayerViewItemBorderWidth = 1.0f;
+
+	const FLinearColor LayerViewItemFillColor = FLinearColor::Transparent;
+	constexpr FLinearColor LayerViewItemBorderColor = FLinearColor(1, 1, 1, 0.15f);
+
+	const FLinearColor LayerItemHoverFillColor = FStyleColors::Recessed.GetSpecifiedColor();
+	constexpr FLinearColor LayerItemHoverBorderColor = FLinearColor(1, 1, 1, 0.2f);
+
+	const FLinearColor LayerItemSelectFillColor = FStyleColors::Header.GetSpecifiedColor();
+	const FLinearColor LayerItemSelectBorderColor = ReplaceColorAlpha(FStyleColors::Select.GetSpecifiedColor(), 0.9f);
+
+	//constexpr FLinearColor LayerItemSelectHoverFillColor = FLinearColor(1, 1, 1, 0.25f);
+	//const FLinearColor LayerItemSelectHoverBorderColor = FStyleColors::Select.GetSpecifiedColor();
+
+	Style->Set("LayerView.Row.Item", new FSlateRoundedBoxBrush(
+		LayerViewItemFillColor, LayerViewItemCornerRadius,
+		LayerViewItemBorderColor, LayerViewItemBorderWidth));
+
+	Style->Set("LayerView.Row.Hovered", new FSlateRoundedBoxBrush(
+		LayerItemHoverFillColor, LayerViewItemCornerRadius,
+		LayerItemHoverBorderColor, LayerViewItemBorderWidth));
+
+	Style->Set("LayerView.Row.Selected", new FSlateRoundedBoxBrush(
+		LayerItemSelectFillColor, LayerViewItemCornerRadius,
+		LayerItemSelectBorderColor, LayerViewItemBorderWidth));
+
+	Style->Set("LayerView.Row.ActiveBrush", new FSlateRoundedBoxBrush(
+		LayerItemSelectFillColor, LayerViewItemCornerRadius,
+		LayerItemSelectBorderColor, LayerViewItemBorderWidth));
+	Style->Set("LayerView.Row.ActiveHoveredBrush", new FSlateRoundedBoxBrush(
+		LayerItemSelectFillColor, LayerViewItemCornerRadius,
+		LayerItemSelectBorderColor, LayerViewItemBorderWidth));
+	Style->Set("LayerView.Row.InactiveBrush", new FSlateRoundedBoxBrush(
+		LayerItemSelectFillColor, LayerViewItemCornerRadius,
+		LayerItemSelectBorderColor, LayerViewItemBorderWidth));
+	Style->Set("LayerView.Row.InactiveHoveredBrush", new FSlateRoundedBoxBrush(
+		LayerItemSelectFillColor, LayerViewItemCornerRadius,
+		LayerItemSelectBorderColor, LayerViewItemBorderWidth));
+
+	Style->Set("LayerView.Row", FTableRowStyle()
+		.SetTextColor(FStyleColors::Foreground)
+		.SetSelectedTextColor(FStyleColors::ForegroundHover)
+		.SetEvenRowBackgroundBrush(*Style->GetBrush(TEXT("LayerView.Row.Item")))
+		.SetEvenRowBackgroundHoveredBrush(*Style->GetBrush(TEXT("LayerView.Row.Hovered")))
+		.SetOddRowBackgroundBrush(*Style->GetBrush(TEXT("LayerView.Row.Item")))
+		.SetOddRowBackgroundHoveredBrush(*Style->GetBrush(TEXT("LayerView.Row.Hovered")))
+		.SetSelectorFocusedBrush(*Style->GetBrush(TEXT("LayerView.Row.Selected")))
+		.SetActiveBrush(*Style->GetBrush(TEXT("LayerView.Row.ActiveBrush")))
+		.SetActiveHoveredBrush(*Style->GetBrush(TEXT("LayerView.Row.ActiveHoveredBrush")))
+		.SetInactiveBrush(*Style->GetBrush(TEXT("LayerView.Row.InactiveBrush")))
+		.SetInactiveHoveredBrush(*Style->GetBrush(TEXT("LayerView.Row.InactiveHoveredBrush")))
+		.SetSelectorFocusedBrush(BORDER_BRUSH("Images/DropIndicators/DropZoneIndicator_Onto", FMargin(4.f / 16.f), Style->GetColor(TEXT("Color.Select.Hover"))))
+		.SetDropIndicator_Onto(BOX_BRUSH("Images/DropIndicators/LayerView_DropIndicator_Onto", FMargin(4.0f / 16.0f), Style->GetColor(TEXT("Color.Select.Hover"))))
+		.SetDropIndicator_Above(BOX_BRUSH("Images/DropIndicators/LayerView_DropIndicator_Above", FMargin(4.0f / 16.0f, 4.0f / 16.0f, 0.f, 0.f), Style->GetColor(TEXT("Color.Select.Hover"))))
+		.SetDropIndicator_Below(BOX_BRUSH("Images/DropIndicators/LayerView_DropIndicator_Below", FMargin(4.0f / 16.0f, 0.f, 0.f, 4.0f / 16.0f), Style->GetColor(TEXT("Color.Select.Hover"))))
+	);
+
+	Style->Set("LayerView.AddIcon", new IMAGE_BRUSH("Icons/EditorIcons/LayerAdd", Icon16x16));
+	Style->Set("LayerView.DuplicateIcon", new IMAGE_BRUSH("Icons/EditorIcons/Duplicate_40x", Icon40x40));
+	Style->Set("LayerView.RemoveIcon", new IMAGE_BRUSH("Icons/EditorIcons/LayerRemove", Icon16x16));
+
+	Style->Set("LayerView.Row.Handle", new IMAGE_BRUSH_SVG("Icons/DragHandle", Icon16x16));
+}
+
+void FDynamicMaterialEditorStyle::SetupLayerViewItemHandleStyles(const TSharedRef<FSlateStyleSet>& Style)
+{
+	constexpr FLinearColor RowHandleFillColor = FLinearColor(1, 1, 1, 0.3f);
+	const FLinearColor RowHandleHoverFillColor = FLinearColor(1, 1, 1, 0.4f);
+	const FLinearColor RowHandleBorderColor = FLinearColor::Transparent;
+	constexpr float RowHandleCornerRadius = 6.0f;
+	constexpr float RowHandleBorderWidth = 1.0f;
+
+	Style->Set("LayerView.Row.Handle.Left", new FSlateRoundedBoxBrush(
+		RowHandleFillColor, FVector4(RowHandleCornerRadius, 0.0f, 0.0f, RowHandleCornerRadius),
+		RowHandleBorderColor, RowHandleBorderWidth));
+	Style->Set("LayerView.Row.Handle.Top", new FSlateRoundedBoxBrush(
+		RowHandleFillColor, FVector4(RowHandleCornerRadius, RowHandleCornerRadius, 0.0f, 0.0f),
+		RowHandleBorderColor, RowHandleBorderWidth));
+	Style->Set("LayerView.Row.Handle.Right", new FSlateRoundedBoxBrush(
+		RowHandleFillColor, FVector4(0.0f, RowHandleCornerRadius, RowHandleCornerRadius, 0.0f),
+		RowHandleBorderColor, RowHandleBorderWidth));
+	Style->Set("LayerView.Row.Handle.Bottom", new FSlateRoundedBoxBrush(
+		RowHandleFillColor, FVector4(0.0f, 0.0f, RowHandleCornerRadius, RowHandleCornerRadius),
+		RowHandleBorderColor, RowHandleBorderWidth));
+
+	Style->Set("LayerView.Row.Handle.Left.Hover", new FSlateRoundedBoxBrush(
+		RowHandleHoverFillColor, FVector4(RowHandleCornerRadius, 0.0f, 0.0f, RowHandleCornerRadius),
+		RowHandleBorderColor, RowHandleBorderWidth));
+	Style->Set("LayerView.Row.Handle.Top.Hover", new FSlateRoundedBoxBrush(
+		RowHandleHoverFillColor, FVector4(RowHandleCornerRadius, RowHandleCornerRadius, 0.0f, 0.0f),
+		RowHandleBorderColor, RowHandleBorderWidth));
+	Style->Set("LayerView.Row.Handle.Right.Hover", new FSlateRoundedBoxBrush(
+		RowHandleHoverFillColor, FVector4(0.0f, RowHandleCornerRadius, RowHandleCornerRadius, 0.0f),
+		RowHandleBorderColor, RowHandleBorderWidth));
+	Style->Set("LayerView.Row.Handle.Bottom.Hover", new FSlateRoundedBoxBrush(
+		RowHandleHoverFillColor, FVector4(0.0f, 0.0f, RowHandleCornerRadius, RowHandleCornerRadius),
+		RowHandleBorderColor, RowHandleBorderWidth));
+
+	Style->Set("LayerView.Row.Handle.Left.Select", new FSlateRoundedBoxBrush(
+		Style->GetColor(TEXT("Color.Select")), 0.0f,
+		RowHandleBorderColor, 0.0f));
+	Style->Set("LayerView.Row.Handle.Top.Select", new FSlateRoundedBoxBrush(
+		Style->GetColor(TEXT("Color.Select")), 0.0f,
+		RowHandleBorderColor, 0.0f));
+	Style->Set("LayerView.Row.Handle.Right.Select", new FSlateRoundedBoxBrush(
+		Style->GetColor(TEXT("Color.Select")), 0.0f,
+		RowHandleBorderColor, 0.0f));
+	Style->Set("LayerView.Row.Handle.Bottom.Select", new FSlateRoundedBoxBrush(
+		Style->GetColor(TEXT("Color.Select")), 0.0f,
+		RowHandleBorderColor, 0.0f));
+
+	Style->Set("LayerView.Row.Handle.Left.Select.Hover", new FSlateRoundedBoxBrush(
+		Style->GetColor(TEXT("Color.Select.Hover")), 0.0f,
+		RowHandleBorderColor, 0.0f));
+	Style->Set("LayerView.Row.Handle.Top.Select.Hover", new FSlateRoundedBoxBrush(
+		Style->GetColor(TEXT("Color.Select.Hover")), 0.0f,
+		RowHandleBorderColor, 0.0f));
+	Style->Set("LayerView.Row.Handle.Right.Select.Hover", new FSlateRoundedBoxBrush(
+		Style->GetColor(TEXT("Color.Select.Hover")), 0.0f,
+		RowHandleBorderColor, 0.0f));
+	Style->Set("LayerView.Row.Handle.Bottom.Select.Hover", new FSlateRoundedBoxBrush(
+		Style->GetColor(TEXT("Color.Select.Hover")), 0.0f,
+		RowHandleBorderColor, 0.0f));
+}
+
+void FDynamicMaterialEditorStyle::SetupEffectsViewStyles(const TSharedRef<FSlateStyleSet>& Style)
+{
+	Style->Set("EffectsView.Background", new FSlateRoundedBoxBrush(
+		FStyleColors::Panel.GetSpecifiedColor(), 0.0f,
+		FStyleColors::Header.GetSpecifiedColor(), 1.0f));
+
+	Style->Set("EffectsView.Details.Background", new FSlateRoundedBoxBrush(
+		FStyleColors::Recessed.GetSpecifiedColor(), 0.0f,
+		FStyleColors::Header.GetSpecifiedColor(), 1.0f));
+
+	Style->Set("EffectsView", FTableViewStyle()
+		.SetBackgroundBrush(*Style->GetBrush(TEXT("EffectsView.Background")))
+	);
+
+	const float LayerViewItemCornerRadius = 0.0f;
+	const float LayerViewItemBorderWidth = 1.0f;
+
+	const FLinearColor LayerViewItemFillColor = FLinearColor::Transparent;
+	constexpr FLinearColor LayerViewItemBorderColor = FLinearColor(1, 1, 1, 0.15f);
+
+	const FLinearColor LayerItemHoverFillColor = FStyleColors::Recessed.GetSpecifiedColor();
+	constexpr FLinearColor LayerItemHoverBorderColor = FLinearColor(1, 1, 1, 0.2f);
+
+	const FLinearColor LayerItemSelectFillColor = FStyleColors::Header.GetSpecifiedColor();
+	const FLinearColor LayerItemSelectBorderColor = ReplaceColorAlpha(FStyleColors::Select.GetSpecifiedColor(), 0.9f);
+
+	Style->Set("EffectsView.Row.Item", new FSlateRoundedBoxBrush(
+		LayerViewItemFillColor, LayerViewItemCornerRadius,
+		LayerViewItemBorderColor, LayerViewItemBorderWidth));
+
+	Style->Set("EffectsView.Row.Hovered", new FSlateRoundedBoxBrush(
+		LayerItemHoverFillColor, LayerViewItemCornerRadius,
+		LayerItemHoverBorderColor, LayerViewItemBorderWidth));
+
+	Style->Set("EffectsView.Row.Selected", new FSlateRoundedBoxBrush(
+		LayerItemSelectFillColor, LayerViewItemCornerRadius,
+		LayerItemSelectBorderColor, LayerViewItemBorderWidth));
+
+	Style->Set("EffectsView.Row.ActiveBrush", new FSlateRoundedBoxBrush(
+		LayerItemSelectFillColor, LayerViewItemCornerRadius,
+		LayerItemSelectBorderColor, LayerViewItemBorderWidth));
+	Style->Set("EffectsView.Row.ActiveHoveredBrush", new FSlateRoundedBoxBrush(
+		LayerItemSelectFillColor, LayerViewItemCornerRadius,
+		LayerItemSelectBorderColor, LayerViewItemBorderWidth));
+	Style->Set("EffectsView.Row.InactiveBrush", new FSlateRoundedBoxBrush(
+		LayerItemSelectFillColor, LayerViewItemCornerRadius,
+		LayerItemSelectBorderColor, LayerViewItemBorderWidth));
+	Style->Set("EffectsView.Row.InactiveHoveredBrush", new FSlateRoundedBoxBrush(
+		LayerItemSelectFillColor, LayerViewItemCornerRadius,
+		LayerItemSelectBorderColor, LayerViewItemBorderWidth));
+
+	Style->Set("EffectsView.Row", FTableRowStyle()
+		.SetTextColor(FStyleColors::Foreground)
+		.SetSelectedTextColor(FStyleColors::ForegroundHover)
+		.SetEvenRowBackgroundBrush(*Style->GetBrush(TEXT("EffectsView.Row.Item")))
+		.SetEvenRowBackgroundHoveredBrush(*Style->GetBrush(TEXT("EffectsView.Row.Hovered")))
+		.SetOddRowBackgroundBrush(*Style->GetBrush(TEXT("EffectsView.Row.Item")))
+		.SetOddRowBackgroundHoveredBrush(*Style->GetBrush(TEXT("EffectsView.Row.Hovered")))
+		.SetSelectorFocusedBrush(*Style->GetBrush(TEXT("EffectsView.Row.Selected")))
+		.SetActiveBrush(*Style->GetBrush(TEXT("EffectsView.Row.ActiveBrush")))
+		.SetActiveHoveredBrush(*Style->GetBrush(TEXT("EffectsView.Row.ActiveHoveredBrush")))
+		.SetInactiveBrush(*Style->GetBrush(TEXT("EffectsView.Row.InactiveBrush")))
+		.SetInactiveHoveredBrush(*Style->GetBrush(TEXT("EffectsView.Row.InactiveHoveredBrush")))
+		.SetSelectorFocusedBrush(BORDER_BRUSH("Images/DropIndicators/DropZoneIndicator_Onto", FMargin(4.f / 16.f), Style->GetColor(TEXT("Color.Select.Hover"))))
+		.SetDropIndicator_Onto(BOX_BRUSH("Images/DropIndicators/LayerView_DropIndicator_Onto", FMargin(4.0f / 16.0f), Style->GetColor(TEXT("Color.Select.Hover"))))
+		.SetDropIndicator_Above(BOX_BRUSH("Images/DropIndicators/LayerView_DropIndicator_Above", FMargin(4.0f / 16.0f, 4.0f / 16.0f, 0.f, 0.f), Style->GetColor(TEXT("Color.Select.Hover"))))
+		.SetDropIndicator_Below(BOX_BRUSH("Images/DropIndicators/LayerView_DropIndicator_Below", FMargin(4.0f / 16.0f, 0.f, 0.f, 4.0f / 16.0f), Style->GetColor(TEXT("Color.Select.Hover"))))
+	);
+
+	Style->Set("EffectsView.Row.Fx.Closed", new IMAGE_BRUSH_SVG("Icons/Fx_Closed", Icon24x24));
+	Style->Set("EffectsView.Row.Fx.Opened", new IMAGE_BRUSH_SVG("Icons/Fx_Opened", Icon24x24));
+}
+
+void FDynamicMaterialEditorStyle::SetupTextStyles(const TSharedRef<FSlateStyleSet>& Style)
+{
+	const FTextBlockStyle NormalTextStyle(FAppStyle::GetWidgetStyle<FTextBlockStyle>(TEXT("NormalText")));
+
+	const FLinearColor LayerViewItemTextShadowColor = FLinearColor(0.0f, 0.0f, 0.0f, 0.5f);
+
+	FFontOutlineSettings HandleFontOutline;
+	HandleFontOutline.OutlineColor = LayerViewItemTextShadowColor;
+	HandleFontOutline.OutlineSize = 1;
+
+	FTextBlockStyle ActorNameTextStyle(NormalTextStyle);
+	ActorNameTextStyle.SetFont(DEFAULT_FONT("Regular", 10));
+	Style->Set("ActorName", ActorNameTextStyle);
+
+	FTextBlockStyle ActorNameBigTextStyle(NormalTextStyle);
+	ActorNameBigTextStyle.SetFont(DEFAULT_FONT("Regular", 14));
+	Style->Set("ActorNameBig", ActorNameBigTextStyle);
+
+	FTextBlockStyle ComponentNameBigTextStyle(NormalTextStyle);
+	ComponentNameBigTextStyle.SetFont(DEFAULT_FONT("Regular", 12));
+	Style->Set("ComponentNameBig", ComponentNameBigTextStyle);
+
+	FTextBlockStyle SlotLayerInfoTextStyle(NormalTextStyle);
+	SlotLayerInfoTextStyle.SetFont(DEFAULT_FONT("Italic", 8));
+	Style->Set("SlotLayerInfo", SlotLayerInfoTextStyle);
+
+	FSlateFontInfo LayerViewItemFont = DEFAULT_FONT("Bold", 12);
+	LayerViewItemFont.OutlineSettings = HandleFontOutline;
+	Style->Set("LayerView.Row.Font", LayerViewItemFont);
+
+	FSlateFontInfo LayerViewItemHandleSmallFont = DEFAULT_FONT("Regular", 10);
+	LayerViewItemHandleSmallFont.OutlineSettings = HandleFontOutline;
+	Style->Set("LayerView.Row.HandleFont", LayerViewItemHandleSmallFont);
+
+	FTextBlockStyle LayerViewItemTextStyle(NormalTextStyle);
+	LayerViewItemTextStyle.SetShadowOffset(FVector2D(1.0f, 1.0f));
+	LayerViewItemTextStyle.SetColorAndOpacity(LayerViewItemTextShadowColor);
+
+	Style->Set("LayerView.Row.HeaderText",
+		FTextBlockStyle(LayerViewItemTextStyle)
+		.SetColorAndOpacity(FStyleColors::Foreground)
+		.SetFont(LayerViewItemFont));
+
+	Style->Set("LayerView.Row.HeaderText.Small",
+		FTextBlockStyle(LayerViewItemTextStyle)
+		.SetColorAndOpacity(FStyleColors::Foreground)
+		.SetFont(LayerViewItemHandleSmallFont));
+
+	FTextBlockStyle StagePropertyDetailsTextStyle(NormalTextStyle);
+	StagePropertyDetailsTextStyle.SetFont(DEFAULT_FONT("Regular", 12));
+	Style->Set("Font.Stage.Details", StagePropertyDetailsTextStyle);
+
+	FTextBlockStyle StagePropertyDetailsBoldTextStyle(NormalTextStyle);
+	StagePropertyDetailsBoldTextStyle.SetFont(DEFAULT_FONT("Bold", 12));
+	Style->Set("Font.Stage.Details.Bold", StagePropertyDetailsBoldTextStyle);
+
+	FTextBlockStyle StagePropertyDetailsSmallTextStyle(NormalTextStyle);
+	StagePropertyDetailsSmallTextStyle.SetFont(IDetailLayoutBuilder::GetDetailFont());
+	Style->Set("Font.Stage.Details.Small", StagePropertyDetailsSmallTextStyle);
+
+	FTextBlockStyle StagePropertyDetailsSmallBoldTextStyle(NormalTextStyle);
+	StagePropertyDetailsSmallBoldTextStyle.SetFont(IDetailLayoutBuilder::GetDetailFontBold());
+	Style->Set("Font.Stage.Details.Small.Bold", StagePropertyDetailsSmallBoldTextStyle);
+
+	Style->Set("Font.Stage.Details.Small.Bold", FTextBlockStyle(NormalTextStyle)
+		.SetFont(IDetailLayoutBuilder::GetDetailFontBold()));
+}
+
+#undef IMAGE_BRUSH
+#undef BOX_BRUSH
+#undef BORDER_BRUSH
+
+#undef IMAGE_BRUSH_SVG
+#undef BOX_BRUSH_SVG
+#undef BORDER_BRUSH_SVG
+
+#undef DEFAULT_FONT
+
+#undef IMAGE_PLUGIN_BRUSH
+#undef IMAGE_PLUGIN_BRUSH_SVG
diff --git a/Engine/Plugins/Experimental/DynamicMaterial/Source/DynamicMaterialEditor/Private/DynamicMaterialEditorStyle.h b/Engine/Plugins/Experimental/DynamicMaterial/Source/DynamicMaterialEditor/Private/DynamicMaterialEditorStyle.h
new file mode 100644
--- /dev/null
+++ b/Engine/Plugins/Experimental/DynamicMaterial/Source/DynamicMaterialEditor/Private/DynamicMaterialEditorStyle.h
@@ -0,0 +1,34 @@
+// Copyright Epic Games, Inc. All Rights Reserved.
+
+#pragma once
+
+#include "Styling/ISlateStyle.h"
+#include "Styling/SlateStyle.h"
+
+enum class EDMValueType : uint8;
+
+class FDynamicMaterialEditorStyle
+{
+public:
+	static void Initialize();
+	static void Shutdown();
+
+	static FName GetStyleSetName();
+	static const ISlateStyle& Get();
+
+	static FLinearColor GetColor(const FName& InName);
+	static const FSlateBrush* GetBrush(const FName& InName);
+
+	static FName GetBrushNameForType(EDMValueType InType);
+
+private:
+	static TSharedPtr<FSlateStyleSet> StyleInstance;
+
+	static TSharedRef<FSlateStyleSet> Create();
+
+	static void SetupStageStyles(const TSharedRef<FSlateStyleSet>& Style);
+	static void SetupLayerViewStyles(const TSharedRef<FSlateStyleSet>& Style);
+	static void SetupLayerViewItemHandleStyles(const TSharedRef<FSlateStyleSet>& Style);
+	static void SetupEffectsViewStyles(const TSharedRef<FSlateStyleSet>& Style);
+	static void SetupTextStyles(const TSharedRef<FSlateStyleSet>& Style);
+};
diff --git a/Engine/Plugins/Experimental/DynamicMaterial/Source/DynamicMaterialEditor/Private/Material/DynamicMaterialInstanceFactory.cpp b/Engine/Plugins/Experimental/DynamicMaterial/Source/DynamicMaterialEditor/Private/Material/DynamicMaterialInstanceFactory.cpp
new file mode 100644
--- /dev/null
+++ b/Engine/Plugins/Experimental/DynamicMaterial/Source/DynamicMaterialEditor/Private/Material/DynamicMaterialInstanceFactory.cpp
@@ -0,0 +1,74 @@
+// Copyright Epic Games, Inc. All Rights Reserved.
+
+#include "Material/DynamicMaterialInstanceFactory.h"
+#include "AssetToolsModule.h"
+#include "DMPrivate.h"
+#include "IAssetTools.h"
+#include "Material/DynamicMaterialInstance.h"
+#include "Model/DynamicMaterialModel.h"
+#include "Model/IDynamicMaterialModelEditorOnlyDataInterface.h"
+#include "Model/DynamicMaterialModelFactory.h"
+#include "Modules/ModuleManager.h"
+
+#define LOCTEXT_NAMESPACE "MaterialDesignerInstanceFactory"
+
+const FString UDynamicMaterialInstanceFactory::BaseDirectory = FString("/Game/DynamicMaterials");
+const FString UDynamicMaterialInstanceFactory::BaseName = FString("M_DynMatInst_");
+
+UDynamicMaterialInstanceFactory::UDynamicMaterialInstanceFactory()
+{
+	SupportedClass = UDynamicMaterialInstance::StaticClass();
+
+	bCreateNew = true;
+	bEditAfterNew = true;
+	bEditorImport = false;
+	bText = false;
+}
+
+UObject* UDynamicMaterialInstanceFactory::FactoryCreateNew(UClass* Class, UObject* InParent, FName Name, 
+	EObjectFlags Flags, UObject* Context, FFeedbackContext* Warn)
+{
+	check(Class->IsChildOf(UDynamicMaterialInstance::StaticClass()));
+
+	UDynamicMaterialInstance* NewInstance = NewObject<UDynamicMaterialInstance>(InParent, Class, Name, Flags | RF_Transactional);
+	check(NewInstance);
+
+	UDynamicMaterialModelFactory* EditorFactory = NewObject<UDynamicMaterialModelFactory>();
+	check(EditorFactory);
+
+	UDynamicMaterialModel* NewModel = Cast<UDynamicMaterialModel>(EditorFactory->FactoryCreateNew(
+		UDynamicMaterialModel::StaticClass(), NewInstance, NAME_None, RF_Transactional, nullptr, GWarn));
+	check(NewModel);
+
+	const FDMInitializationGuard InitGuard;
+
+	NewInstance->SetMaterialModel(NewModel);
+
+	NewModel->SetDynamicMaterialInstance(NewInstance);
+
+	if (IDynamicMaterialModelEditorOnlyDataInterface* ModelEditorOnlyData = NewModel->GetEditorOnlyData())
+	{
+		ModelEditorOnlyData->RequestMaterialBuild();
+	}
+
+	NewInstance->InitializeMIDPublic();
+
+	IAssetTools& AssetTools = FModuleManager::LoadModuleChecked<FAssetToolsModule>("AssetTools").Get();
+	AssetTools.OpenEditorForAssets({NewInstance});
+
+	return NewInstance;
+}
+
+FText UDynamicMaterialInstanceFactory::GetDisplayName() const
+{
+	static const FText DisplayName = LOCTEXT("MaterialDesignerInstance", "Material Designer Instance");
+	return DisplayName;
+}
+
+FText UDynamicMaterialInstanceFactory::GetToolTip() const
+{
+	static const FText Tooltip = LOCTEXT("MaterialDesignerInstanceTooltip", "The Material Designer Instance is a combination of a Material Instance Dyanmic and a Material Designer Model.");
+	return Tooltip;
+}
+
+#undef LOCTEXT_NAMESPACE
diff --git a/Engine/Plugins/Experimental/DynamicMaterial/Source/DynamicMaterialEditor/Private/Menus/DMMaterialSlotLayerMenus.cpp b/Engine/Plugins/Experimental/DynamicMaterial/Source/DynamicMaterialEditor/Private/Menus/DMMaterialSlotLayerMenus.cpp
new file mode 100644
--- /dev/null
+++ b/Engine/Plugins/Experimental/DynamicMaterial/Source/DynamicMaterialEditor/Private/Menus/DMMaterialSlotLayerMenus.cpp
@@ -0,0 +1,579 @@
+// Copyright Epic Games, Inc. All Rights Reserved.
+
+#include "Menus/DMMaterialSlotLayerMenus.h"
+#include "Components/DMMaterialProperty.h"
+#include "Components/DMMaterialSlot.h"
+#include "Components/DMMaterialStageBlend.h"
+#include "Components/DMMaterialStageExpression.h"
+#include "Components/DMMaterialStageGradient.h"
+#include "Components/DMMaterialValue.h"
+#include "Components/MaterialStageExpressions/DMMSETextureSample.h"
+#include "Components/MaterialStageExpressions/DMMSETextureSampleEdgeColor.h"
+#include "DMDefs.h"
+#include "DMMaterialSlotLayerAddEffectMenus.h"
+#include "DMValueDefinition.h"
+#include "DynamicMaterialEditorCommands.h"
+#include "DynamicMaterialEditorModule.h"
+#include "Framework/Commands/GenericCommands.h"
+#include "Menus/DMMenuContext.h"
+#include "Model/DynamicMaterialModel.h"
+#include "Model/DynamicMaterialModelEditorOnlyData.h"
+#include "Slate/SDMSlot.h"
+#include "Styling/SlateIconFinder.h"
+#include "ToolMenus.h"
+
+#define LOCTEXT_NAMESPACE "FDMMaterialSlotLayerMenus"
+
+namespace UE::DynamicMaterialEditor::Private
+{
+	static const FName SlotLayerMenuName = TEXT("MaterialDesigner.MaterialSlot.Layer");
+	static const FName SlotLayerAddSectionName = TEXT("AddLayer");
+	static const FName SlotLayerModifySectionName = TEXT("ModifyLayer");
+	static const FName GlobalValuesSectionName = TEXT("GlobalValues");
+
+	void AddLayerModifySection(UToolMenu* InMenu)
+	{
+		if (!IsValid(InMenu) || InMenu->ContainsSection(SlotLayerModifySectionName))
+		{
+			return;
+		}
+
+		FToolMenuSection& NewSection = InMenu->AddSection(SlotLayerModifySectionName, LOCTEXT("LayerActions", "LayerActions"));
+
+		NewSection.AddMenuEntry(
+			FDynamicMaterialEditorCommands::Get().InsertDefaultLayerAbove,
+			TAttribute<FText>(),
+			TAttribute<FText>(),
+			FSlateIconFinder::FindIcon("EditableComboBox.Add")
+		);
+
+		NewSection.AddMenuEntry(FGenericCommands::Get().Copy);
+		NewSection.AddMenuEntry(FGenericCommands::Get().Cut);
+		NewSection.AddMenuEntry(FGenericCommands::Get().Paste);
+		NewSection.AddMenuEntry(FGenericCommands::Get().Duplicate);
+		NewSection.AddMenuEntry(FGenericCommands::Get().Delete);
+	}
+
+	void AddLayerAddEffectsSection(UToolMenu* InMenu, UDMMaterialLayerObject* InLayerObject)
+	{
+		FDMMaterialSlotLayerAddEffectMenus::AddEffectSubMenu(InMenu, InLayerObject);
+	}
+
+	void AddGlobalValueSection(UToolMenu* InMenu)
+	{
+		if (!IsValid(InMenu) || InMenu->ContainsSection(GlobalValuesSectionName))
+		{
+			return;
+		}
+
+		const UDMMenuContext* const MenuContext = InMenu->FindContext<UDMMenuContext>();
+
+		if (!MenuContext)
+		{
+			return;
+		}
+
+		UDynamicMaterialModel* const MaterialModel = MenuContext->GetModel();
+
+		if (!MaterialModel)
+		{
+			return;
+		}
+
+		const TArray<UDMMaterialValue*>& Values = MaterialModel->GetValues();
+
+		if (Values.IsEmpty())
+		{
+			return;
+		}
+
+		FToolMenuSection& NewSection = InMenu->AddSection(GlobalValuesSectionName, LOCTEXT("GlobalValues", "Add Global Value"));
+
+		NewSection.AddSubMenu(
+			NAME_None,
+			LOCTEXT("AddValueStage", "Global Value"),
+			LOCTEXT("AddValueStageTooltip", "Add a Material Stage based on a Material Value defined above."),
+			FNewToolMenuDelegate::CreateLambda([](UToolMenu* InMenu)
+				{
+					const UDMMenuContext* const MenuContext = InMenu->FindContext<UDMMenuContext>();
+
+					if (!MenuContext)
+					{
+						return;
+					}
+
+					UDynamicMaterialModel* const MaterialModel = MenuContext->GetModel();
+
+					if (!MaterialModel)
+					{
+						return;
+					}
+
+					const TArray<UDMMaterialValue*>& Values = MaterialModel->GetValues();
+
+					for (UDMMaterialValue* Value : Values)
+					{
+						InMenu->AddMenuEntry(NAME_None,
+							FToolMenuEntry::InitMenuEntry(NAME_None,
+								Value->GetDescription(),
+								LOCTEXT("AddValueStageSpecificTooltip", "Add a Material Stage based on this Material Value."),
+								FSlateIcon(),
+								FUIAction(FExecuteAction::CreateSP(MenuContext->GetSlotWidget().Pin().Get(), &SDMSlot::AddNewLayer_GlobalValue, Value)))
+						);
+					}
+				})
+		);
+
+		NewSection.AddSubMenu(
+			NAME_None,
+			LOCTEXT("AddNewValueStage", "New Global Value"),
+			LOCTEXT("AddNewValueStageTooltip", "Add a new global Material Value as use it as a Material Stage."),
+			FNewToolMenuDelegate::CreateLambda([](UToolMenu* InMenu)
+				{
+					const UDMMenuContext* const MenuContext = InMenu->FindContext<UDMMenuContext>();
+
+					if (!MenuContext)
+					{
+						return;
+					}
+
+					for (EDMValueType ValueType : UDMValueDefinitionLibrary::GetValueTypes())
+					{
+						FText Name = UDMValueDefinitionLibrary::GetValueDefinition(ValueType).GetDisplayName();
+						FText FormattedTooltip = FText::Format(LOCTEXT("AddTypeTooltipTemplate", "Add a new {0} Value and use it as a Material Stage."), Name);
+
+						InMenu->AddMenuEntry(NAME_None,
+							FToolMenuEntry::InitMenuEntry(NAME_None,
+								Name,
+								FormattedTooltip,
+								FSlateIcon(),
+								FUIAction(FExecuteAction::CreateSP(MenuContext->GetSlotWidget().Pin().Get(), &SDMSlot::AddNewLayer_NewGlobalValue, ValueType)))
+						);
+					}
+				})
+		);
+	}
+
+	void AddSlotMenuEntry(SDMSlot* SlotWidget, UToolMenu* InMenu, const FText& Name, UDMMaterialSlot* Slot, EDMMaterialPropertyType MaterialProperty)
+	{
+		InMenu->AddMenuEntry(NAME_None,
+			FToolMenuEntry::InitMenuEntry(NAME_None,
+				Name,
+				LOCTEXT("AddValueStageSpecificTooltip", "Add a Material Stage based on this Material Value."),
+				FSlateIcon(),
+				FUIAction(FExecuteAction::CreateSP(SlotWidget, &SDMSlot::AddNewLayer_Slot, Slot, MaterialProperty))
+			));
+	}
+
+	void AddLayerInputsMenu_Slot_Properties(UToolMenu* InMenu, UDMMaterialSlot* InSlot)
+	{
+		if (!IsValid(InMenu))
+		{
+			return;
+		}
+
+		const UDMMenuContext* const MenuContext = InMenu->FindContext<UDMMenuContext>();
+
+		if (!MenuContext)
+		{
+			return;
+		}
+
+		const TSharedPtr<SDMSlot> SlotWidget = MenuContext->GetSlotWidget().Pin();
+
+		if (!SlotWidget.IsValid())
+		{
+			return;
+		}
+
+		UDynamicMaterialModel* const MaterialModel = MenuContext->GetModel();
+
+		if (!MaterialModel)
+		{
+			return;
+		}
+
+		UDynamicMaterialModelEditorOnlyData* const ModelEditorOnlyData = UDynamicMaterialModelEditorOnlyData::Get(MaterialModel);
+
+		if (!ensure(ModelEditorOnlyData) || !ensure(InSlot->GetMaterialModelEditorOnlyData() == ModelEditorOnlyData))
+		{
+			return;
+		}
+
+		const TArray<EDMMaterialPropertyType> SlotProperties = ModelEditorOnlyData->GetMaterialPropertiesForSlot(InSlot);
+
+		for (EDMMaterialPropertyType SlotProperty : SlotProperties)
+		{
+			UDMMaterialProperty* MaterialProperty = ModelEditorOnlyData->GetMaterialProperty(SlotProperty);
+
+			if (ensure(MaterialProperty))
+			{
+				UE::DynamicMaterialEditor::Private::AddSlotMenuEntry(
+					SlotWidget.Get(),
+					InMenu,
+					MaterialProperty->GetDescription(),
+					InSlot,
+					SlotProperty
+				);
+			}
+		}
+	}
+
+	void AddLayerInputsMenu_Slots(UToolMenu* InMenu)
+	{
+		if (!IsValid(InMenu) || InMenu->ContainsSection(SlotLayerAddSectionName))
+		{
+			return;
+		}
+
+		const UDMMenuContext* const MenuContext = InMenu->FindContext<UDMMenuContext>();
+
+		if (!MenuContext)
+		{
+			return;
+		}
+
+		const TSharedPtr<SDMSlot> SlotWidget = MenuContext->GetSlotWidget().Pin();
+
+		if (!SlotWidget.IsValid())
+		{
+			return;
+		}
+
+		const UDMMaterialSlot* const Slot = SlotWidget->GetSlot();
+
+		if (!Slot)
+		{
+			return;
+		}
+
+		const UDynamicMaterialModelEditorOnlyData* const ModelEditorOnlyData = Slot->GetMaterialModelEditorOnlyData();
+
+		if (!ModelEditorOnlyData)
+		{
+			return;
+		}
+
+		const TArray<UDMMaterialSlot*>& Slots = ModelEditorOnlyData->GetSlots();
+
+		if (Slots.Num() <= 1)
+		{
+			return;
+		}
+
+		for (UDMMaterialSlot* SlotIter : Slots)
+		{
+			if (Slot == SlotIter)
+			{
+				continue;
+			}
+
+			if (SlotIter->GetLayers().IsEmpty())
+			{
+				continue;
+			}
+
+			const TArray<EDMMaterialPropertyType> SlotProperties = ModelEditorOnlyData->GetMaterialPropertiesForSlot(SlotIter);
+
+			if (SlotProperties.Num() == 1)
+			{
+				static const FText SlotNameFormatTemplate = LOCTEXT("SlotAndProperty", "{0} [{1}]");
+
+				UDMMaterialProperty* MaterialProperty = ModelEditorOnlyData->GetMaterialProperty(SlotProperties[0]);
+
+				if (ensure(MaterialProperty))
+				{
+					UE::DynamicMaterialEditor::Private::AddSlotMenuEntry(
+						SlotWidget.Get(),
+						InMenu,
+						FText::Format(SlotNameFormatTemplate, SlotIter->GetDescription(), MaterialProperty->GetDescription()),
+						SlotIter,
+						SlotProperties[0]
+					);
+				}
+			}
+			else
+			{
+				InMenu->AddMenuEntry(NAME_None,
+					FToolMenuEntry::InitSubMenu(NAME_None,
+						LOCTEXT("AddSlotStage2", "Add Slot Output"),
+						LOCTEXT("AddSlotStageTooltip2", "Add a Material Stage based on the output of another Material Slot."),
+						FNewToolMenuDelegate::CreateStatic(&AddLayerInputsMenu_Slot_Properties, SlotIter))
+				);
+			}
+		}
+	}
+
+	void AddLayerMenu_Gradients(UToolMenu* InMenu)
+	{
+		if (!IsValid(InMenu))
+		{
+			return;
+		}
+
+		const UDMMenuContext* const MenuContext = InMenu->FindContext<UDMMenuContext>();
+
+		if (!MenuContext)
+		{
+			return;
+		}
+
+		const TSharedPtr<SDMSlot> SlotWidget = MenuContext->GetSlotWidget().Pin();
+
+		if (!SlotWidget.IsValid())
+		{
+			return;
+		}
+
+		const TArray<TStrongObjectPtr<UClass>>& Gradients = UDMMaterialStageGradient::GetAvailableGradients();
+
+		for (const TStrongObjectPtr<UClass>& Gradient : Gradients)
+		{
+			UDMMaterialStageGradient* GradientCDO = Cast<UDMMaterialStageGradient>(Gradient->GetDefaultObject());
+
+			if (!ensure(GradientCDO))
+			{
+				continue;
+			}
+
+			const FText MenuName = GradientCDO->GetDescription();
+
+			InMenu->AddMenuEntry(NAME_None,
+				FToolMenuEntry::InitMenuEntry(NAME_None,
+					MenuName,
+					LOCTEXT("ChangeGradientSourceTooltip", "Change the source of this stage to a Material Gradient."),
+					FSlateIcon(),
+					FUIAction(FExecuteAction::CreateSP(
+						SlotWidget.Get(),
+						&SDMSlot::AddNewLayer_Gradient,
+						TSubclassOf<UDMMaterialStageGradient>(Gradient.Get()))
+					)
+				)
+			);
+		}
+	}
+}
+
+UToolMenu* FDMMaterialSlotLayerMenus::GenerateSlotLayerMenu(const TSharedPtr<SDMSlot>& InSlotWidget, UDMMaterialLayerObject* InLayerObject)
+{
+	using namespace UE::DynamicMaterialEditor::Private;
+
+	UToolMenu* NewToolMenu = UDMMenuContext::GenerateContextMenuLayer(SlotLayerMenuName, InSlotWidget, InLayerObject);
+
+	if (!NewToolMenu)
+	{
+		return nullptr;
+	}
+
+	AddAddLayerSection(NewToolMenu);
+
+	if constexpr (UE::DynamicMaterialEditor::bGlobalValuesEnabled)
+	{
+		AddGlobalValueSection(NewToolMenu);
+	}
+
+	AddLayerAddEffectsSection(NewToolMenu, InLayerObject);
+
+	AddLayerModifySection(NewToolMenu);
+
+	return NewToolMenu;
+}
+
+void FDMMaterialSlotLayerMenus::AddAddLayerSection(UToolMenu* InMenu)
+{
+	using namespace UE::DynamicMaterialEditor::Private;
+
+	if (!IsValid(InMenu) || InMenu->ContainsSection(SlotLayerAddSectionName))
+	{
+		return;
+	}
+
+	const UDMMenuContext* const MenuContext = InMenu->FindContext<UDMMenuContext>();
+
+	if (!MenuContext)
+	{
+		return;
+	}
+
+	const TSharedPtr<SDMSlot> SlotWidget = MenuContext->GetSlotWidget().Pin();
+
+	if (!SlotWidget.IsValid())
+	{
+		return;
+	}
+
+	const UDMMaterialSlot* const Slot = SlotWidget->GetSlot();
+
+	if (!Slot)
+	{
+		return;
+	}
+
+	UDynamicMaterialModelEditorOnlyData* const ModelEditorOnlyData = Slot->GetMaterialModelEditorOnlyData();
+
+	if (!ModelEditorOnlyData)
+	{
+		return;
+	}
+
+	UDynamicMaterialModel* MaterialModel = ModelEditorOnlyData->GetMaterialModel();
+
+	if (!MaterialModel)
+	{
+		return;
+	}
+
+	bool bHasValidSlot = false;
+
+	if constexpr (UE::DynamicMaterialEditor::bAdvancedSlotsEnabled)
+	{
+		const TArray<UDMMaterialSlot*>& Slots = ModelEditorOnlyData->GetSlots();
+
+		for (UDMMaterialSlot* SlotIter : Slots)
+		{
+			if (Slot == SlotIter)
+			{
+				continue;
+			}
+
+			if (SlotIter->GetLayers().IsEmpty())
+			{
+				continue;
+			}
+
+			TArray<EDMMaterialPropertyType> SlotProperties = ModelEditorOnlyData->GetMaterialPropertiesForSlot(SlotIter);
+
+			if (SlotProperties.IsEmpty())
+			{
+				continue;
+			}
+
+			bHasValidSlot = true;
+			break;
+		}
+	}
+
+	FToolMenuSection& NewSection = InMenu->AddSection(SlotLayerAddSectionName, LOCTEXT("AddLayer", "Add Layer"));
+
+	NewSection.AddMenuEntry(
+		NAME_None,
+		LOCTEXT("AddTextureSample", "Texture"),
+		LOCTEXT("AddTextureSampleTooltip", "Add a Material Stage based on a Texture."),
+		FSlateIcon(),
+		FUIAction(FExecuteAction::CreateSP(
+			SlotWidget.Get(),
+			&SDMSlot::AddNewLayer_Expression,
+			TSubclassOf<UDMMaterialStageExpression>(UDMMaterialStageExpressionTextureSample::StaticClass()),
+			EDMMaterialLayerStage::All
+		))
+	);
+
+	NewSection.AddMenuEntry(
+		NAME_None,
+		LOCTEXT("AddTextureSampleBaseOnly", "Texture (No alpha)"),
+		LOCTEXT("AddTextureSampleBaseOnlyTooltip", "Add a Material Stage based on a Texture with the Alpha disabled."),
+		FSlateIcon(),
+		FUIAction(FExecuteAction::CreateSP(
+			SlotWidget.Get(),
+			&SDMSlot::AddNewLayer_Expression,
+			TSubclassOf<UDMMaterialStageExpression>(UDMMaterialStageExpressionTextureSample::StaticClass()),
+			EDMMaterialLayerStage::Base
+		))
+	);
+
+	if (Slot->GetLayers().IsEmpty() == false)
+	{
+		NewSection.AddMenuEntry(
+			NAME_None,
+			LOCTEXT("AddAlphaOnly", "Alpha Only"),
+			LOCTEXT("AddAlphaOnlyTooltip", "Add an Alpha-Only Material Layer.\n\nThe base layer will be disabled by default. It can still be re-enabled later."),
+			FSlateIcon(),
+			FUIAction(FExecuteAction::CreateSP(
+				SlotWidget.Get(),
+				&SDMSlot::AddNewLayer_Expression,
+				TSubclassOf<UDMMaterialStageExpression>(UDMMaterialStageExpressionTextureSample::StaticClass()),
+				EDMMaterialLayerStage::Mask
+			))
+		);
+	}
+
+	NewSection.AddMenuEntry(
+		NAME_None,
+		LOCTEXT("AddColor", "Solid Color"),
+		LOCTEXT("AddColorTooltip", "Add a new Material Layer with a solid RGB color."),
+		FSlateIcon(),
+		FUIAction(FExecuteAction::CreateSP(SlotWidget.Get(), &SDMSlot::AddNewLayer_NewLocalValue, EDMValueType::VT_Float3_RGB))
+	);
+
+	NewSection.AddMenuEntry(
+		NAME_None,
+		LOCTEXT("AddColorAtlas", "Color Atlas"),
+		LOCTEXT("AddColorAtlasTooltip", "Add a new Material Layer with a Color Atlas."),
+		FSlateIcon(),
+		FUIAction(FExecuteAction::CreateSP(SlotWidget.Get(), &SDMSlot::AddNewLayer_NewLocalValue, EDMValueType::VT_ColorAtlas))
+	);
+
+	NewSection.AddMenuEntry(
+		NAME_None,
+		LOCTEXT("AddEdgeColor", "Texture Edge Color"),
+		LOCTEXT("AddEdgeColorTooltip", "Add a new Material Layer with a solid color based on the edge color on a texture."),
+		FSlateIcon(),
+		FUIAction(FExecuteAction::CreateSP(
+			SlotWidget.Get(),
+			&SDMSlot::AddNewLayer_Expression,
+			TSubclassOf<UDMMaterialStageExpression>(UDMMaterialStageExpressionTextureSampleEdgeColor::StaticClass()),
+			EDMMaterialLayerStage::All
+		))
+	);
+
+	NewSection.AddMenuEntry(
+		NAME_None,
+		LOCTEXT("AddSceneTexture", "Scene Texture"),
+		LOCTEXT("AddSceneTextureTooltip", "Add a new Material Layer that represents the Scene Texture for a post process material."),
+		FSlateIcon(),
+		FUIAction(FExecuteAction::CreateSP(
+			SlotWidget.Get(),
+			&SDMSlot::AddNewLayer_SceneTexture
+		))
+	);
+
+	if constexpr (UE::DynamicMaterialEditor::bAdvancedSlotsEnabled)
+	{
+		if (bHasValidSlot)
+		{
+			NewSection.AddSubMenu(
+				NAME_None,
+				LOCTEXT("AddSlotStage", "Slot Output"),
+				LOCTEXT("AddSlotStageTooltip", "Add a Material Stage based on the output of another Material Slot."),
+				FNewToolMenuDelegate::CreateStatic(&AddLayerInputsMenu_Slots)
+			);
+		}
+	}
+
+	const TArray<TStrongObjectPtr<UClass>>& Gradients = UDMMaterialStageGradient::GetAvailableGradients();
+
+	if (!Gradients.IsEmpty())
+	{
+		NewSection.AddSubMenu(
+			NAME_None,
+			LOCTEXT("AddGradientStage", "Gradient"),
+			LOCTEXT("AddGradientStageTooltip", "Add a Material Stage based on a Material Gradient."),
+			FNewToolMenuDelegate::CreateStatic(&AddLayerMenu_Gradients)
+		);
+	}
+
+	NewSection.AddMenuEntry(
+		NAME_None,
+		LOCTEXT("AddMaterialFunction", "Material Function"),
+		LOCTEXT("AddMaterialFunctionTooltip", "Add a new Material Layer based on a Material Function."),
+		FSlateIcon(),
+		FUIAction(FExecuteAction::CreateSP(SlotWidget.Get(), &SDMSlot::AddNewLayer_MaterialFunction))
+	);
+
+	if constexpr (UE::DynamicMaterialEditor::bGlobalValuesEnabled)
+	{
+		AddGlobalValueSection(InMenu);
+	}
+}
+
+#undef LOCTEXT_NAMESPACE
diff --git a/Engine/Plugins/Experimental/DynamicMaterial/Source/DynamicMaterialEditor/Private/Menus/DMMaterialSlotLayerMenus.h b/Engine/Plugins/Experimental/DynamicMaterial/Source/DynamicMaterialEditor/Private/Menus/DMMaterialSlotLayerMenus.h
new file mode 100644
--- /dev/null
+++ b/Engine/Plugins/Experimental/DynamicMaterial/Source/DynamicMaterialEditor/Private/Menus/DMMaterialSlotLayerMenus.h
@@ -0,0 +1,19 @@
+// Copyright Epic Games, Inc. All Rights Reserved.
+
+#pragma once
+
+#include "Components/DMMaterialStageExpression.h"
+#include "Templates/SharedPointer.h"
+
+class SDMSlot;
+class UDMMaterialLayerObject;
+class UDMMaterialSlot;
+class UToolMenu;
+
+class FDMMaterialSlotLayerMenus
+{
+public:
+	static UToolMenu* GenerateSlotLayerMenu(const TSharedPtr<SDMSlot>& InSlotWidget, UDMMaterialLayerObject* InLayerObject);
+
+	static void AddAddLayerSection(UToolMenu* InMenu);
+};
diff --git a/Engine/Plugins/Experimental/DynamicMaterial/Source/DynamicMaterialEditor/Private/Menus/DMMaterialSlotMenus.cpp b/Engine/Plugins/Experimental/DynamicMaterial/Source/DynamicMaterialEditor/Private/Menus/DMMaterialSlotMenus.cpp
new file mode 100644
--- /dev/null
+++ b/Engine/Plugins/Experimental/DynamicMaterial/Source/DynamicMaterialEditor/Private/Menus/DMMaterialSlotMenus.cpp
@@ -0,0 +1,69 @@
+// Copyright Epic Games, Inc. All Rights Reserved.
+
+#include "Menus/DMMaterialSlotMenus.h"
+#include "Menus/DMMaterialSlotLayerMenus.h"
+#include "Menus/DMMenuContext.h"
+#include "Slate/SDMSlot.h"
+#include "ToolMenus.h"
+#include "Widgets/SWidget.h"
+
+#define LOCTEXT_NAMESPACE "FDMMaterialSlotMenus"
+
+namespace UE::DynamicMaterialEditor::Private
+{
+	static const FName SlotAddLayerMenuName = TEXT("MaterialDesigner.MaterialSlot.AddLayer");
+	static const FName SlotSectionName = TEXT("Slot");
+
+	void AddSlotMenu(UToolMenu* InMenu)
+	{
+		if (!IsValid(InMenu) || InMenu->ContainsSection(SlotSectionName))
+		{
+			return;
+		}
+
+		const UDMMenuContext* const MenuContext = InMenu->FindContext<UDMMenuContext>();
+
+		if (!MenuContext)
+		{
+			return;
+		}
+
+		FToolMenuSection& NewSection = InMenu->AddSection(SlotSectionName, LOCTEXT("MaterialSlotMenu", "Material Slot"));
+		{
+			NewSection.AddMenuEntry(
+				NAME_None,
+				LOCTEXT("RemoveSlot", "Remove"),
+				LOCTEXT("RemoveSlotTooltip", "Remove this stage from its Material Slot.\n\n"
+					"Warning: Removing a stage off may result in inputs being reset where incompatibilities are found."),
+				FSlateIcon(),
+				FUIAction(FExecuteAction::CreateSP(
+					MenuContext->GetSlotWidget().Pin().Get(),
+					&SDMSlot::RemoveSlot
+				))
+			);
+		}
+	}
+}
+
+TSharedRef<SWidget> FDMMaterialSlotMenus::MakeAddLayerButtonMenu(const TSharedPtr<SDMSlot>& InSlotWidget)
+{
+	using namespace UE::DynamicMaterialEditor::Private;
+
+	if (!UToolMenus::Get()->IsMenuRegistered(SlotAddLayerMenuName))
+	{
+		UToolMenu* NewToolMenu = UDMMenuContext::GenerateContextMenuDefault(SlotAddLayerMenuName);
+
+		if (!NewToolMenu)
+		{
+			return SNullWidget::NullWidget;
+		}
+
+		NewToolMenu->AddDynamicSection(NAME_None, FNewToolMenuDelegate::CreateStatic(&FDMMaterialSlotLayerMenus::AddAddLayerSection));
+	}
+
+	FToolMenuContext MenuContext(UDMMenuContext::CreateSlot(InSlotWidget));
+
+	return UToolMenus::Get()->GenerateWidget(SlotAddLayerMenuName, MenuContext);
+}
+
+#undef LOCTEXT_NAMESPACE
diff --git a/Engine/Plugins/Experimental/DynamicMaterial/Source/DynamicMaterialEditor/Private/Menus/DMMaterialSlotMenus.h b/Engine/Plugins/Experimental/DynamicMaterial/Source/DynamicMaterialEditor/Private/Menus/DMMaterialSlotMenus.h
new file mode 100644
--- /dev/null
+++ b/Engine/Plugins/Experimental/DynamicMaterial/Source/DynamicMaterialEditor/Private/Menus/DMMaterialSlotMenus.h
@@ -0,0 +1,15 @@
+// Copyright Epic Games, Inc. All Rights Reserved.
+
+#pragma once
+
+#include "Templates/SharedPointer.h"
+
+class SDMSlot;
+class SWidget;
+class UToolMenu;
+
+class FDMMaterialSlotMenus
+{
+public:
+	static TSharedRef<SWidget> MakeAddLayerButtonMenu(const TSharedPtr<SDMSlot>& InSlotWidget);
+};
