[MaterialDesigner]

diff --git a/Engine/Plugins/Experimental/DynamicMaterial/Source/DynamicMaterialEditor/Private/Menus/DMMaterialStageInputMenus.cpp b/Engine/Plugins/Experimental/DynamicMaterial/Source/DynamicMaterialEditor/Private/Menus/DMMaterialStageInputMenus.cpp
new file mode 100644
--- /dev/null
+++ b/Engine/Plugins/Experimental/DynamicMaterial/Source/DynamicMaterialEditor/Private/Menus/DMMaterialStageInputMenus.cpp
@@ -0,0 +1,2257 @@
+// Copyright Epic Games, Inc. All Rights Reserved.
+
+#include "DMMaterialStageInputMenus.h"
+#include "Components/DMMaterialLayer.h"
+#include "Components/DMMaterialProperty.h"
+#include "Components/DMMaterialSlot.h"
+#include "Components/DMMaterialStage.h"
+#include "Components/DMMaterialStageExpression.h"
+#include "Components/DMMaterialStageGradient.h"
+#include "Components/DMMaterialStageSource.h"
+#include "Components/DMMaterialValue.h"
+#include "Components/MaterialStageExpressions/DMMSETextureSample.h"
+#include "Components/MaterialStageExpressions/DMMSETextureSampleEdgeColor.h"
+#include "DMValueDefinition.h"
+#include "DynamicMaterialEditorModule.h"
+#include "Framework/MultiBox/MultiBoxBuilder.h"
+#include "Materials/MaterialExpression.h"
+#include "Menus/DMMaterialStageSourceMenus.h"
+#include "Model/DynamicMaterialModel.h"
+#include "Model/DynamicMaterialModelEditorOnlyData.h"
+#include "ScopedTransaction.h"
+#include "Widgets/SNullWidget.h"
+
+#define LOCTEXT_NAMESPACE "FDMMaterialStageInputMenus"
+
+FName FDMMaterialStageInputMenus::GetStageChangeInputMenuName()
+{
+	return "DM.MaterialStage.ChangeInput";
+}
+
+bool FDMMaterialStageInputMenus::CanAcceptSubChannels(UDMMaterialStageThroughput* const InThroughput, const int32 InInputIndex, 
+	const EDMValueType InOutputType)
+{
+	return InOutputType != EDMValueType::VT_Float1
+		&& UDMValueDefinitionLibrary::GetValueDefinition(InOutputType).IsFloatType()
+		&& InThroughput->CanInputAcceptType(InInputIndex, EDMValueType::VT_Float1);
+}
+
+bool FDMMaterialStageInputMenus::CanAcceptSubChannels(UDMMaterialStageThroughput* const InThroughput, const int32 InInputIndex, 
+	const FDMMaterialStageConnector& InOutputConnector)
+{
+	return CanAcceptSubChannels(InThroughput, InInputIndex, InOutputConnector.Type);
+}
+
+TSharedRef<SWidget> FDMMaterialStageInputMenus::MakeStageChangeInputMenu(UDMMaterialStageThroughput* InThroughput, const int32 InInputIndex, 
+	const int32 InInputChannel)
+{
+	if (!ensure(InThroughput))
+	{
+		return SNullWidget::NullWidget;
+	}
+
+	UDMMaterialStage* Stage = InThroughput->GetStage();
+	if (!ensure(Stage))
+	{
+		return SNullWidget::NullWidget;
+	}
+
+	UDMMaterialLayerObject* Layer = Stage->GetLayer();
+	if (!ensure(Layer))
+	{
+		return SNullWidget::NullWidget;
+	}
+
+	const TArray<FDMMaterialStageConnector>& InputConnectors = InThroughput->GetInputConnectors();
+	if (!ensure(InputConnectors.IsValidIndex(InInputIndex)))
+	{
+		return SNullWidget::NullWidget;
+	}
+
+	UDMMaterialSlot* Slot = Layer->GetSlot();
+	if (!ensure(Slot))
+	{
+		return SNullWidget::NullWidget;
+	}
+
+	UDynamicMaterialModelEditorOnlyData* ModelEditorOnlyData = Slot->GetMaterialModelEditorOnlyData();
+	if (!ensure(ModelEditorOnlyData))
+	{
+		return SNullWidget::NullWidget;
+	}
+
+	UDynamicMaterialModel* MaterialModel = ModelEditorOnlyData->GetMaterialModel();
+	if (!ensure(MaterialModel))
+	{
+		return SNullWidget::NullWidget;
+	}	
+
+	/*
+	 * TODO: Implement this part of the UI.
+	 *
+	if constexpr (UE::DynamicMaterialEditor::bMultipleSlotPropertiesEnabled)
+	{
+		const TArray<EDMMaterialPropertyType> Properties = ModelEditorOnlyData->GetMaterialPropertiesForSlot(Slot);
+	}
+	*/
+
+	bool bValidNewValue = false;
+
+	if constexpr (UE::DynamicMaterialEditor::bGlobalValuesEnabled)
+	{
+		for (EDMValueType ValueType : UDMValueDefinitionLibrary::GetValueTypes())
+		{
+			if (!InThroughput->CanInputAcceptType(InInputIndex, ValueType))
+			{
+				continue;
+			}
+
+			bValidNewValue = true;
+		}
+	}
+
+	bool bValidGlobalValues = false;
+
+	if constexpr (UE::DynamicMaterialEditor::bGlobalValuesEnabled)
+	{
+		const TArray<UDMMaterialValue*>& Values = MaterialModel->GetValues();
+
+		for (UDMMaterialValue* Value : Values)
+		{
+			if (InThroughput->CanInputAcceptType(InInputIndex, Value->GetType())
+				|| CanAcceptSubChannels(InThroughput, InInputIndex, Value->GetType()))
+			{
+				bValidGlobalValues = true;
+				break;
+			}
+		}
+	}
+
+	bool bHasValidPreviousStages = false;
+	const TArray<TObjectPtr<UDMMaterialLayerObject>>& Layers = Slot->GetLayers();
+
+	for (UDMMaterialLayerObject* PreviousLayer : Layers)
+	{
+		if (PreviousLayer->HasValidStage(Stage))
+		{
+			break;
+		}
+
+		if (UDMMaterialStageSource* PreviousStageSource = PreviousLayer->GetStage(EDMMaterialLayerStage::Mask)->GetSource())
+		{
+			const TArray<FDMMaterialStageConnector>& PreviousStageOutputConnectors = PreviousStageSource->GetOutputConnectors();
+
+			for (const FDMMaterialStageConnector& OutputConnector : PreviousStageOutputConnectors)
+			{
+				if (InThroughput->CanInputConnectTo(InInputIndex, OutputConnector, FDMMaterialStageConnectorChannel::WHOLE_CHANNEL, true))
+				{
+					bHasValidPreviousStages = true;
+					goto EndStageSearch;
+				}
+			}
+		}
+	}
+
+EndStageSearch:
+
+	bool bHasValidSlot = false;
+
+	if constexpr (UE::DynamicMaterialEditor::bAdvancedSlotsEnabled)
+	{
+		const TArray<UDMMaterialSlot*>& Slots = ModelEditorOnlyData->GetSlots();
+
+		for (UDMMaterialSlot* SlotIter : Slots)
+		{
+			if (SlotIter == Slot)
+			{
+				continue;
+			}
+
+			if (SlotIter->GetLayers().IsEmpty())
+			{
+				continue;
+			}
+
+			const TSet<EDMValueType> SlotOutputTypes = SlotIter->GetAllOutputConnectorTypes();
+
+			for (EDMValueType ValueType : SlotOutputTypes)
+			{
+				if (InThroughput->CanInputAcceptType(InInputIndex, ValueType)
+					|| CanAcceptSubChannels(InThroughput, InInputIndex, ValueType))
+				{
+					bHasValidSlot = true;
+					break;
+				}
+			}
+
+			if (bHasValidSlot)
+			{
+				break;
+			}
+		}
+	}
+
+	const TArray<TStrongObjectPtr<UClass>>& Gradients = UDMMaterialStageGradient::GetAvailableGradients();
+
+	const bool bHasValidGradients = Gradients.Num() > 0 && InThroughput->CanInputAcceptType(InInputIndex, EDMValueType::VT_Float1);
+
+	FMenuBuilder MenuBuilder(true, nullptr);
+
+	MenuBuilder.BeginSection("ChangeMaterialStageInput", LOCTEXT("ChangeMaterialStageInput", "Change Input"));
+	{
+		if (UDMMaterialStageExpression* TextureSample = Cast<UDMMaterialStageExpression>(UDMMaterialStageExpressionTextureSample::StaticClass()->GetDefaultObject(true)))
+		{
+			GenerateChangeInputMenu_Expression(MenuBuilder, TextureSample, InThroughput, InInputIndex, InInputChannel);
+		}
+
+		if (UDMMaterialStageExpression* TextureSampleEdgeColor = Cast<UDMMaterialStageExpression>(UDMMaterialStageExpressionTextureSampleEdgeColor::StaticClass()->GetDefaultObject(true)))
+		{
+			GenerateChangeInputMenu_Expression(MenuBuilder, TextureSampleEdgeColor, InThroughput, InInputIndex, InInputChannel);
+		}
+
+		if constexpr (UE::DynamicMaterialEditor::bAdvancedSlotsEnabled)
+		{
+			if (bHasValidSlot)
+			{
+				MenuBuilder.AddSubMenu(
+					LOCTEXT("ChangeInputSlot", "Slot Output"),
+					LOCTEXT("ChangeInputSlotTooltip", "Change the source of this input to the output from another Material Slot."),
+					FNewMenuDelegate::CreateStatic(
+						&FDMMaterialStageInputMenus::GenerateChangeInputMenu_Slots,
+						InThroughput,
+						InInputIndex,
+						InInputChannel
+					)
+				);
+			}
+		}
+
+		MenuBuilder.AddMenuEntry(
+			LOCTEXT("ChangeInputColorRGB", "Solid Color"),
+			LOCTEXT("ChangeInputColorRGBTooltip", "Change the source of this input to a Solid Color."),
+			FSlateIcon(),
+			FUIAction(FExecuteAction::CreateWeakLambda(
+				Stage,
+				[Stage, InInputIndex, InInputChannel]()
+				{
+					FScopedTransaction Transaction(LOCTEXT("SetStageInput", "Set Material Stage Input"));
+					Stage->Modify();
+					Stage->ChangeInput_NewLocalValue(InInputIndex, InInputChannel, EDMValueType::VT_Float3_RGB, FDMMaterialStageConnectorChannel::WHOLE_CHANNEL);
+				})
+			)
+		);
+
+		MenuBuilder.AddMenuEntry(
+			LOCTEXT("ChangeInputColorAtlas", "Color Atlas"),
+			LOCTEXT("ChangeInputColorAtlasTooltip", "Change the source of this input to a Color Atlas."),
+			FSlateIcon(),
+			FUIAction(FExecuteAction::CreateWeakLambda(
+				Stage,
+				[Stage, InInputIndex, InInputChannel]()
+				{
+					FScopedTransaction Transaction(LOCTEXT("SetStageInput", "Set Material Stage Input"));
+					Stage->Modify();
+					Stage->ChangeInput_NewLocalValue(InInputIndex, InInputChannel, EDMValueType::VT_ColorAtlas, FDMMaterialStageConnectorChannel::THREE_CHANNELS);
+				})
+			)
+		);
+
+		if constexpr (UE::DynamicMaterialEditor::bGlobalValuesEnabled)
+		{
+			if (bValidNewValue)
+			{
+				MenuBuilder.AddSubMenu(
+					LOCTEXT("ChangeInputLocalValue", "New Local Value"),
+					LOCTEXT("ChangeInputLocalValueTooltip", "Change the source of this input to a Local Value."),
+					FNewMenuDelegate::CreateStatic(
+						&FDMMaterialStageInputMenus::GenerateChangeInputMenu_NewLocalValues,
+						InThroughput,
+						InInputIndex,
+						InInputChannel
+					)
+				);
+			}
+
+			if (bValidGlobalValues)
+			{
+				MenuBuilder.AddSubMenu(
+					LOCTEXT("ChangeInputGlobalValue", "Global Value"),
+					LOCTEXT("ChangeInputGlobalValueTooltip", "Change the source of this input to a Global Material Value."),
+					FNewMenuDelegate::CreateStatic(
+						&FDMMaterialStageInputMenus::GenerateChangeInputMenu_GlobalValues,
+						InThroughput,
+						InInputIndex,
+						InInputChannel
+					)
+				);
+			}
+
+			if (bValidNewValue)
+			{
+				MenuBuilder.AddSubMenu(
+					LOCTEXT("ChangeInputNewGlobalValue", "New Global Value"),
+					LOCTEXT("ChangeInputNewGlobalValueTooltip", "Add a new Global Material Value and use it as the source of this input."),
+					FNewMenuDelegate::CreateStatic(
+						&FDMMaterialStageInputMenus::GenerateChangeInputMenu_NewGlobalValues,
+						InThroughput,
+						InInputIndex,
+						InInputChannel
+					)
+				);
+			}
+		}
+
+		if (bHasValidGradients)
+		{
+			MenuBuilder.AddSubMenu(
+				LOCTEXT("ChangeInputGradient", "Gradient"),
+				LOCTEXT("ChangeInputGradientTooltip", "Change the source of this input to a Material Gradient."),
+				FNewMenuDelegate::CreateStatic(
+					&FDMMaterialStageInputMenus::GenerateChangeInputMenu_Gradients,
+					InThroughput,
+					InInputIndex,
+					InInputChannel
+				)
+			);
+		}
+
+		if (bHasValidPreviousStages)
+		{
+			MenuBuilder.AddSubMenu(
+				LOCTEXT("ChangeInputPreviousStage", "Stage"),
+				LOCTEXT("ChangeInputPreviousStageTooltip", "Change the source of this input to the output of one of this Material Slot's Stages."),
+				FNewMenuDelegate::CreateStatic(
+					&FDMMaterialStageInputMenus::GenerateChangeInputMenu_PreviousStages,
+					InThroughput,
+					InInputIndex,
+					InInputChannel
+				)
+			);
+		}
+	}
+	MenuBuilder.EndSection();
+
+	return MenuBuilder.MakeWidget();
+}
+
+namespace UE::DynamicMaterialEditor::Private
+{
+	void GenerateChangeInputMenu_PreviousStage_Outputs_Impl(FMenuBuilder& InChildMenuBuilder, const FText& InMenuName,
+		UDMMaterialStage* InStage, UDMMaterialStageThroughput* InThroughput, int32 InInputIndex, int32 InInputChannel, 
+		EDMMaterialPropertyType InSlotProperty)
+	{
+		InChildMenuBuilder.AddSubMenu(
+			InMenuName,
+			LOCTEXT("ChangeInputPreviousStageTooltip", "Change the source of this input to a previous Material Stage."),
+			FNewMenuDelegate::CreateStatic(
+				&FDMMaterialStageInputMenus::GenerateChangeInputMenu_PreviousStage_Outputs,
+				InThroughput,
+				InInputIndex,
+				InInputChannel,
+				InSlotProperty
+			)
+		);
+	}
+
+	void GenerateChangeInputMenu_PreviousStage_Output_Channels_Impl(FMenuBuilder& InChildMenuBuilder, const FText& InMenuName,
+		UDMMaterialStage* InStage, UDMMaterialStageThroughput* InThroughput, int32 InInputIndex, int32 InInputChannel, 
+		EDMMaterialPropertyType InSlotProperty, int32 InOutputIndex)
+	{
+		InChildMenuBuilder.AddSubMenu(
+			InMenuName,
+			LOCTEXT("ChangeInputPreviousStageTooltip", "Change the source of this input to a previous Material Stage."),
+			FNewMenuDelegate::CreateStatic(
+				&FDMMaterialStageInputMenus::GenerateChangeInputMenu_PreviousStage_Output_Channels,
+				InThroughput,
+				InInputIndex,
+				InInputChannel,
+				InSlotProperty,
+				InOutputIndex
+			)
+		);
+	}
+
+	void GenerateChangeInputMenu_PreviousStage_Output_Channel_Impl(FMenuBuilder& InChildMenuBuilder, const FText& InMenuName,
+		UDMMaterialStage* InStage, UDMMaterialStageThroughput* InThroughput, int32 InInputIndex, int32 InInputChannel, 
+		EDMMaterialPropertyType InSlotProperty, int32 InOutputIndex, int32 InOutputChannel)
+	{
+		InChildMenuBuilder.AddMenuEntry(
+			InMenuName,
+			LOCTEXT("ChangeInputPreviousStageTooltip", "Change the source of this input to a previous Material Stage."),
+			FSlateIcon(),
+			FUIAction(FExecuteAction::CreateWeakLambda(
+				InStage,
+				[InStage, InInputIndex, InInputChannel, InSlotProperty, InOutputIndex, InOutputChannel]()
+				{
+					FScopedTransaction Transaction(LOCTEXT("SetStageInput", "Set Material Stage Input"));
+					InStage->Modify();
+					InStage->ChangeInput_PreviousStage(InInputIndex, InInputChannel, InSlotProperty, InOutputIndex, InOutputChannel);
+				})
+			)
+		);
+	}
+
+	void GenerateChangeInputMenu_PreviousStage_Output_Channel_Channels_Impl(FMenuBuilder& InChildMenuBuilder, const FText& InMenuName,
+		UDMMaterialStage* InStage, UDMMaterialStageThroughput* InThroughput, int32 InInputIndex, int32 InInputChannel, 
+		EDMMaterialPropertyType InSlotProperty, int32 InOutputIndex)
+	{
+		InChildMenuBuilder.AddSubMenu(
+			InMenuName,
+			LOCTEXT("ChangeInputPreviousStageTooltip", "Change the source of this input to a previous Material Stage."),
+			FNewMenuDelegate::CreateStatic(
+				&FDMMaterialStageInputMenus::GenerateChangeInputMenu_PreviousStage_Output_Channels,
+				InThroughput,
+				InInputIndex,
+				InInputChannel,
+				InSlotProperty,
+				InOutputIndex
+			),
+			FUIAction(FExecuteAction::CreateWeakLambda(
+				InStage,
+				[InStage, InInputIndex, InInputChannel, InSlotProperty, InOutputIndex]()
+				{
+					FScopedTransaction Transaction(LOCTEXT("SetStageInput", "Set Material Stage Input"));
+					InStage->Modify();
+					InStage->ChangeInput_PreviousStage(InInputIndex, InInputChannel, InSlotProperty, InOutputIndex, FDMMaterialStageConnectorChannel::WHOLE_CHANNEL);
+				})
+			),
+			NAME_None,
+			EUserInterfaceActionType::Button
+		);
+	}
+
+	void GenerateChangeInputMenu_PreviousStage_Output_Impl(FMenuBuilder& InChildMenuBuilder, const FText& InMenuName, 
+		UDMMaterialStage* InStage, UDMMaterialStageThroughput* InThroughput, int32 InInputIndex, int32 InInputChannel,
+		EDMMaterialPropertyType InSlotProperty, int32 InOutputIndex, bool bInAcceptsWholeChannel, bool bInAcceptsSubChannels)
+	{
+		if (!bInAcceptsWholeChannel && !bInAcceptsSubChannels)
+		{
+			return;
+		}
+
+		if (bInAcceptsSubChannels)
+		{
+			if (bInAcceptsWholeChannel)
+			{
+				UE::DynamicMaterialEditor::Private::GenerateChangeInputMenu_PreviousStage_Output_Channel_Channels_Impl(InChildMenuBuilder, InMenuName, InStage, InThroughput,
+					InInputIndex, InInputChannel, InSlotProperty, InOutputIndex);
+			}
+			else
+			{
+				UE::DynamicMaterialEditor::Private::GenerateChangeInputMenu_PreviousStage_Output_Channels_Impl(InChildMenuBuilder, InMenuName, InStage, InThroughput,
+					InInputIndex, InInputChannel, InSlotProperty, InOutputIndex);
+			}
+		}
+		else
+		{
+			UE::DynamicMaterialEditor::Private::GenerateChangeInputMenu_PreviousStage_Output_Channel_Impl(InChildMenuBuilder, InMenuName, InStage, InThroughput,
+				InInputIndex, InInputChannel, InSlotProperty, InOutputIndex, FDMMaterialStageConnectorChannel::WHOLE_CHANNEL);
+		}
+	}
+}
+
+void FDMMaterialStageInputMenus::GenerateChangeInputMenu_PreviousStages(FMenuBuilder& InChildMenuBuilder, UDMMaterialStageThroughput* InThroughput,
+	const int32 InInputIndex, const int32 InInputChannel)
+{
+	if (!ensure(InThroughput))
+	{
+		return;
+	}
+
+	UDMMaterialStage* Stage = InThroughput->GetStage();
+	if (!ensure(Stage))
+	{
+		return;
+	}
+
+	UDMMaterialLayerObject* Layer = Stage->GetLayer();
+	if (!ensure(Layer))
+	{
+		return;
+	}
+
+	const TArray<FDMMaterialStageConnector>& InputConnectors = InThroughput->GetInputConnectors();
+	if (!ensure(InputConnectors.IsValidIndex(InInputIndex)))
+	{
+		return;
+	}
+
+	UDMMaterialSlot* Slot = Layer->GetSlot();
+	if (!ensure(Slot))
+	{
+		return;
+	}
+
+	UDynamicMaterialModelEditorOnlyData* ModelEditorOnlyData = Slot->GetMaterialModelEditorOnlyData();
+	if (!ensure(ModelEditorOnlyData))
+	{
+		return;
+	}
+
+	EDMMaterialPropertyType StageProperty = Layer->GetMaterialProperty();
+
+	TArray<EDMMaterialPropertyType> SlotProperties = ModelEditorOnlyData->GetMaterialPropertiesForSlot(Slot);
+
+	for (EDMMaterialPropertyType SlotProperty : SlotProperties)
+	{
+		const UDMMaterialLayerObject* PreviousLayer = Layer->GetPreviousLayer(StageProperty, EDMMaterialLayerStage::Base);
+
+		if (!PreviousLayer)
+		{
+			continue;
+		}
+
+		if (!ensure(PreviousLayer->GetStage(EDMMaterialLayerStage::Mask)))
+		{
+			continue;
+		}
+
+		UDMMaterialStageSource* PreviousStageSource = PreviousLayer->GetStage(EDMMaterialLayerStage::Mask)->GetSource();
+		if (!ensure(PreviousStageSource))
+		{
+			continue;
+		}
+
+		const TArray<FDMMaterialStageConnector>& PreviousStageOutputConnectors = PreviousStageSource->GetOutputConnectors();
+		int32 ValidOutputConnectors = 0;
+		int32 LastValidOutputConnector = INDEX_NONE;
+
+		for (int32 OutputIdx = 0; OutputIdx < PreviousStageOutputConnectors.Num(); ++OutputIdx)
+		{
+			if (InThroughput->CanInputConnectTo(InInputIndex, PreviousStageOutputConnectors[OutputIdx], FDMMaterialStageConnectorChannel::WHOLE_CHANNEL, true))
+			{
+				++ValidOutputConnectors;
+				LastValidOutputConnector = OutputIdx;
+			}
+		}
+
+		if (ValidOutputConnectors == 0)
+		{
+			continue;
+		}
+
+		UDMMaterialProperty* PropertyObj = ModelEditorOnlyData->GetMaterialProperty(SlotProperty);
+		if (!ensure(PropertyObj))
+		{
+			continue;
+		}
+
+		const bool bJustOneConnector = (ValidOutputConnectors == 1);
+		const bool bAcceptsWholeChannel = InThroughput->CanInputConnectTo(InInputIndex, PreviousStageOutputConnectors[LastValidOutputConnector], FDMMaterialStageConnectorChannel::WHOLE_CHANNEL);
+		const bool bAcceptsSubChannels = CanAcceptSubChannels(InThroughput, InInputIndex, PreviousStageOutputConnectors[LastValidOutputConnector]);
+
+		if (bJustOneConnector)
+		{
+			static const FText ExpressionNameFormatTemplate = LOCTEXT("ExpressionAndOutput", "{0} [{1}]");
+			const FText PropertyAndChannel = FText::Format(ExpressionNameFormatTemplate, PropertyObj->GetDescription(), PreviousStageOutputConnectors[LastValidOutputConnector].Name);
+
+			UE::DynamicMaterialEditor::Private::GenerateChangeInputMenu_PreviousStage_Output_Impl(InChildMenuBuilder,
+				PropertyAndChannel, Stage, InThroughput, InInputIndex, InInputChannel, SlotProperty, LastValidOutputConnector, 
+				bAcceptsWholeChannel, bAcceptsSubChannels);
+		}
+		else
+		{
+			UE::DynamicMaterialEditor::Private::GenerateChangeInputMenu_PreviousStage_Outputs_Impl(InChildMenuBuilder, PropertyObj->GetDescription(), 
+				Stage, InThroughput, InInputIndex, InInputChannel, SlotProperty);
+		}
+	}
+}
+
+void FDMMaterialStageInputMenus::GenerateChangeInputMenu_PreviousStage_Outputs(FMenuBuilder& InChildMenuBuilder, UDMMaterialStageThroughput* InThroughput,
+	const int32 InInputIndex, const int32 InInputChannel, EDMMaterialPropertyType InMaterialProperty)
+{
+	if (!ensure(InThroughput))
+	{
+		return;
+	}
+
+	UDMMaterialStage* Stage = InThroughput->GetStage();
+	if (!ensure(Stage))
+	{
+		return;
+	}
+
+	UDMMaterialLayerObject* Layer = Stage->GetLayer();
+	if (!ensure(Layer))
+	{
+		return;
+	}
+
+	const TArray<FDMMaterialStageConnector>& InputConnectors = InThroughput->GetInputConnectors();
+	if (!ensure(InputConnectors.IsValidIndex(InInputIndex)))
+	{
+		return;
+	}
+
+	UDMMaterialSlot* Slot = Layer->GetSlot();
+	if (!ensure(Slot))
+	{
+		return;
+	}
+
+	UDynamicMaterialModelEditorOnlyData* ModelEditorOnlyData = Slot->GetMaterialModelEditorOnlyData();
+	if (!ensure(ModelEditorOnlyData))
+	{
+		return;
+	}
+
+	const TArray<EDMMaterialPropertyType> SlotProperties = ModelEditorOnlyData->GetMaterialPropertiesForSlot(Slot);
+	if (!ensure(SlotProperties.Contains(InMaterialProperty)))
+	{
+		return;
+	}
+
+	EDMMaterialPropertyType StageProperty = Layer->GetMaterialProperty();
+
+	const UDMMaterialLayerObject* PreviousLayer = Layer->GetPreviousLayer(StageProperty, EDMMaterialLayerStage::Base);
+	if (!ensure(PreviousLayer) || !ensure(PreviousLayer->GetStage(EDMMaterialLayerStage::Mask)))
+	{
+		return;
+	}
+
+	UDMMaterialStageSource* PreviousStageSource = PreviousLayer->GetStage(EDMMaterialLayerStage::Mask)->GetSource();
+	if (!ensure(PreviousStageSource))
+	{
+		return;
+	}
+
+	const TArray<FDMMaterialStageConnector>& PreviousStageOutputConnectors = PreviousStageSource->GetOutputConnectors();
+
+	InChildMenuBuilder.BeginSection("ChangeInputOutputs", LOCTEXT("ChangeInputOutputs", "Outputs"));
+	{
+		for (int32 OutputIdx = 0; OutputIdx < PreviousStageOutputConnectors.Num(); ++OutputIdx)
+		{
+			const bool bAcceptsWholeChannel = InThroughput->CanInputConnectTo(InInputIndex, PreviousStageOutputConnectors[OutputIdx], FDMMaterialStageConnectorChannel::WHOLE_CHANNEL);
+			const bool bAcceptsSubChannels = CanAcceptSubChannels(InThroughput, InInputIndex, PreviousStageOutputConnectors[OutputIdx]);
+
+			UE::DynamicMaterialEditor::Private::GenerateChangeInputMenu_PreviousStage_Output_Impl(InChildMenuBuilder,
+				PreviousStageOutputConnectors[OutputIdx].Name, Stage, InThroughput, InInputIndex, InInputChannel, InMaterialProperty, OutputIdx,
+				bAcceptsWholeChannel, bAcceptsSubChannels);
+		}
+	}
+}
+
+void FDMMaterialStageInputMenus::GenerateChangeInputMenu_PreviousStage_Output_Channels(FMenuBuilder& InChildMenuBuilder, UDMMaterialStageThroughput* InThroughput,
+	const int32 InInputIndex, const int32 InInputChannel, EDMMaterialPropertyType InMaterialProperty, const int32 OutputIndex)
+{
+	if (!ensure(InThroughput))
+	{
+		return;
+	}
+
+	UDMMaterialStage* Stage = InThroughput->GetStage();
+	if (!ensure(Stage))
+	{
+		return;
+	}
+
+	UDMMaterialLayerObject* Layer = Stage->GetLayer();
+	if (!ensure(Layer))
+	{
+		return;
+	}
+
+	const TArray<FDMMaterialStageConnector>& InputConnectors = InThroughput->GetInputConnectors();
+	if (!ensure(InputConnectors.IsValidIndex(InInputIndex)))
+	{
+		return;
+	}
+
+	UDMMaterialSlot* Slot = Layer->GetSlot();
+	if (!ensure(Slot))
+	{
+		return;
+	}
+
+	UDynamicMaterialModelEditorOnlyData* ModelEditorOnlyData = Slot->GetMaterialModelEditorOnlyData();
+	if (!ensure(ModelEditorOnlyData))
+	{
+		return;
+	}
+
+	const TArray<EDMMaterialPropertyType> SlotProperties = ModelEditorOnlyData->GetMaterialPropertiesForSlot(Slot);
+	if (!ensure(SlotProperties.Contains(InMaterialProperty)))
+	{
+		return;
+	}
+
+	EDMMaterialPropertyType StageProperty = Layer->GetMaterialProperty();
+
+	const UDMMaterialLayerObject* PreviousLayer = Layer->GetPreviousLayer(StageProperty, EDMMaterialLayerStage::Base);
+	if (!ensure(PreviousLayer) || !ensure(PreviousLayer->GetStage(EDMMaterialLayerStage::Mask)))
+	{
+		return;
+	}
+
+	UDMMaterialStageSource* PreviousStageSource = PreviousLayer->GetStage(EDMMaterialLayerStage::Mask)->GetSource();
+	if (!ensure(PreviousStageSource))
+	{
+		return;
+	}
+
+	const TArray<FDMMaterialStageConnector>& PreviousStageOutputConnectors = PreviousStageSource->GetOutputConnectors();
+	if (!ensure(PreviousStageOutputConnectors.IsValidIndex(OutputIndex))
+		|| !ensure(CanAcceptSubChannels(InThroughput, InInputIndex, PreviousStageOutputConnectors[OutputIndex])))
+	{
+		return;
+	}
+
+	const int32 FloatSize = UDMValueDefinitionLibrary::GetValueDefinition(PreviousStageOutputConnectors[OutputIndex].Type).GetFloatCount();
+
+	if (PreviousStageOutputConnectors[OutputIndex].Type != EDMValueType::VT_Float_Any)
+	{
+		if (!ensure(FloatSize > 1))
+		{
+			return;
+		}
+	}
+
+	InChildMenuBuilder.BeginSection("ChangeInputOutputChannels", LOCTEXT("ChangeInputOutputChannels", "Output Channels"));
+	{
+		for (int32 ChannelIndex = 0; ChannelIndex < FloatSize; ++ChannelIndex)
+		{
+			const FText ChannelName = UDMValueDefinitionLibrary::GetValueDefinition(PreviousStageOutputConnectors[OutputIndex].Type).GetChannelName(ChannelIndex + 1);
+			const int32 OutputChannel = UE::DynamicMaterialEditor::ChannelIndexToChannelBit(ChannelIndex + 1);
+
+			UE::DynamicMaterialEditor::Private::GenerateChangeInputMenu_PreviousStage_Output_Channel_Impl(InChildMenuBuilder, ChannelName, Stage, InThroughput,
+				InInputIndex, InInputChannel, InMaterialProperty, OutputIndex, OutputChannel);
+		}
+	}
+	InChildMenuBuilder.EndSection();
+}
+
+void FDMMaterialStageInputMenus::GenerateChangeInputMenu_NewLocalValues(FMenuBuilder& InChildMenuBuilder, UDMMaterialStageThroughput* InThroughput,
+	const int32 InInputIndex, const int32 InInputChannel)
+{
+	if (!ensure(InThroughput))
+	{
+		return;
+	}
+
+	UDMMaterialStage* Stage = InThroughput->GetStage();
+	if (!ensure(Stage))
+	{
+		return;
+	}
+
+	const TArray<FDMMaterialStageConnector>& InputConnectors = InThroughput->GetInputConnectors();
+	if (!ensure(InputConnectors.IsValidIndex(InInputIndex)))
+	{
+		return;
+	}
+
+	for (EDMValueType ValueType : UDMValueDefinitionLibrary::GetValueTypes())
+	{
+		if (!InThroughput->CanInputAcceptType(InInputIndex, ValueType))
+		{
+			continue;
+		}
+
+		const FText Name = UDMValueDefinitionLibrary::GetValueDefinition(ValueType).GetDisplayName();
+		const FText FormattedTooltip = FText::Format(LOCTEXT("ChangeInputNewLocalValueInputTooltipTemplate", 
+			"Add a new {0} Local Value and use it as the source of this input."), Name);
+
+		InChildMenuBuilder.AddMenuEntry(
+			Name,
+			FormattedTooltip,
+			FSlateIcon(),
+			FUIAction(FExecuteAction::CreateWeakLambda(
+				Stage,
+				[Stage, InInputIndex, InInputChannel, ValueType]()
+				{
+					FScopedTransaction Transaction(LOCTEXT("SetStageInput", "Set Material Stage Input"));
+					Stage->Modify();
+					Stage->ChangeInput_NewLocalValue(InInputIndex, InInputChannel, ValueType, FDMMaterialStageConnectorChannel::WHOLE_CHANNEL);
+				})
+			)
+		);
+	}
+}
+
+namespace UE::DynamicMaterialEditor::Private
+{
+	void GenerateChangeInputMenu_GlobalValue_Channels_Impl(FMenuBuilder& InChildMenuBuilder, const FText& InMenuName,
+		UDMMaterialStage* InStage, UDMMaterialStageThroughput* InThroughput, int32 InInputIndex, int32 InInputChannel,
+		int32 InValueIndex, UDMMaterialValue* InValue)
+	{
+		InChildMenuBuilder.AddSubMenu(
+			InMenuName,
+			LOCTEXT("ChangeInputGlobalValueTooltip", "Change the source of this input to this Global Material Value."),
+			FNewMenuDelegate::CreateStatic(
+				&FDMMaterialStageInputMenus::GenerateChangeInputMenu_GlobalValue_Channels,
+				InThroughput,
+				InInputIndex,
+				InInputChannel,
+				InValueIndex
+			)
+		);
+	}
+
+	void GenerateChangeInputMenu_GlobalValue_Channel_Impl(FMenuBuilder& InChildMenuBuilder, const FText& InMenuName,
+		UDMMaterialStage* InStage, UDMMaterialStageThroughput* InThroughput, int32 InInputIndex, int32 InInputChannel,
+		int32 InValueIndex, UDMMaterialValue* InValue, int32 InOutputChannel)
+	{		
+		InChildMenuBuilder.AddMenuEntry(
+			InMenuName,
+			LOCTEXT("ChangeInputGlobalValueTooltip", "Change the source of this input to this Global Material Value."),
+			FSlateIcon(),
+			FUIAction(FExecuteAction::CreateWeakLambda(
+				InStage,
+				[InStage, InInputIndex, InInputChannel, ValueWeak = TWeakObjectPtr<UDMMaterialValue>(InValue), InOutputChannel]()
+				{
+					if (UDMMaterialValue* Value = ValueWeak.Get())
+					{
+						FScopedTransaction Transaction(LOCTEXT("SetStageInput", "Set Material Stage Input"));
+						InStage->Modify();
+						Value->Modify();
+						InStage->ChangeInput_Value(InInputIndex, InInputChannel, Value, InOutputChannel);
+					}
+				})
+			)
+		);
+	}
+
+	void GenerateChangeInputMenu_GlobalValue_Value_Channels_Impl(FMenuBuilder& InChildMenuBuilder, const FText& InMenuName,
+		UDMMaterialStage* InStage, UDMMaterialStageThroughput* InThroughput, int32 InInputIndex, int32 InInputChannel,
+		int32 InValueIndex, UDMMaterialValue* InValue)
+	{
+		InChildMenuBuilder.AddSubMenu(
+			InMenuName,
+			LOCTEXT("ChangeInputGlobalValueTooltip", "Change the source of this input to this Global Material Value."),
+			FNewMenuDelegate::CreateStatic(
+				&FDMMaterialStageInputMenus::GenerateChangeInputMenu_GlobalValue_Channels,
+				InThroughput,
+				InInputIndex,
+				InInputChannel,
+				InValueIndex
+			),
+			FUIAction(FExecuteAction::CreateWeakLambda(
+				InStage,
+				[InStage, InInputIndex, InInputChannel, ValueWeak = TWeakObjectPtr<UDMMaterialValue>(InValue)]()
+				{
+					if (UDMMaterialValue* Value = ValueWeak.Get())
+					{
+						FScopedTransaction Transaction(LOCTEXT("SetStageInput", "Set Material Stage Input"));
+						InStage->Modify();
+						Value->Modify();
+						InStage->ChangeInput_Value(InInputIndex, InInputChannel, Value, FDMMaterialStageConnectorChannel::WHOLE_CHANNEL);
+					}
+				})
+			),
+			NAME_None,
+			EUserInterfaceActionType::Button
+		);
+	}
+
+	void GenerateChangeInputMenu_GlobalValue_Impl(FMenuBuilder& InChildMenuBuilder, const FText& InMenuName, 
+		UDMMaterialStage* InStage, UDMMaterialStageThroughput* InThroughput, int32 InInputIndex, int32 InInputChannel,
+		int32 InValueIndex, UDMMaterialValue* InValue, bool bInAcceptsWholeChannel, bool bInAcceptsSubChannels)
+	{
+		if (!bInAcceptsWholeChannel && !bInAcceptsSubChannels)
+		{
+			return;
+		}
+
+		if (bInAcceptsSubChannels)
+		{
+			if (bInAcceptsWholeChannel)
+			{
+				UE::DynamicMaterialEditor::Private::GenerateChangeInputMenu_GlobalValue_Value_Channels_Impl(InChildMenuBuilder, InMenuName, InStage, InThroughput,
+					InInputIndex, InInputChannel, InValueIndex, InValue);
+			}
+			else
+			{
+				UE::DynamicMaterialEditor::Private::GenerateChangeInputMenu_GlobalValue_Channels_Impl(InChildMenuBuilder, InMenuName, InStage, InThroughput,
+					InInputIndex, InInputChannel, InValueIndex, InValue);
+			}
+		}
+		else
+		{
+			UE::DynamicMaterialEditor::Private::GenerateChangeInputMenu_GlobalValue_Channel_Impl(InChildMenuBuilder, InMenuName, InStage, InThroughput,
+				InInputIndex, InInputChannel, InValueIndex, InValue, FDMMaterialStageConnectorChannel::WHOLE_CHANNEL);
+		}
+	}
+}
+
+void FDMMaterialStageInputMenus::GenerateChangeInputMenu_GlobalValues(FMenuBuilder& InChildMenuBuilder, UDMMaterialStageThroughput* InThroughput,
+	const int32 InInputIndex, const int32 InInputChannel)
+{
+	if (!ensure(InThroughput))
+	{
+		return;
+	}
+
+	UDMMaterialStage* Stage = InThroughput->GetStage();
+	if (!ensure(Stage))
+	{
+		return;
+	}
+
+	UDMMaterialLayerObject* Layer = Stage->GetLayer();
+	if (!ensure(Layer))
+	{
+		return;
+	}
+
+	const TArray<FDMMaterialStageConnector>& InputConnectors = InThroughput->GetInputConnectors();
+	if (!ensure(InputConnectors.IsValidIndex(InInputIndex)))
+	{
+		return;
+	}
+
+	UDMMaterialSlot* Slot = Layer->GetSlot();
+	if (!ensure(Slot))
+	{
+		return;
+	}
+
+	UDynamicMaterialModelEditorOnlyData* ModelEditorOnlyData = Slot->GetMaterialModelEditorOnlyData();
+	if (!ensure(ModelEditorOnlyData))
+	{
+		return;
+	}
+
+	UDynamicMaterialModel* MaterialModel = ModelEditorOnlyData->GetMaterialModel();
+	if (!ensure(MaterialModel))
+	{
+		return;
+	}
+
+	const TArray<UDMMaterialValue*>& Values = MaterialModel->GetValues();
+
+	if (Values.IsEmpty())
+	{
+		return;
+	}
+
+	for (int32 Index = 0; Index < Values.Num(); ++Index)
+	{
+		const bool bAcceptsWholeChannel = InThroughput->CanInputAcceptType(InInputIndex, Values[Index]->GetType());
+		const bool bAcceptsSubChannels = CanAcceptSubChannels(InThroughput, InInputIndex, Values[Index]->GetType());
+
+		UE::DynamicMaterialEditor::Private::GenerateChangeInputMenu_GlobalValue_Impl(InChildMenuBuilder,
+			Values[Index]->GetDescription(), Stage, InThroughput, InInputIndex, InInputChannel, Index, Values[Index],
+			bAcceptsWholeChannel, bAcceptsSubChannels);
+	}
+}
+
+void FDMMaterialStageInputMenus::GenerateChangeInputMenu_GlobalValue_Channels(FMenuBuilder& InChildMenuBuilder, UDMMaterialStageThroughput* InThroughput,
+	const int32 InInputIndex, const int32 InInputChannel, const int32 InValueIndex)
+{
+	if (!ensure(InThroughput))
+	{
+		return;
+	}
+
+	UDMMaterialStage* Stage = InThroughput->GetStage();
+	if (!ensure(Stage))
+	{
+		return;
+	}
+
+	UDMMaterialLayerObject* Layer = Stage->GetLayer();
+	if (!ensure(Layer))
+	{
+		return;
+	}
+
+	const TArray<FDMMaterialStageConnector>& InputConnectors = InThroughput->GetInputConnectors();
+	if (!ensure(InputConnectors.IsValidIndex(InInputIndex)))
+	{
+		return;
+	}
+
+	UDMMaterialSlot* Slot = Layer->GetSlot();
+	if (!ensure(Slot))
+	{
+		return;
+	}
+
+	UDynamicMaterialModelEditorOnlyData* ModelEditorOnlyData = Slot->GetMaterialModelEditorOnlyData();
+	if (!ensure(ModelEditorOnlyData))
+	{
+		return;
+	}
+
+	UDynamicMaterialModel* MaterialModel = ModelEditorOnlyData->GetMaterialModel();
+	if (!ensure(MaterialModel))
+	{
+		return;
+	}
+
+	const TArray<UDMMaterialValue*>& Values = MaterialModel->GetValues();
+	if (!ensure(Values.IsValidIndex(InValueIndex))
+		|| !ensure(CanAcceptSubChannels(InThroughput, InInputIndex, Values[InValueIndex]->GetType())))
+	{
+		return;
+	}
+
+	int32 FloatSize = UDMValueDefinitionLibrary::GetValueDefinition(Values[InValueIndex]->GetType()).GetFloatCount();
+	if (!ensure(FloatSize > 1))
+	{
+		return;
+	}
+
+	InChildMenuBuilder.BeginSection("ChangeInputOutputChannels", LOCTEXT("ChangeInputOutputChannels", "Output Channels"));
+	{
+		for (int32 ChannelIndex = 0; ChannelIndex < FloatSize; ++ChannelIndex)
+		{
+			const FText ChannelName = UDMValueDefinitionLibrary::GetValueDefinition(Values[InValueIndex]->GetType()).GetChannelName(ChannelIndex + 1);
+			const int32 OutputChannel = UE::DynamicMaterialEditor::ChannelIndexToChannelBit(ChannelIndex + 1);
+
+			UE::DynamicMaterialEditor::Private::GenerateChangeInputMenu_GlobalValue_Channel_Impl(InChildMenuBuilder,
+				ChannelName, Stage, InThroughput, InInputIndex, InInputChannel, InValueIndex, Values[InValueIndex], OutputChannel);
+		}
+	}
+	InChildMenuBuilder.EndSection();
+}
+
+void FDMMaterialStageInputMenus::GenerateChangeInputMenu_NewGlobalValues(FMenuBuilder& InChildMenuBuilder, UDMMaterialStageThroughput* InThroughput, 
+	const int32 InInputIndex, const int32 InInputChannel)
+{
+	if (!ensure(InThroughput))
+	{
+		return;
+	}
+
+	UDMMaterialStage* Stage = InThroughput->GetStage();
+	if (!ensure(Stage))
+	{
+		return;
+	}
+
+	const TArray<FDMMaterialStageConnector>& InputConnectors = InThroughput->GetInputConnectors();
+	if (!ensure(InputConnectors.IsValidIndex(InInputIndex)))
+	{
+		return;
+	}
+
+	for (EDMValueType ValueType : UDMValueDefinitionLibrary::GetValueTypes())
+	{
+		if (!InThroughput->CanInputAcceptType(InInputIndex, ValueType))
+		{
+			continue;
+		}
+
+		const FText Name = UDMValueDefinitionLibrary::GetValueDefinition(ValueType).GetDisplayName();
+		const FText FormattedTooltip = FText::Format(LOCTEXT("ChangeInputNewGlobalValueInputTooltipTemplate", 
+			"Add a new {0} Global Material Value and use it as the source of this input."), Name);
+
+		InChildMenuBuilder.AddMenuEntry(
+			Name,
+			FormattedTooltip,
+			FSlateIcon(),
+			FUIAction(FExecuteAction::CreateWeakLambda(
+				Stage,
+				[Stage, InInputIndex, InInputChannel, ValueType]()
+				{
+					FScopedTransaction Transaction(LOCTEXT("SetStageInput", "Set Material Stage Input"));
+					Stage->Modify();
+					Stage->ChangeInput_NewValue(InInputIndex, InInputChannel, ValueType, FDMMaterialStageConnectorChannel::WHOLE_CHANNEL);
+				})
+			)
+		);
+	}
+}
+
+namespace UE::DynamicMaterialEditor::Private
+{
+	void GenerateChangeInputMenu_Slot_Properties_Impl(FMenuBuilder& InChildMenuBuilder, const FText& InMenuName,
+		UDMMaterialStage* InStage, UDMMaterialStageThroughput* InThroughput, int32 InInputIndex, int32 InInputChannel,
+		UDMMaterialSlot* InSlot)
+	{
+		InChildMenuBuilder.AddSubMenu(
+			InMenuName,
+			LOCTEXT("ChangeInputSlotTooltip", "Change the source of this input to the output from this Material Slot."),
+			FNewMenuDelegate::CreateStatic(
+				&FDMMaterialStageInputMenus::GenerateChangeInputMenu_Slot_Properties,
+				InThroughput,
+				InInputIndex,
+				InInputChannel,
+				InSlot
+			)
+		);
+	}
+
+	void GenerateChangeInputMenu_Slot_Property_Outputs_Impl(FMenuBuilder& InChildMenuBuilder, const FText& InMenuName,
+		UDMMaterialStage* InStage, UDMMaterialStageThroughput* InThroughput, int32 InInputIndex, int32 InInputChannel, 
+		UDMMaterialSlot* InSlot, EDMMaterialPropertyType InSlotProperty)
+	{
+		InChildMenuBuilder.AddSubMenu(
+			InMenuName,
+			LOCTEXT("ChangeInputSlotTooltip", "Change the source of this input to the output from this Material Slot."),
+			FNewMenuDelegate::CreateStatic(
+				&FDMMaterialStageInputMenus::GenerateChangeInputMenu_Slot_Property_Outputs,
+				InThroughput,
+				InInputIndex,
+				InInputChannel,
+				InSlot,
+				InSlotProperty
+			)
+		);
+	}
+
+	void GenerateChangeInputMenu_Slot_Property_Output_Channels_Impl(FMenuBuilder& InChildMenuBuilder, const FText& InMenuName,
+		UDMMaterialStage* InStage, UDMMaterialStageThroughput* InThroughput, int32 InInputIndex, int32 InInputChannel, 
+		UDMMaterialSlot* InSlot, EDMMaterialPropertyType InSlotProperty, int32 InOutputIndex)
+	{
+		InChildMenuBuilder.AddSubMenu(
+			InMenuName,
+			LOCTEXT("ChangeInputSlotTooltip", "Change the source of this input to the output from this Material Slot."),
+			FNewMenuDelegate::CreateStatic(
+				&FDMMaterialStageInputMenus::GenerateChangeInputMenu_Slot_Property_Output_Channels,
+				InThroughput,
+				InInputIndex,
+				InInputChannel,
+				InSlot,
+				InSlotProperty,
+				InOutputIndex
+			)
+		);
+	}
+
+	void GenerateChangeInputMenu_Slot_Property_Output_Channel_Impl(FMenuBuilder& InChildMenuBuilder, const FText& InMenuName,
+		UDMMaterialStage* InStage, UDMMaterialStageThroughput* InThroughput, int32 InInputIndex, int32 InInputChannel, 
+		UDMMaterialSlot* InSlot, EDMMaterialPropertyType InSlotProperty, int32 InOutputIndex, int32 InOutputChannel)
+	{
+		InChildMenuBuilder.AddMenuEntry(
+			InMenuName,
+			LOCTEXT("ChangeInputSlotTooltip", "Change the source of this input to the output from this Material Slot."),
+			FSlateIcon(),
+			FUIAction(FExecuteAction::CreateWeakLambda(
+				InStage,
+				[InStage, InInputIndex, InInputChannel, SlotWeak = TWeakObjectPtr<UDMMaterialSlot>(InSlot), InSlotProperty, InOutputIndex, InOutputChannel]()
+				{
+					if (UDMMaterialSlot* Slot = SlotWeak.Get())
+					{
+						FScopedTransaction Transaction(LOCTEXT("SetStageInput", "Set Material Stage Input"));
+						InStage->Modify();
+						Slot->Modify();
+						InStage->ChangeInput_Slot(InInputIndex, InInputChannel, Slot, InSlotProperty, InOutputIndex, InOutputChannel);
+					}
+				})
+			)
+		);
+	}
+
+	void GenerateChangeInputMenu_Slot_Property_Output_Channel_Channels_Impl(FMenuBuilder& InChildMenuBuilder, const FText& InMenuName,
+		UDMMaterialStage* InStage, UDMMaterialStageThroughput* InThroughput, int32 InInputIndex, int32 InInputChannel, 
+		UDMMaterialSlot* InSlot, EDMMaterialPropertyType InSlotProperty, int32 InOutputIndex)
+	{
+		InChildMenuBuilder.AddSubMenu(
+			InMenuName,
+			LOCTEXT("ChangeInputSlotTooltip", "Change the source of this input to the output from this Material Slot."),
+			FNewMenuDelegate::CreateStatic(
+				&FDMMaterialStageInputMenus::GenerateChangeInputMenu_Slot_Property_Output_Channels,
+				InThroughput,
+				InInputIndex,
+				InInputChannel,
+				InSlot,
+				InSlotProperty,
+				InOutputIndex
+			),
+			FUIAction(FExecuteAction::CreateWeakLambda(
+				InStage,
+				[InStage, InInputIndex, InInputChannel, SlotWeak = TWeakObjectPtr<UDMMaterialSlot>(InSlot), InSlotProperty, InOutputIndex]()
+				{
+					if (UDMMaterialSlot* Slot = SlotWeak.Get())
+					{
+						FScopedTransaction Transaction(LOCTEXT("SetStageInput", "Set Material Stage Input"));
+						InStage->Modify();
+						InStage->ChangeInput_Slot(InInputIndex, InInputChannel, Slot, InSlotProperty, InOutputIndex, FDMMaterialStageConnectorChannel::WHOLE_CHANNEL);
+					}
+				})
+			),
+			NAME_None,
+			EUserInterfaceActionType::Button
+		);
+	}
+
+	void GenerateChangeInputMenu_Slot_Property_Output_Impl(FMenuBuilder& InChildMenuBuilder, const FText& InMenuName, 
+		UDMMaterialStage* InStage, UDMMaterialStageThroughput* InThroughput, int32 InInputIndex, int32 InInputChannel, 
+		UDMMaterialSlot* InSlot, EDMMaterialPropertyType InSlotProperty, int32 InOutputIndex, bool bInAcceptsWholeChannel, bool bInAcceptsSubChannels)
+	{
+		if (!bInAcceptsWholeChannel && !bInAcceptsSubChannels)
+		{
+			return;
+		}
+
+		if (bInAcceptsSubChannels)
+		{
+			if (bInAcceptsWholeChannel)
+			{
+				UE::DynamicMaterialEditor::Private::GenerateChangeInputMenu_Slot_Property_Output_Channel_Channels_Impl(InChildMenuBuilder, InMenuName, InStage, InThroughput,
+					InInputIndex, InInputChannel, InSlot, InSlotProperty, InOutputIndex);
+			}
+			else
+			{
+				UE::DynamicMaterialEditor::Private::GenerateChangeInputMenu_Slot_Property_Output_Channels_Impl(InChildMenuBuilder, InMenuName, InStage, InThroughput,
+					InInputIndex, InInputChannel, InSlot, InSlotProperty, InOutputIndex);
+			}
+		}
+		else
+		{
+			UE::DynamicMaterialEditor::Private::GenerateChangeInputMenu_Slot_Property_Output_Channel_Impl(InChildMenuBuilder, InMenuName, InStage, InThroughput,
+				InInputIndex, InInputChannel, InSlot, InSlotProperty, InOutputIndex, FDMMaterialStageConnectorChannel::WHOLE_CHANNEL);
+		}
+	}
+}
+
+void FDMMaterialStageInputMenus::GenerateChangeInputMenu_Slots(FMenuBuilder& InChildMenuBuilder, UDMMaterialStageThroughput* InThroughput, 
+	const int32 InInputIndex, const int32 InInputChannel)
+{
+	if (!ensure(InThroughput))
+	{
+		return;
+	}
+
+	UDMMaterialStage* Stage = InThroughput->GetStage();
+	if (!ensure(Stage))
+	{
+		return;
+	}
+
+	UDMMaterialLayerObject* Layer = Stage->GetLayer();
+	if (!ensure(Layer))
+	{
+		return;
+	}
+
+	const TArray<FDMMaterialStageConnector>& InputConnectors = InThroughput->GetInputConnectors();
+	if (!ensure(InputConnectors.IsValidIndex(InInputIndex)))
+	{
+		return;
+	}
+
+	UDMMaterialSlot* Slot = Layer->GetSlot();
+	if (!ensure(Slot))
+	{
+		return;
+	}
+
+	UDynamicMaterialModelEditorOnlyData* ModelEditorOnlyData = Slot->GetMaterialModelEditorOnlyData();
+	if (!ensure(ModelEditorOnlyData))
+	{
+		return;
+	}
+
+	const TArray<UDMMaterialSlot*>& Slots = ModelEditorOnlyData->GetSlots();
+
+	if (Slots.Num() <= 1)
+	{
+		return;
+	}
+
+	for (UDMMaterialSlot* SlotIter : Slots)
+	{
+		if (Slot == SlotIter)
+		{
+			continue;
+		}
+
+		if (SlotIter->GetLayers().IsEmpty())
+		{
+			continue;
+		}
+
+		const TArray<EDMMaterialPropertyType> SlotProperties = ModelEditorOnlyData->GetMaterialPropertiesForSlot(SlotIter);
+
+		if (SlotProperties.IsEmpty())
+		{
+			continue;
+		}
+
+		const TSet<EDMValueType> SlotOutputTypes = SlotIter->GetAllOutputConnectorTypes();
+		bool bFoundCompatibleOutput = false;
+
+		for (EDMValueType OutputType : SlotOutputTypes)
+		{
+			if (InThroughput->CanInputAcceptType(InInputIndex, OutputType) || CanAcceptSubChannels(InThroughput, InInputIndex, OutputType))
+			{
+				bFoundCompatibleOutput = true;
+				break;
+			}
+		}
+
+		if (!bFoundCompatibleOutput)
+		{
+			continue;
+		}
+
+		const bool bOneSlotProperty = (SlotProperties.Num() == 1);
+
+		if (bOneSlotProperty)
+		{
+			static const FText SlotNameFormatTemplate = LOCTEXT("SlotAndProperty", "{0} [{1}]");
+
+			UDMMaterialProperty* PropertyObj = ModelEditorOnlyData->GetMaterialProperty(SlotProperties[0]);
+			if (!ensure(PropertyObj))
+			{
+				continue;
+			}
+
+			const FText SlotPropertyName = FText::Format(SlotNameFormatTemplate, SlotIter->GetDescription(), PropertyObj->GetDescription());
+
+			const TArray<EDMValueType>& SlotPropertyOutputTypes = SlotIter->GetOutputConnectorTypesForMaterialProperty(SlotProperties[0]);
+			const bool bOneOutputType = (SlotPropertyOutputTypes.Num() == 1);
+
+			if (bOneOutputType)
+			{
+				const UDMMaterialLayerObject* PropertyLayer = SlotIter->GetLastLayerForMaterialProperty(SlotProperties[0]);
+
+				// Would use ensure, but stages might have been disabled or removed.
+				if (!PropertyLayer)
+				{
+					continue;
+				}
+
+				if (!ensure(PropertyLayer->GetStage(EDMMaterialLayerStage::Mask)))
+				{
+					continue;
+				}
+
+				UDMMaterialStageSource* LastPropertySource = PropertyLayer->GetStage(EDMMaterialLayerStage::Mask)->GetSource();
+				if (!ensure(LastPropertySource))
+				{
+					continue;
+				}
+
+				const TArray<FDMMaterialStageConnector>& LastPropertySourceOutputConnectors = LastPropertySource->GetOutputConnectors();
+				if (!ensure(LastPropertySourceOutputConnectors.IsValidIndex(0)))
+				{
+					continue;
+				}
+
+				const bool bAcceptsWholeChannel = InThroughput->CanInputConnectTo(InInputIndex, LastPropertySourceOutputConnectors[0], FDMMaterialStageConnectorChannel::WHOLE_CHANNEL);
+				const bool bAcceptsSubChannels = CanAcceptSubChannels(InThroughput, InInputIndex, LastPropertySourceOutputConnectors[0]);
+
+				UE::DynamicMaterialEditor::Private::GenerateChangeInputMenu_Slot_Property_Output_Impl(InChildMenuBuilder, SlotPropertyName,
+					Stage, InThroughput, InInputIndex, InInputChannel, SlotIter, SlotProperties[0], 0, bAcceptsWholeChannel, bAcceptsSubChannels);
+			}
+			else // !bOneOutputType
+			{
+				UE::DynamicMaterialEditor::Private::GenerateChangeInputMenu_Slot_Property_Outputs_Impl(InChildMenuBuilder, SlotPropertyName,
+					Stage, InThroughput, InInputIndex, InInputChannel, SlotIter, SlotProperties[0]);
+			}
+		}
+		else
+		{
+			UE::DynamicMaterialEditor::Private::GenerateChangeInputMenu_Slot_Properties_Impl(InChildMenuBuilder, SlotIter->GetDescription(),
+				Stage, InThroughput, InInputIndex, InInputChannel, SlotIter);
+		}
+	}
+}
+
+void FDMMaterialStageInputMenus::GenerateChangeInputMenu_Slot_Properties(FMenuBuilder& InChildMenuBuilder, UDMMaterialStageThroughput* InThroughput, 
+	const int32 InInputIndex, const int32 InInputChannel, UDMMaterialSlot* InSlot)
+{
+	if (!ensure(InThroughput))
+	{
+		return;
+	}
+
+	UDMMaterialStage* Stage = InThroughput->GetStage();
+	if (!ensure(Stage))
+	{
+		return;
+	}
+
+	UDMMaterialLayerObject* Layer = Stage->GetLayer();
+	if (!ensure(Layer))
+	{
+		return;
+	}
+
+	UDMMaterialSlot* Slot = Layer->GetSlot();
+	if (!ensure(Slot))
+	{
+		return;
+	}
+
+	const TArray<FDMMaterialStageConnector>& InputConnectors = InThroughput->GetInputConnectors();
+	if (!ensure(InputConnectors.IsValidIndex(InInputIndex)))
+	{
+		return;
+	}
+
+	UDynamicMaterialModelEditorOnlyData* ModelEditorOnlyData = InSlot->GetMaterialModelEditorOnlyData();
+	if (!ensure(ModelEditorOnlyData))
+	{
+		return;
+	}
+
+	if (!ensure(IsValid(InSlot)) || !ensure(Slot) || !ensure(InSlot != Slot)
+		|| !ensure(ModelEditorOnlyData == Slot->GetMaterialModelEditorOnlyData()))
+	{
+		return;
+	}
+
+	InChildMenuBuilder.BeginSection("ChangeInputMaterialProperties", LOCTEXT("ChangeInputMaterialProperties", "Material Properties"));
+	{
+		for (EDMMaterialPropertyType SlotProperty : ModelEditorOnlyData->GetMaterialPropertiesForSlot(InSlot))
+		{
+			const TArray<EDMValueType>& SlotPropertyOutputTypes = InSlot->GetOutputConnectorTypesForMaterialProperty(SlotProperty);
+			bool bFoundCompatibleOutput = false;
+
+			for (EDMValueType ValueType : SlotPropertyOutputTypes)
+			{
+				if (InThroughput->CanInputAcceptType(InInputIndex, ValueType) || CanAcceptSubChannels(InThroughput, InInputIndex, ValueType))
+				{
+					bFoundCompatibleOutput = true;
+					break;
+				}
+			}
+
+			if (!bFoundCompatibleOutput)
+			{
+				continue;
+			}
+
+			UDMMaterialProperty* PropertyObj = ModelEditorOnlyData->GetMaterialProperty(SlotProperty);
+			if (!ensure(PropertyObj))
+			{
+				continue;
+			}
+
+			const bool bOneOutputType = (SlotPropertyOutputTypes.Num() == 1);
+
+			if (bOneOutputType)
+			{
+				const UDMMaterialLayerObject* PropertyLayer = InSlot->GetLastLayerForMaterialProperty(SlotProperty);
+
+				// Would use ensure, but stages might have been disabled or removed.
+				if (!PropertyLayer)
+				{
+					continue;
+				}
+
+				if (!ensure(PropertyLayer->GetStage(EDMMaterialLayerStage::Mask)))
+				{
+					continue;
+				}
+
+				UDMMaterialStageSource* LastPropertySource = PropertyLayer->GetStage(EDMMaterialLayerStage::Mask)->GetSource();
+				if (!ensure(LastPropertySource))
+				{
+					continue;
+				}
+
+				const TArray<FDMMaterialStageConnector>& LastPropertySourceOutputConnectors = LastPropertySource->GetOutputConnectors();
+				if (!ensure(LastPropertySourceOutputConnectors.IsValidIndex(0)))
+				{
+					continue;
+				}
+
+				const bool bAcceptsWholeChannel = InThroughput->CanInputConnectTo(InInputIndex, LastPropertySourceOutputConnectors[0], FDMMaterialStageConnectorChannel::WHOLE_CHANNEL);
+				const bool bAcceptsSubChannels = CanAcceptSubChannels(InThroughput, InInputIndex, LastPropertySourceOutputConnectors[0]);
+
+				UE::DynamicMaterialEditor::Private::GenerateChangeInputMenu_Slot_Property_Output_Impl(InChildMenuBuilder, PropertyObj->GetDescription(),
+					Stage, InThroughput, InInputIndex, InInputChannel, InSlot, SlotProperty, 0, bAcceptsWholeChannel, bAcceptsSubChannels);
+			}
+			else // !bOneOutputType
+			{
+				UE::DynamicMaterialEditor::Private::GenerateChangeInputMenu_Slot_Property_Outputs_Impl(InChildMenuBuilder, PropertyObj->GetDescription(),
+					Stage, InThroughput, InInputIndex, InInputChannel, InSlot, SlotProperty);
+			}
+		}
+	}
+	InChildMenuBuilder.EndSection();
+}
+
+void FDMMaterialStageInputMenus::GenerateChangeInputMenu_Slot_Property_Outputs(FMenuBuilder& InChildMenuBuilder, UDMMaterialStageThroughput* InThroughput, 
+	const int32 InInputIndex, const int32 InInputChannel, UDMMaterialSlot* InSlot, EDMMaterialPropertyType InMaterialProperty)
+{
+	if (!ensure(InThroughput))
+	{
+		return;
+	}
+
+	UDMMaterialStage* Stage = InThroughput->GetStage();
+	if (!ensure(Stage))
+	{
+		return;
+	}
+
+	UDMMaterialLayerObject* Layer = Stage->GetLayer();
+	if (!ensure(Layer))
+	{
+		return;
+	}
+
+	UDMMaterialSlot* Slot = Layer->GetSlot();
+	if (!ensure(Slot))
+	{
+		return;
+	}
+
+	const TArray<FDMMaterialStageConnector>& InputConnectors = InThroughput->GetInputConnectors();
+	if (!ensure(InputConnectors.IsValidIndex(InInputIndex)))
+	{
+		return;
+	}
+
+	UDynamicMaterialModelEditorOnlyData* ModelEditorOnlyData = InSlot->GetMaterialModelEditorOnlyData();
+	if (!ensure(ModelEditorOnlyData))
+	{
+		return;
+	}
+
+	if (!ensure(IsValid(InSlot)) || !ensure(Slot) || !ensure(InSlot != Slot)
+		|| !ensure(ModelEditorOnlyData == Slot->GetMaterialModelEditorOnlyData()))
+	{
+		return;
+	}
+
+	UDMMaterialSlot* PropertySlot = ModelEditorOnlyData->GetSlotForMaterialProperty(InMaterialProperty);
+	if (!ensure(PropertySlot == InSlot))
+	{
+		return;
+	}
+
+	const TArray<EDMValueType>& SlotPropertyOutputTypes = InSlot->GetOutputConnectorTypesForMaterialProperty(InMaterialProperty);
+
+	InChildMenuBuilder.BeginSection("ChangeInputOutputs", LOCTEXT("ChangeInputOutputs", "Outputs"));
+	{
+		for (int32 OutputIdx = 0; OutputIdx < SlotPropertyOutputTypes.Num(); ++OutputIdx)
+		{
+			const UDMMaterialLayerObject* PropertyLayer = InSlot->GetLastLayerForMaterialProperty(InMaterialProperty);
+
+			// Would use ensure, but stages might have been disabled or removed.
+			if (!PropertyLayer)
+			{
+				continue;
+			}
+
+			if (!ensure(PropertyLayer->GetStage(EDMMaterialLayerStage::Mask)))
+			{
+				continue;
+			}
+
+			UDMMaterialStageSource* LastPropertySource = PropertyLayer->GetStage(EDMMaterialLayerStage::Mask)->GetSource();
+			if (!ensure(LastPropertySource))
+			{
+				continue;
+			}
+
+			const TArray<FDMMaterialStageConnector>& LastPropertySourceOutputConnectors = LastPropertySource->GetOutputConnectors();
+			if (!ensure(LastPropertySourceOutputConnectors.IsValidIndex(OutputIdx)))
+			{
+				continue;
+			}
+
+			const bool bAcceptsWholeChannel = InThroughput->CanInputConnectTo(InInputIndex, LastPropertySourceOutputConnectors[OutputIdx], FDMMaterialStageConnectorChannel::WHOLE_CHANNEL);
+			const bool bAcceptsSubChannels = CanAcceptSubChannels(InThroughput, InInputIndex, LastPropertySourceOutputConnectors[OutputIdx]);
+
+			UE::DynamicMaterialEditor::Private::GenerateChangeInputMenu_Slot_Property_Output_Impl(InChildMenuBuilder, LastPropertySourceOutputConnectors[OutputIdx].Name,
+				Stage, InThroughput, InInputIndex, InInputChannel, InSlot, InMaterialProperty, OutputIdx, bAcceptsWholeChannel, bAcceptsSubChannels);
+		}
+	}
+	InChildMenuBuilder.EndSection();
+}
+
+void FDMMaterialStageInputMenus::GenerateChangeInputMenu_Slot_Property_Output_Channels(FMenuBuilder& InChildMenuBuilder, UDMMaterialStageThroughput* InThroughput,
+	const int32 InInputIndex, const int32 InInputChannel, UDMMaterialSlot* InSlot, EDMMaterialPropertyType InMaterialProperty, int32 OutputIndex)
+{
+	if (!ensure(InThroughput))
+	{
+		return;
+	}
+
+	UDMMaterialStage* Stage = InThroughput->GetStage();
+	if (!ensure(Stage))
+	{
+		return;
+	}
+
+	UDMMaterialLayerObject* Layer = Stage->GetLayer();
+	if (!ensure(Layer))
+	{
+		return;
+	}
+
+	UDMMaterialSlot* Slot = Layer->GetSlot();
+	if (!ensure(Slot))
+	{
+		return;
+	}
+
+	const TArray<FDMMaterialStageConnector>& InputConnectors = InThroughput->GetInputConnectors();
+	if (!ensure(InputConnectors.IsValidIndex(InInputIndex)))
+	{
+		return;
+	}
+
+	UDynamicMaterialModelEditorOnlyData* ModelEditorOnlyData = InSlot->GetMaterialModelEditorOnlyData();
+	if (!ensure(ModelEditorOnlyData))
+	{
+		return;
+	}
+
+	if (!ensure(IsValid(InSlot)) || !ensure(Slot) || !ensure(InSlot != Slot)
+		|| !ensure(ModelEditorOnlyData == Slot->GetMaterialModelEditorOnlyData()))
+	{
+		return;
+	}
+
+	UDMMaterialSlot* PropertySlot = ModelEditorOnlyData->GetSlotForMaterialProperty(InMaterialProperty);
+	if (!ensure(PropertySlot == InSlot))
+	{
+		return;
+	}
+
+	const UDMMaterialLayerObject* PropertyLayer = InSlot->GetLastLayerForMaterialProperty(InMaterialProperty);
+	if (!ensure(PropertyLayer) || !ensure(PropertyLayer->GetStage(EDMMaterialLayerStage::Mask)))
+	{
+		return;
+	}
+
+	UDMMaterialStageSource* LastPropertySource = PropertyLayer->GetStage(EDMMaterialLayerStage::Mask)->GetSource();
+	if (!ensure(LastPropertySource))
+	{
+		return;
+	}
+
+	const TArray<FDMMaterialStageConnector>& LastPropertySourceOutputConnectors = LastPropertySource->GetOutputConnectors();
+	if (!ensure(LastPropertySourceOutputConnectors.IsValidIndex(OutputIndex))
+		|| !ensure(CanAcceptSubChannels(InThroughput, InInputIndex, LastPropertySourceOutputConnectors[OutputIndex])))
+	{
+		return;
+	}
+
+	const int32 FloatSize = UDMValueDefinitionLibrary::GetValueDefinition(LastPropertySourceOutputConnectors[OutputIndex].Type).GetFloatCount();
+
+	if (LastPropertySourceOutputConnectors[OutputIndex].Type != EDMValueType::VT_Float_Any)
+	{
+		if (!ensure(FloatSize > 1))
+		{
+			return;
+		}
+	}
+
+	InChildMenuBuilder.BeginSection("ChangeInputOutputChannels", LOCTEXT("ChangeInputOutputChannels", "Output Channels"));
+	{
+		for (int32 ChannelIndex = 0; ChannelIndex < FloatSize; ++ChannelIndex)
+		{
+			const FText ChannelName = UDMValueDefinitionLibrary::GetValueDefinition(LastPropertySourceOutputConnectors[OutputIndex].Type).GetChannelName(ChannelIndex + 1);
+			const int32 OutputChannel = UE::DynamicMaterialEditor::ChannelIndexToChannelBit(ChannelIndex + 1);
+
+			UE::DynamicMaterialEditor::Private::GenerateChangeInputMenu_Slot_Property_Output_Channel_Impl(InChildMenuBuilder, ChannelName,
+				Stage, InThroughput, InInputIndex, InInputChannel, InSlot, InMaterialProperty, OutputIndex, OutputChannel);
+		}
+	}
+	InChildMenuBuilder.EndSection();
+}
+
+void FDMMaterialStageInputMenus::GenerateChangeInputMenu_Expressions(FMenuBuilder& InChildMenuBuilder, UDMMaterialStageThroughput* InThroughput, int32 InInputIndex, int32 InInputChannel)
+{
+	if (!ensure(InThroughput))
+	{
+		return;
+	}
+
+	UDMMaterialStage* Stage = InThroughput->GetStage();
+	if (!ensure(Stage))
+	{
+		return;
+	}
+
+	UDMMaterialLayerObject* Layer = Stage->GetLayer();
+	if (!ensure(Layer))
+	{
+		return;
+	}
+
+	UDMMaterialSlot* Slot = Layer->GetSlot();
+	if (!ensure(Slot))
+	{
+		return;
+	}
+
+	const TArray<FDMMaterialStageConnector>& InputConnectors = InThroughput->GetInputConnectors();
+	if (!ensure(InputConnectors.IsValidIndex(InInputIndex)))
+	{
+		return;
+	}
+
+	const TArray<TStrongObjectPtr<UClass>>& Expressions = UDMMaterialStageExpression::GetAvailableSourceExpressions();
+
+	if (Expressions.IsEmpty())
+	{
+		return;
+	}
+
+	TArray<TStrongObjectPtr<UClass>> FilteredExpressions;
+
+	for (const TStrongObjectPtr<UClass>& ExpressionClass : Expressions)
+	{
+		TSubclassOf<UDMMaterialStageExpression> InputClass = ExpressionClass.Get();
+
+		if (!InputClass)
+		{
+			continue;
+		}
+
+		UDMMaterialStageExpression* ExpressionCDO = Cast<UDMMaterialStageExpression>(InputClass->GetDefaultObject(true));
+		if (!ensure(ExpressionCDO))
+		{
+			continue;
+		}
+
+		if (!ExpressionCDO->IsInputRequired() || ExpressionCDO->AllowsNestedInputs())
+		{
+			const TArray<FDMMaterialStageConnector>& OutputConnectors = ExpressionCDO->GetOutputConnectors();
+			int32 ValidOutputConnectorCount = 0;
+
+			for (int32 OutputIdx = 0; OutputIdx < OutputConnectors.Num(); ++OutputIdx)
+			{
+				if (InThroughput->CanInputConnectTo(InInputIndex, OutputConnectors[OutputIdx], FDMMaterialStageConnectorChannel::WHOLE_CHANNEL, true))
+				{
+					++ValidOutputConnectorCount;
+				}
+			}
+
+			if (ValidOutputConnectorCount == 0)
+			{
+				continue;
+			}
+
+			FilteredExpressions.Add(ExpressionClass);
+		}
+	}
+
+	auto CreateTreeMenu = [&InChildMenuBuilder, &InputConnectors, InThroughput, InInputIndex, InInputChannel](EDMExpressionMenu InMenu, const TArray<UDMMaterialStageExpression*>& InSubmenuExpressionList)
+	{
+		FText MenuName = StaticEnum<EDMExpressionMenu>()->GetDisplayNameTextByValue(static_cast<int64>(InMenu));
+
+		InChildMenuBuilder.AddSubMenu(
+			MenuName,
+			LOCTEXT("AddMaterialStageMaterialExpressionTooltip", "Add Material Stage based on one of these Material Expressions."),
+			FNewMenuDelegate::CreateWeakLambda(
+				InThroughput,
+				[InSubmenuExpressionList, MenuName, InThroughput, InputConnectors, InInputIndex, InInputChannel](FMenuBuilder& ChildSubMenuBuilder)
+				{
+					ChildSubMenuBuilder.BeginSection(NAME_None, MenuName);
+
+					for (UDMMaterialStageExpression* ExpressionCDO : InSubmenuExpressionList)
+					{
+						GenerateChangeInputMenu_Expression(ChildSubMenuBuilder, ExpressionCDO, InThroughput, InInputIndex, InInputChannel);
+					}
+
+					ChildSubMenuBuilder.EndSection();
+				})
+		);
+	};
+
+	FDMMaterialStageSourceMenus::CreateSourceMenuTree(CreateTreeMenu, FilteredExpressions);
+}
+
+namespace UE::DynamicMaterialEditor::Private
+{
+	void GenerateChangeInputMenu_Expression_Outputs_Impl(FMenuBuilder& InChildMenuBuilder, const FText& InMenuName,
+		UDMMaterialStage* InStage, UDMMaterialStageThroughput* InThroughput, int32 InInputIndex, int32 InInputChannel, 
+		TSubclassOf<UDMMaterialStageExpression> InExpressionClass)
+	{
+		InChildMenuBuilder.AddSubMenu(
+			InMenuName,
+			LOCTEXT("ChangeInputExpressionTooltip", "Change the source of this input to a Material Expression."),
+			FNewMenuDelegate::CreateStatic(
+				&FDMMaterialStageInputMenus::GenerateChangeInputMenu_Expression_Outputs,
+				InThroughput,
+				InInputIndex,
+				InInputChannel,
+				InExpressionClass
+			)
+		);
+	}
+
+	void GenerateChangeInputMenu_Expression_Output_Channels_Impl(FMenuBuilder& InChildMenuBuilder, const FText& InMenuName,
+		UDMMaterialStage* InStage, UDMMaterialStageThroughput* InThroughput, int32 InInputIndex, int32 InInputChannel, 
+		TSubclassOf<UDMMaterialStageExpression> InExpressionClass, int32 InOutputIndex)
+	{
+		InChildMenuBuilder.AddSubMenu(
+			InMenuName,
+			LOCTEXT("ChangeInputExpressionTooltip", "Change the source of this input to a Material Expression."),
+			FNewMenuDelegate::CreateStatic(
+				&FDMMaterialStageInputMenus::GenerateChangeInputMenu_Expression_Output_Channels,
+				InThroughput,
+				InInputIndex,
+				InInputChannel,
+				InExpressionClass,
+				InOutputIndex
+			)
+		);
+	}
+
+	void GenerateChangeInputMenu_Expression_Output_Channel_Impl(FMenuBuilder& InChildMenuBuilder, const FText& InMenuName,
+		UDMMaterialStage* InStage, UDMMaterialStageThroughput* InThroughput, int32 InInputIndex, int32 InInputChannel, 
+		TSubclassOf<UDMMaterialStageExpression> InExpressionClass, int32 InOutputIndex, int32 InOutputChannel)
+	{
+		InChildMenuBuilder.AddMenuEntry(
+			InMenuName,
+			LOCTEXT("ChangeInputExpressionTooltip", "Change the source of this input to a Material Expression."),
+			FSlateIcon(),
+			FUIAction(FExecuteAction::CreateWeakLambda(
+				InStage,
+				[InStage, InInputIndex, InInputChannel, InExpressionClass, InOutputIndex, InOutputChannel]()
+				{
+					FScopedTransaction Transaction(LOCTEXT("SetStageInput", "Set Material Stage Input"));
+					InStage->Modify();
+					InStage->ChangeInput_Expression(InInputIndex, InInputChannel, InExpressionClass, InOutputIndex, InOutputChannel);
+				})
+			)
+		);
+	}
+
+	void GenerateChangeInputMenu_Expression_Output_Channel_Channels_Impl(FMenuBuilder& InChildMenuBuilder, const FText& InMenuName,
+		UDMMaterialStage* InStage, UDMMaterialStageThroughput* InThroughput, int32 InInputIndex, int32 InInputChannel,
+		TSubclassOf<UDMMaterialStageExpression> InExpressionClass, int32 InOutputIndex)
+	{
+		InChildMenuBuilder.AddSubMenu(
+			InMenuName,
+			LOCTEXT("ChangeInputExpressionTooltip", "Change the source of this input to a Material Expression."),
+			FNewMenuDelegate::CreateStatic(
+				&FDMMaterialStageInputMenus::GenerateChangeInputMenu_Expression_Output_Channels,
+				InThroughput,
+				InInputIndex,
+				InInputChannel,
+				InExpressionClass,
+				InOutputIndex
+			),
+			FUIAction(FExecuteAction::CreateWeakLambda(
+				InStage,
+				[InStage, InInputIndex, InInputChannel, InExpressionClass, InOutputIndex]()
+				{
+					FScopedTransaction Transaction(LOCTEXT("SetStageInput", "Set Material Stage Input"));
+					InStage->Modify();
+					InStage->ChangeInput_Expression(InInputIndex, InInputChannel, InExpressionClass, InOutputIndex, FDMMaterialStageConnectorChannel::WHOLE_CHANNEL);
+				})
+			),
+			NAME_None,
+			EUserInterfaceActionType::Button
+		);
+	}
+
+	void GenerateChangeInputMenu_Expression_Output_Impl(FMenuBuilder& InChildMenuBuilder, const FText& InMenuName,
+		UDMMaterialStage* InStage, UDMMaterialStageThroughput* InThroughput, int32 InInputIndex, int32 InInputChannel, 
+		TSubclassOf<UDMMaterialStageExpression> InExpressionClass, int32 InOutputIndex, bool bInAcceptsWholeChannel, bool bInAcceptsSubChannels)
+	{
+		if (!bInAcceptsWholeChannel && !bInAcceptsSubChannels)
+		{
+			return;
+		}
+
+		if (bInAcceptsSubChannels)
+		{
+			if (bInAcceptsWholeChannel)
+			{
+				UE::DynamicMaterialEditor::Private::GenerateChangeInputMenu_Expression_Output_Channel_Channels_Impl(InChildMenuBuilder, InMenuName, InStage, InThroughput,
+					InInputIndex, InInputChannel, InExpressionClass, InOutputIndex);
+			}
+			else
+			{
+				UE::DynamicMaterialEditor::Private::GenerateChangeInputMenu_Expression_Output_Channels_Impl(InChildMenuBuilder, InMenuName, InStage, InThroughput,
+					InInputIndex, InInputChannel, InExpressionClass, InOutputIndex);
+			}
+		}
+		else
+		{
+			UE::DynamicMaterialEditor::Private::GenerateChangeInputMenu_Expression_Output_Channel_Impl(InChildMenuBuilder, InMenuName, InStage, InThroughput,
+				InInputIndex, InInputChannel, InExpressionClass, InOutputIndex, FDMMaterialStageConnectorChannel::WHOLE_CHANNEL);
+		}
+	}
+}
+
+void FDMMaterialStageInputMenus::GenerateChangeInputMenu_Expression(FMenuBuilder& InChildMenuBuilder, UDMMaterialStageExpression* InExpressionCDO, UDMMaterialStageThroughput* InThroughput,
+	const int32 InInputIndex, const int32 InInputChannel)
+{
+	UDMMaterialStage* Stage = InThroughput->GetStage();
+	if (!ensure(Stage))
+	{
+		return;
+	}
+
+	const TArray<FDMMaterialStageConnector>& OutputConnectors = InExpressionCDO->GetOutputConnectors();
+	int32 ValidOutputConnectorCount = 0;
+	int32 LastValidOutputConnectorIdx = INDEX_NONE;
+	TSubclassOf<UDMMaterialStageExpression> InputClass = InExpressionCDO->GetClass();
+
+	for (int32 OutputIdx = 0; OutputIdx < OutputConnectors.Num(); ++OutputIdx)
+	{
+		if (InThroughput->CanInputConnectTo(InInputIndex, OutputConnectors[OutputIdx], FDMMaterialStageConnectorChannel::WHOLE_CHANNEL)
+			|| CanAcceptSubChannels(InThroughput, InInputIndex, OutputConnectors[OutputIdx]))
+		{
+			++ValidOutputConnectorCount;
+			LastValidOutputConnectorIdx = OutputIdx;
+		}
+	}
+
+	if (ValidOutputConnectorCount == 0)
+	{
+		return;
+	}
+
+	if (ValidOutputConnectorCount == 1)
+	{
+		UMaterialExpression* MatExpressionCDO = Cast<UMaterialExpression>(InExpressionCDO->GetMaterialExpressionClass()->GetDefaultObject(true));
+		TArray<FString> Captions;
+		MatExpressionCDO->GetCaption(Captions);
+		TArray<FText> TextCaptions;
+
+		for (const FString& Caption : Captions)
+		{
+			TextCaptions.Add(FText::FromString(Caption));
+		}
+
+		const FText Tooltip = FText::Join(LOCTEXT("NewLine", "\n"), TextCaptions);
+
+		static const FText ExpressionNameFormatTemplate = LOCTEXT("ExpressionAndOutput", "{0} [{1}]");
+		const FText ExpressionName = FText::Format(ExpressionNameFormatTemplate, InExpressionCDO->GetDescription(), OutputConnectors[LastValidOutputConnectorIdx].Name);
+
+		const bool bAcceptsWholeChannel = InThroughput->CanInputConnectTo(InInputIndex, OutputConnectors[LastValidOutputConnectorIdx], FDMMaterialStageConnectorChannel::WHOLE_CHANNEL);
+		const bool bAcceptsSubChannels = CanAcceptSubChannels(InThroughput, InInputIndex, OutputConnectors[LastValidOutputConnectorIdx]);
+
+		UE::DynamicMaterialEditor::Private::GenerateChangeInputMenu_Expression_Output_Impl(InChildMenuBuilder, ExpressionName, Stage, InThroughput,
+			InInputIndex, InInputChannel, InputClass, LastValidOutputConnectorIdx, bAcceptsWholeChannel, bAcceptsSubChannels);
+	}
+	else
+	{
+		UE::DynamicMaterialEditor::Private::GenerateChangeInputMenu_Expression_Outputs_Impl(InChildMenuBuilder, InExpressionCDO->GetDescription(), 
+			Stage, InThroughput, InInputIndex, InInputChannel, InputClass);
+	}
+}
+
+void FDMMaterialStageInputMenus::GenerateChangeInputMenu_Expression_Outputs(FMenuBuilder& InChildMenuBuilder, UDMMaterialStageThroughput* InThroughput, 
+	const int32 InInputIndex, const int32 InInputChannel, TSubclassOf<UDMMaterialStageExpression> InExpressionClass)
+{
+	if (!ensure(InThroughput))
+	{
+		return;
+	}
+
+	UDMMaterialStage* Stage = InThroughput->GetStage();
+	if (!ensure(Stage))
+	{
+		return;
+	}
+
+	const TArray<FDMMaterialStageConnector>& InputConnectors = InThroughput->GetInputConnectors();
+	if (!ensure(InputConnectors.IsValidIndex(InInputIndex)))
+	{
+		return;
+	}
+
+	if (!ensure(InExpressionClass.Get()))
+	{
+		return;
+	}
+
+	UDMMaterialStageExpression* ExpressionCDO = Cast<UDMMaterialStageExpression>(InExpressionClass->GetDefaultObject(true));
+	if (!ensure(ExpressionCDO) || !ensure(!ExpressionCDO->IsInputRequired() || ExpressionCDO->AllowsNestedInputs()))
+	{
+		return;
+	}
+
+	const TArray<FDMMaterialStageConnector>& OutputConnectors = ExpressionCDO->GetOutputConnectors();
+
+	InChildMenuBuilder.BeginSection("ChangeInputOutputs", LOCTEXT("ChangeInputOutputs", "Outputs"));
+	{
+		for (int32 OutputIdx = 0; OutputIdx < OutputConnectors.Num(); ++OutputIdx)
+		{
+			const bool bAcceptsWholeChannel = InThroughput->CanInputConnectTo(InInputIndex, OutputConnectors[OutputIdx], FDMMaterialStageConnectorChannel::WHOLE_CHANNEL);
+			const bool bAcceptsSubChannels = CanAcceptSubChannels(InThroughput, InInputIndex, OutputConnectors[OutputIdx]);
+
+			UE::DynamicMaterialEditor::Private::GenerateChangeInputMenu_Expression_Output_Impl(InChildMenuBuilder, OutputConnectors[OutputIdx].Name,
+				Stage, InThroughput, InInputIndex, InInputChannel, InExpressionClass, OutputIdx, bAcceptsWholeChannel, bAcceptsSubChannels);
+		}
+	}
+	InChildMenuBuilder.EndSection();
+}
+
+void FDMMaterialStageInputMenus::GenerateChangeInputMenu_Expression_Output_Channels(FMenuBuilder& InChildMenuBuilder, UDMMaterialStageThroughput* InThroughput, 
+	const int32 InInputIndex, const int32 InInputChannel, TSubclassOf<UDMMaterialStageExpression> InExpressionClass, int32 InOutputIndex)
+{
+	if (!ensure(InThroughput))
+	{
+		return;
+	}
+
+	UDMMaterialStage* Stage = InThroughput->GetStage();
+	if (!ensure(Stage))
+	{
+		return;
+	}
+
+	const TArray<FDMMaterialStageConnector>& InputConnectors = InThroughput->GetInputConnectors();
+	if (!ensure(InputConnectors.IsValidIndex(InInputIndex)))
+	{
+		return;
+	}
+
+	if (!ensure(InExpressionClass.Get()))
+	{
+		return;
+	}
+
+	UDMMaterialStageExpression* ExpressionCDO = Cast<UDMMaterialStageExpression>(InExpressionClass->GetDefaultObject(true));
+	if (!ensure(ExpressionCDO)
+		|| !ensure(!ExpressionCDO->IsInputRequired() || ExpressionCDO->AllowsNestedInputs()))
+	{
+		return;
+	}
+
+	const TArray<FDMMaterialStageConnector>& OutputConnectors = ExpressionCDO->GetOutputConnectors();
+	if (!ensure(OutputConnectors.IsValidIndex(InOutputIndex))
+		|| !ensure(CanAcceptSubChannels(InThroughput, InInputIndex, OutputConnectors[InOutputIndex])))
+	{
+		return;
+	}
+
+	const int32 FloatSize = UDMValueDefinitionLibrary::GetValueDefinition(OutputConnectors[InOutputIndex].Type).GetFloatCount();
+
+	if (OutputConnectors[InOutputIndex].Type != EDMValueType::VT_Float_Any)
+	{
+		if (!ensure(FloatSize > 1))
+		{
+			return;
+		}
+	}
+
+	InChildMenuBuilder.BeginSection("ChangeInputOutputChannels", LOCTEXT("ChangeInputOutputChannels", "Output Channels"));
+	{
+		for (int32 ChannelIndex = 0; ChannelIndex < FloatSize; ++ChannelIndex)
+		{
+			const FText ChannelName = UDMValueDefinitionLibrary::GetValueDefinition(OutputConnectors[InOutputIndex].Type).GetChannelName(ChannelIndex + 1);
+			const int32 OutputChannel = UE::DynamicMaterialEditor::ChannelIndexToChannelBit(ChannelIndex + 1);
+
+			UE::DynamicMaterialEditor::Private::GenerateChangeInputMenu_Expression_Output_Channel_Impl(InChildMenuBuilder, ChannelName, Stage, InThroughput,
+				InInputIndex, InInputChannel, InExpressionClass, InOutputIndex, OutputChannel);
+		}
+	}
+	InChildMenuBuilder.EndSection();
+}
+
+namespace UE::DynamicMaterialEditor::Private
+{
+	void GenerateChangeInputMenu_UV_Channels_Impl(FMenuBuilder& InChildMenuBuilder, const FText& InMenuName,
+		UDMMaterialStage* InStage, UDMMaterialStageThroughput* InThroughput, int32 InInputIndex, int32 InInputChannel)
+	{
+		InChildMenuBuilder.AddSubMenu(
+			InMenuName,
+			LOCTEXT("ChangeInputExpressionUVTooltip", "Change the source of this input to a Texture UV coordinate."),
+			FNewMenuDelegate::CreateStatic(
+				&FDMMaterialStageInputMenus::GenerateChangeInputMenu_UV_Channels,
+				InThroughput,
+				InInputIndex,
+				InInputChannel
+			)
+		);
+	}
+
+	void GenerateChangeInputMenu_UV_Channel_Impl(FMenuBuilder& InChildMenuBuilder, const FText& InMenuName,
+		UDMMaterialStage* InStage, UDMMaterialStageThroughput* InThroughput, int32 InInputIndex, int32 InInputChannel, int32 InOutputChannel)
+	{
+		InChildMenuBuilder.AddMenuEntry(
+			InMenuName,
+			LOCTEXT("ChangeInputExpressionUVTooltip", "Change the source of this input to a Texture UV coordinate."),
+			FSlateIcon(),
+			FUIAction(FExecuteAction::CreateWeakLambda(
+				InStage,
+				[InStage, InInputIndex, InInputChannel, InOutputChannel]()
+				{
+					FScopedTransaction Transaction(LOCTEXT("SetStageInput", "Set Material Stage Input"));
+					InStage->Modify();
+					InStage->ChangeInput_UV(InInputIndex, InInputChannel, InOutputChannel);
+				})
+			)
+		);
+	}
+
+	void GenerateChangeInputMenu_UV_UV_Channels_Impl(FMenuBuilder& InChildMenuBuilder, const FText& InMenuName,
+		UDMMaterialStage* InStage, UDMMaterialStageThroughput* InThroughput, int32 InInputIndex, int32 InInputChannel)
+	{
+		InChildMenuBuilder.AddSubMenu(
+			InMenuName,
+			LOCTEXT("ChangeInputExpressionUVTooltip", "Change the source of this input to a Texture UV coordinate."),
+			FNewMenuDelegate::CreateStatic(
+				&FDMMaterialStageInputMenus::GenerateChangeInputMenu_UV_Channels,
+				InThroughput,
+				InInputIndex,
+				InInputChannel
+			),
+			FUIAction(FExecuteAction::CreateWeakLambda(
+				InStage,
+				[InStage, InInputIndex, InInputChannel]()
+				{
+					FScopedTransaction Transaction(LOCTEXT("SetStageInput", "Set Material Stage Input"));
+					InStage->Modify();
+					InStage->ChangeInput_UV(InInputIndex, InInputChannel, FDMMaterialStageConnectorChannel::WHOLE_CHANNEL);
+				})
+			),
+			NAME_None,
+			EUserInterfaceActionType::Button
+		);
+	}
+
+	void GenerateChangeInputMenu_UV_Impl(FMenuBuilder& InChildMenuBuilder, const FText& InMenuName,
+		UDMMaterialStage* InStage, UDMMaterialStageThroughput* InThroughput, int32 InInputIndex, int32 InInputChannel,
+		bool bInAcceptsWholeChannel, bool bInAcceptsSubChannels)
+	{
+		if (!bInAcceptsWholeChannel && !bInAcceptsSubChannels)
+		{
+			return;
+		}
+
+		if (bInAcceptsSubChannels)
+		{
+			if (bInAcceptsWholeChannel)
+			{
+				UE::DynamicMaterialEditor::Private::GenerateChangeInputMenu_UV_UV_Channels_Impl(InChildMenuBuilder, InMenuName, InStage, InThroughput,
+					InInputIndex, InInputChannel);
+			}
+			else
+			{
+				UE::DynamicMaterialEditor::Private::GenerateChangeInputMenu_UV_Channels_Impl(InChildMenuBuilder, InMenuName, InStage, InThroughput,
+					InInputIndex, InInputChannel);
+			}
+		}
+		else
+		{
+			UE::DynamicMaterialEditor::Private::GenerateChangeInputMenu_UV_Channel_Impl(InChildMenuBuilder, InMenuName, InStage, InThroughput,
+				InInputIndex, InInputChannel, FDMMaterialStageConnectorChannel::WHOLE_CHANNEL);
+		}
+	}
+}
+
+void FDMMaterialStageInputMenus::GenerateChangeInputMenu_UV(FMenuBuilder& InChildMenuBuilder, UDMMaterialStageThroughput* InThroughput, 
+	const int32 InInputIndex, const int32 InInputChannel)
+{
+	if (!ensure(InThroughput))
+	{
+		return;
+	}
+
+	UDMMaterialStage* Stage = InThroughput->GetStage();
+	if (!ensure(Stage))
+	{
+		return;
+	}
+
+	const TArray<FDMMaterialStageConnector>& InputConnectors = InThroughput->GetInputConnectors();
+	if (!ensure(InputConnectors.IsValidIndex(InInputIndex)))
+	{
+		return;
+	}
+
+	const bool bCompatibleWithTextureUV = InThroughput->CanInputAcceptType(InInputIndex, EDMValueType::VT_Float2);
+	const bool bCompatibleWithTextureUorV = InThroughput->CanInputAcceptType(InInputIndex, EDMValueType::VT_Float1);
+
+	UE::DynamicMaterialEditor::Private::GenerateChangeInputMenu_UV_Impl(InChildMenuBuilder, LOCTEXT("ChangeInputExpressionUV", "Texture UV"),
+		Stage, InThroughput, InInputIndex, InInputChannel, bCompatibleWithTextureUV, bCompatibleWithTextureUorV);
+}
+
+void FDMMaterialStageInputMenus::GenerateChangeInputMenu_UV_Channels(FMenuBuilder& InChildMenuBuilder, UDMMaterialStageThroughput* InThroughput, 
+	const int32 InInputIndex, const int32 InInputChannel)
+{
+	if (!ensure(InThroughput))
+	{
+		return;
+	}
+
+	UDMMaterialStage* Stage = InThroughput->GetStage();
+	if (!ensure(Stage))
+	{
+		return;
+	}
+
+	const TArray<FDMMaterialStageConnector>& InputConnectors = InThroughput->GetInputConnectors();
+	if (!ensure(InputConnectors.IsValidIndex(InInputIndex))
+		|| !ensure(InThroughput->CanInputAcceptType(InInputIndex, EDMValueType::VT_Float1)))
+	{
+		return;
+	}
+
+	InChildMenuBuilder.BeginSection("ChangeInputOutputChannels", LOCTEXT("ChangeInputOutputChannel", "Output Channels"));
+	{
+		UE::DynamicMaterialEditor::Private::GenerateChangeInputMenu_UV_Channel_Impl(InChildMenuBuilder,
+			LOCTEXT("ChangeInputExpressionUVU", "U"), Stage, InThroughput, InInputIndex, InInputChannel, 0);
+
+		UE::DynamicMaterialEditor::Private::GenerateChangeInputMenu_UV_Channel_Impl(InChildMenuBuilder,
+			LOCTEXT("ChangeInputExpressionUVV", "V"), Stage, InThroughput, InInputIndex, InInputChannel, 1);
+	}
+	InChildMenuBuilder.EndSection();
+}
+
+void FDMMaterialStageInputMenus::GenerateChangeInputMenu_Gradients(FMenuBuilder& InChildMenuBuilder, UDMMaterialStageThroughput* InThroughput, 
+	const int32 InInputIndex, const int32 InInputChannel)
+{
+	if (!ensure(InThroughput))
+	{
+		return;
+	}
+
+	const TArray<FDMMaterialStageConnector>& InputConnectors = InThroughput->GetInputConnectors();
+	if (!ensure(InputConnectors.IsValidIndex(InInputIndex))
+		|| !ensure(InputConnectors[InInputIndex].IsCompatibleWith(EDMValueType::VT_Float1)))
+	{
+		return;
+	}
+
+	const TArray<TStrongObjectPtr<UClass>>& Gradients = UDMMaterialStageGradient::GetAvailableGradients();
+
+	if (Gradients.IsEmpty())
+	{
+		return;
+	}
+
+	for (const TStrongObjectPtr<UClass>& GradientClass : Gradients)
+	{
+		TSubclassOf<UDMMaterialStageGradient> InputClass = GradientClass.Get();
+
+		if (!InputClass)
+		{
+			continue;
+		}
+
+		UDMMaterialStageGradient* GradientCDO = Cast<UDMMaterialStageGradient>(InputClass->GetDefaultObject(true));
+		if (!ensure(GradientCDO))
+		{
+			continue;;
+		}
+
+		GenerateChangeInputMenu_Gradient(InChildMenuBuilder, GradientCDO, InThroughput, InInputIndex, InInputChannel);
+	}
+}
+
+void FDMMaterialStageInputMenus::GenerateChangeInputMenu_Gradient(FMenuBuilder& InChildSubMenuBuilder, UDMMaterialStageGradient* InGradientCDO, UDMMaterialStageThroughput* InThroughput, 
+	const int32 InInputIndex, const int32 InInputChannel)
+{
+	UDMMaterialStage* Stage = InThroughput->GetStage();
+	if (!ensure(Stage))
+	{
+		return;
+	}
+
+	if (!InThroughput->CanInputAcceptType(InInputIndex, EDMValueType::VT_Float1))
+	{
+		return;
+	}
+
+	const FText GradientName = InGradientCDO->GetDescription();
+
+	InChildSubMenuBuilder.AddMenuEntry(
+		GradientName,
+		LOCTEXT("ChangeInputGradientInputTooltip2", "Change the source of this input to a Material Gradient."),
+		FSlateIcon(),
+		FUIAction(FExecuteAction::CreateWeakLambda(
+			Stage,
+			[Stage, InInputIndex, InInputChannel, GradientClass = TSubclassOf<UDMMaterialStageGradient>(InGradientCDO->GetClass())]()
+			{
+				FScopedTransaction Transaction(LOCTEXT("SetStageInput", "Set Material Stage Input"));
+				Stage->Modify();
+				Stage->ChangeInput_Gradient(InInputIndex, InInputChannel, GradientClass, FDMMaterialStageConnectorChannel::WHOLE_CHANNEL);
+			})
+		)
+	);
+}
+
+#undef LOCTEXT_NAMESPACE
diff --git a/Engine/Plugins/Experimental/DynamicMaterial/Source/DynamicMaterialEditor/Private/Menus/DMMaterialStageInputMenus.h b/Engine/Plugins/Experimental/DynamicMaterial/Source/DynamicMaterialEditor/Private/Menus/DMMaterialStageInputMenus.h
new file mode 100644
--- /dev/null
+++ b/Engine/Plugins/Experimental/DynamicMaterial/Source/DynamicMaterialEditor/Private/Menus/DMMaterialStageInputMenus.h
@@ -0,0 +1,85 @@
+// Copyright Epic Games, Inc. All Rights Reserved.
+
+#pragma once
+
+#include "Templates/SharedPointer.h"
+#include "Templates/SubclassOf.h"
+
+class FMenuBuilder;
+class SDMSlot;
+class SDMStage;
+class SWidget;
+class UDMMaterialSlot;
+class UDMMaterialStageExpression;
+class UDMMaterialStageGradient;
+class UDMMaterialStageThroughput;
+class UToolMenu;
+enum class EDMMaterialPropertyType : uint8;
+enum class EDMValueType : uint8;
+struct FDMMaterialStageConnector;
+
+class FDMMaterialStageInputMenus
+{
+public:
+	static FName GetStageChangeInputMenuName();
+
+	static bool CanAcceptSubChannels(UDMMaterialStageThroughput* const InThroughput, const int32 InInputIndex, const EDMValueType InOutputType);
+	static bool CanAcceptSubChannels(UDMMaterialStageThroughput* const InThroughput, const int32 InInputIndex, const FDMMaterialStageConnector& InOutputConnector);
+
+	/** Generate right click menu for changing inputs */
+	static TSharedRef<SWidget> MakeStageChangeInputMenu(UDMMaterialStageThroughput* InThroughput, const int32 InInputIndex, const int32 InInputChannel);
+
+	/** Previous Stage */
+	static void GenerateChangeInputMenu_PreviousStages(FMenuBuilder& InChildMenuBuilder, UDMMaterialStageThroughput* InThroughput, const int32 InInputIndex,
+		const int32 InInputChannel);
+	static void GenerateChangeInputMenu_PreviousStage_Outputs(FMenuBuilder& InChildMenuBuilder, UDMMaterialStageThroughput* InThroughput,
+		const int32 InInputIndex, const int32 InInputChannel, EDMMaterialPropertyType InMaterialProperty);
+	static void GenerateChangeInputMenu_PreviousStage_Output_Channels(FMenuBuilder& InChildMenuBuilder, UDMMaterialStageThroughput* InThroughput,
+		const int32 InInputIndex, const int32 InInputChannel, EDMMaterialPropertyType InMaterialProperty, const int32 InOutputChannel);
+
+	/** New local value */
+	static void GenerateChangeInputMenu_NewLocalValues(FMenuBuilder& InChildMenuBuilder, UDMMaterialStageThroughput* InThroughput, const int32 InInputIndex,
+		const int32 InInputChannel);
+
+	/** Global values */
+	static void GenerateChangeInputMenu_GlobalValues(FMenuBuilder& InChildMenuBuilder, UDMMaterialStageThroughput* InThroughput, const int32 InInputIndex,
+		const int32 InInputChannel);
+	static void GenerateChangeInputMenu_GlobalValue_Channels(FMenuBuilder& InChildMenuBuilder, UDMMaterialStageThroughput* InThroughput, const int32 InInputIndex,
+		const int32 InInputChannel, const int32 InValueIndex);
+
+	/** New global value */
+	static void GenerateChangeInputMenu_NewGlobalValues(FMenuBuilder& InChildMenuBuilder, UDMMaterialStageThroughput* InThroughput, const int32 InInputIndex,
+		const int32 InInputChannel);
+
+	/** Other slots */
+	static void GenerateChangeInputMenu_Slots(FMenuBuilder& InChildMenuBuilder, UDMMaterialStageThroughput* InThroughput, const int32 InInputIndex,
+		const int32 InInputChannel);
+	static void GenerateChangeInputMenu_Slot_Properties(FMenuBuilder& InChildMenuBuilder, UDMMaterialStageThroughput* InThroughput,
+		const int32 InInputIndex, const int32 InInputChannel, UDMMaterialSlot* InSlot);
+	static void GenerateChangeInputMenu_Slot_Property_Outputs(FMenuBuilder& InChildMenuBuilder, UDMMaterialStageThroughput* InThroughput,
+		const int32 InInputIndex, const int32 InInputChannel, UDMMaterialSlot* InSlot, EDMMaterialPropertyType InMaterialProperty);
+	static void GenerateChangeInputMenu_Slot_Property_Output_Channels(FMenuBuilder& InChildMenuBuilder, UDMMaterialStageThroughput* InThroughput,
+		const int32 InInputIndex, const int32 InInputChannel, UDMMaterialSlot* InSlot, EDMMaterialPropertyType InMaterialProperty, const int32 InOutputChannel);
+
+	/** Expressions */
+	static void GenerateChangeInputMenu_Expressions(FMenuBuilder& InChildMenuBuilder, UDMMaterialStageThroughput* InThroughput, const int32 InInputIndex,
+		const int32 InInputChannel);
+	static void GenerateChangeInputMenu_Expression(FMenuBuilder& InChildSubMenuBuilder, UDMMaterialStageExpression* InExpressionCDO, 
+		UDMMaterialStageThroughput* InThroughput, const int32 InInputIndex, const int32 InInputChannel);
+	static void GenerateChangeInputMenu_Expression_Outputs(FMenuBuilder& InChildMenuBuilder, UDMMaterialStageThroughput* InThroughput,
+		const int32 InInputIndex, const int32 InInputChannel, TSubclassOf<UDMMaterialStageExpression> InExpressionClass);
+	static void GenerateChangeInputMenu_Expression_Output_Channels(FMenuBuilder& InChildMenuBuilder, UDMMaterialStageThroughput* InThroughput,
+		const int32 InInputIndex, const int32 InInputChannel, TSubclassOf<UDMMaterialStageExpression> InExpressionClass, const int32 InOutputChannel);
+
+	/** UV */
+	static void GenerateChangeInputMenu_UV(FMenuBuilder& InChildMenuBuilder, UDMMaterialStageThroughput* InThroughput, const int32 InInputIndex,
+		const int32 InInputChannel);
+	static void GenerateChangeInputMenu_UV_Channels(FMenuBuilder& InChildMenuBuilder, UDMMaterialStageThroughput* InThroughput, const int32 InInputIndex,
+		const int32 InInputChannel);
+
+	/** Gradients */
+	static void GenerateChangeInputMenu_Gradients(FMenuBuilder& InChildMenuBuilder, UDMMaterialStageThroughput* InThroughput, const int32 InInputIndex,
+		const int32 InInputChannel);
+	static void GenerateChangeInputMenu_Gradient(FMenuBuilder& InChildSubMenuBuilder, UDMMaterialStageGradient* InGradientCDO, UDMMaterialStageThroughput* InThroughput,
+		const int32 InInputIndex, const int32 InInputChannel);
+};
diff --git a/Engine/Plugins/Experimental/DynamicMaterial/Source/DynamicMaterialEditor/Private/Menus/DMMaterialStageMenus.cpp b/Engine/Plugins/Experimental/DynamicMaterial/Source/DynamicMaterialEditor/Private/Menus/DMMaterialStageMenus.cpp
new file mode 100644
--- /dev/null
+++ b/Engine/Plugins/Experimental/DynamicMaterial/Source/DynamicMaterialEditor/Private/Menus/DMMaterialStageMenus.cpp
@@ -0,0 +1,162 @@
+// Copyright Epic Games, Inc. All Rights Reserved.
+
+#include "Menus/DMMaterialStageMenus.h"
+#include "Components/DMMaterialLayer.h"
+#include "Components/DMMaterialStage.h"
+#include "Components/DMMaterialValue.h"
+#include "Menus/DMMenuContext.h"
+#include "ScopedTransaction.h"
+#include "Slate/SDMSlot.h"
+#include "Slate/SDMStage.h"
+#include "ToolMenu.h"
+
+#define LOCTEXT_NAMESPACE "FDMMaterialStageMenus"
+
+FName FDMMaterialStageMenus::GetStageMenuName()
+{
+	return "MaterialDesigner.MaterialStage";
+}
+
+FName FDMMaterialStageMenus::GetStageToggleSectionName()
+{
+	return "StageToggle";
+}
+
+UToolMenu* FDMMaterialStageMenus::GenerateStageMenu(const TSharedPtr<SDMSlot>& InSlotWidget, const TSharedPtr<SDMStage>& InStageWidget)
+{
+	UToolMenu* NewToolMenu = UDMMenuContext::GenerateContextMenuStage(GetStageMenuName(), InSlotWidget, InStageWidget);
+	if (!NewToolMenu)
+	{
+		return nullptr;
+	}
+
+	AddStageSection(NewToolMenu);
+
+	return NewToolMenu;
+}
+
+void FDMMaterialStageMenus::AddStageSection(UToolMenu* InMenu)
+{
+	if (!IsValid(InMenu) || InMenu->ContainsSection(GetStageToggleSectionName()))
+	{
+		return;
+	}
+
+	const UDMMenuContext* const MenuContext = InMenu->FindContext<UDMMenuContext>();
+	if (!MenuContext)
+	{
+		return;
+	}
+
+	UDMMaterialStage* Stage = MenuContext->GetStage();
+	if (!Stage)
+	{
+		return;
+	}
+
+	const UDMMaterialLayerObject* Layer = MenuContext->GetLayer();
+	if (!Layer)
+	{
+		return;
+	}
+
+	const int32 LayerIndex = Layer->FindIndex();
+
+	const bool bAllowRemoveLayer = LayerIndex >= 1;
+	const EDMMaterialLayerStage StageType = Layer->GetStageType(Stage);
+
+	if (bAllowRemoveLayer || StageType == EDMMaterialLayerStage::Mask)
+	{
+		FToolMenuSection& NewSection = InMenu->AddSection(GetStageToggleSectionName(), LOCTEXT("MaterialStageMenu", "Material Stage"));
+
+		if (bAllowRemoveLayer)
+		{
+			NewSection.AddMenuEntry(NAME_None,
+				LOCTEXT("ToggleLayer", "Toggle"),
+				LOCTEXT("ToggleLayerTooltip", "Toggle the entire layer on and off.\n\n"
+					"Warning: Toggling a layer off may result in inputs being reset where incompatibilities are found.\n\nAlt+Left Click"),
+				FSlateIcon(),
+				FUIAction(FExecuteAction::CreateWeakLambda(
+					Layer,
+					[Layer]()
+					{
+						FScopedTransaction Transaction(LOCTEXT("ToggleAllStageEnabled", "Material Designer Toggle All Stage Enabled"));
+
+						for (UDMMaterialStage* Stage : Layer->GetStages(EDMMaterialLayerStage::All))
+						{
+							Stage->Modify();
+							Stage->SetEnabled(!Stage->IsEnabled());
+						}
+					}
+				))
+			);
+
+			if (StageType == EDMMaterialLayerStage::Base)
+			{
+				NewSection.AddMenuEntry(NAME_None,
+					LOCTEXT("ToggleLayerBase", "Toggle Base"),
+					LOCTEXT("ToggleLayerBaseTooltip", "Toggle the layer base on and off.\n\n"
+						"Warning: Toggling a layer base off may result in inputs being reset where incompatibilities are found.\n\nAlt+Shift+Left Click"),
+					FSlateIcon(),
+					FUIAction(FExecuteAction::CreateWeakLambda(
+						Layer,
+						[Layer]()
+						{
+							FScopedTransaction Transaction(LOCTEXT("ToggleBaseStageEnabled", "Material Designer Toggle Base Stage Enabled"));
+
+							if (UDMMaterialStage* Stage = Layer->GetStage(EDMMaterialLayerStage::Base))
+							{
+								Stage->Modify();
+								Stage->SetEnabled(!Stage->IsEnabled());
+							}
+						}
+					))
+				);
+			}
+		}
+
+		if (StageType == EDMMaterialLayerStage::Mask)
+		{
+			NewSection.AddMenuEntry(NAME_None,
+				LOCTEXT("ToggleLayerMask", "Toggle Mask"),
+				LOCTEXT("ToggleLayerMaskTooltip", "Toggle the layer mask on and off.\n\nAlt+Shift+Left Click"),
+				FSlateIcon(),
+				FUIAction(FExecuteAction::CreateWeakLambda(
+					Layer,
+					[Layer]()
+					{
+						FScopedTransaction Transaction(LOCTEXT("ToggleBaseStageEnabled", "Material Designer Toggle Base Stage Enabled"));
+
+						if (UDMMaterialStage* Stage = Layer->GetStage(EDMMaterialLayerStage::Mask))
+						{
+							Stage->Modify();
+							Stage->SetEnabled(!Stage->IsEnabled());
+						}
+					}
+				))
+			);
+		}
+
+		if (bAllowRemoveLayer)
+		{
+			if (TSharedPtr<SDMSlot> SlotWidget = MenuContext->GetSlotWidget().Pin())
+			{
+				//FDMSlotLayerMenus::AddLayerSection
+				NewSection.AddMenuEntry(NAME_None,
+					LOCTEXT("RemoveLayer", "Remove"),
+					LOCTEXT("RemoveLayerTooltip", "Remove this layer from its Material Slot.\n\n"
+						"Warning: Removing a stage off may result in inputs being reset where incompatibilities are found."),
+					FSlateIcon(),
+					FUIAction(FExecuteAction::CreateSP(
+						SlotWidget.Get(),
+						&SDMSlot::RemoveLayerByStage,
+						Stage, 
+						false
+					))
+				);
+			}
+		}
+	}
+}
+
+#undef LOCTEXT_NAMESPACE
diff --git a/Engine/Plugins/Experimental/DynamicMaterial/Source/DynamicMaterialEditor/Private/Menus/DMMaterialStageMenus.h b/Engine/Plugins/Experimental/DynamicMaterial/Source/DynamicMaterialEditor/Private/Menus/DMMaterialStageMenus.h
new file mode 100644
--- /dev/null
+++ b/Engine/Plugins/Experimental/DynamicMaterial/Source/DynamicMaterialEditor/Private/Menus/DMMaterialStageMenus.h
@@ -0,0 +1,21 @@
+// Copyright Epic Games, Inc. All Rights Reserved.
+
+#pragma once
+
+#include "Templates/SharedPointer.h"
+
+class SDMSlot;
+class SDMStage;
+class UToolMenu;
+
+class FDMMaterialStageMenus
+{
+public:
+	static FName GetStageMenuName();
+
+	static FName GetStageToggleSectionName();
+
+	static UToolMenu* GenerateStageMenu(const TSharedPtr<SDMSlot>& InSlotWidget, const TSharedPtr<SDMStage>& InStageWidget);
+
+	static void AddStageSection(UToolMenu* InMenu);
+};
