[MaterialDesigner]

diff --git a/Engine/Plugins/Experimental/DynamicMaterial/Source/DynamicMaterialEditor/Private/Menus/DMMaterialStageSourceMenus.cpp b/Engine/Plugins/Experimental/DynamicMaterial/Source/DynamicMaterialEditor/Private/Menus/DMMaterialStageSourceMenus.cpp
new file mode 100644
--- /dev/null
+++ b/Engine/Plugins/Experimental/DynamicMaterial/Source/DynamicMaterialEditor/Private/Menus/DMMaterialStageSourceMenus.cpp
@@ -0,0 +1,1227 @@
+// Copyright Epic Games, Inc. All Rights Reserved.
+
+#include "DMMaterialStageSourceMenus.h"
+#include "Components/DMMaterialLayer.h"
+#include "Components/DMMaterialProperty.h"
+#include "Components/DMMaterialSlot.h"
+#include "Components/DMMaterialStage.h"
+#include "Components/DMMaterialStageBlend.h"
+#include "Components/DMMaterialStageExpression.h"
+#include "Components/DMMaterialStageGradient.h"
+#include "Components/DMMaterialStageSource.h"
+#include "Components/DMMaterialStageThroughput.h"
+#include "Components/DMMaterialStageThroughputLayerBlend.h"
+#include "Components/DMMaterialValue.h"
+#include "Components/MaterialStageExpressions/DMMSESceneTexture.h"
+#include "Components/MaterialStageExpressions/DMMSETextureSample.h"
+#include "Components/MaterialStageExpressions/DMMSETextureSampleEdgeColor.h"
+#include "DMDefs.h"
+#include "DMValueDefinition.h"
+#include "DynamicMaterialEditorModule.h"
+#include "Menus/DMMenuContext.h"
+#include "Model/DynamicMaterialModel.h"
+#include "Model/DynamicMaterialModelEditorOnlyData.h"
+#include "ScopedTransaction.h"
+#include "Slate/SDMSlot.h"
+#include "Slate/SDMStage.h"
+#include "Templates/SharedPointer.h"
+#include "ToolMenus.h"
+#include "Widgets/SNullWidget.h"
+
+#define LOCTEXT_NAMESPACE "FDMMaterialStageSourceMenus"
+
+namespace UE::DynamicMaterialEditor::Private
+{
+	static const FName ChangeStageSourceMenuName = TEXT("MaterialDesigner.MaterialStage.ChangeSource");
+
+	void GenerateChangeSourceMenu_NewLocalValues(UToolMenu* InMenu)
+	{
+		if (!ensure(IsValid(InMenu)))
+		{
+			return;
+		}
+
+		UDMMenuContext* MenuContext = InMenu->FindContext<UDMMenuContext>();
+
+		if (!ensure(MenuContext))
+		{
+			return;
+		}
+
+		UDMMaterialStage* Stage = MenuContext->GetStage();
+
+		if (!Stage)
+		{
+			return;
+		}
+
+		for (EDMValueType ValueType : UDMValueDefinitionLibrary::GetValueTypes())
+		{
+			static const FText NameTooltipFormat = LOCTEXT("ChangeSourceNewValueSourceTooltipTemplate", "Add a new {0} Value and use it as the source of this stage.");
+
+			const FText Name = UDMValueDefinitionLibrary::GetValueDefinition(ValueType).GetDisplayName();
+			const FText FormattedTooltip = FText::Format(NameTooltipFormat, Name);
+
+			InMenu->AddMenuEntry(NAME_None, FToolMenuEntry::InitMenuEntry(NAME_None,
+				Name,
+				FormattedTooltip,
+				FSlateIcon(),
+				FUIAction(
+					FExecuteAction::CreateWeakLambda(
+						MenuContext,
+						[ValueType, MenuContext]()
+						{
+							UDMMaterialStage* const Stage = MenuContext->GetStage();
+							if (!Stage)
+							{
+								return;
+							}
+
+							UDMMaterialStageSource* const StageSource = Stage->GetSource();
+							if (!StageSource)
+							{
+								return;
+							}
+
+							if (StageSource->IsA<UDMMaterialStageBlend>())
+							{
+								const int32 OutputChannel = ValueType == EDMValueType::VT_ColorAtlas
+									? FDMMaterialStageConnectorChannel::THREE_CHANNELS
+									: FDMMaterialStageConnectorChannel::WHOLE_CHANNEL;
+
+								FScopedTransaction Transaction(LOCTEXT("SetStageInputBase", "Set Material Designer Base Source"));
+								Stage->Modify();
+								Stage->ChangeInput_NewLocalValue(UDMMaterialStageBlend::InputB, FDMMaterialStageConnectorChannel::WHOLE_CHANNEL, ValueType, OutputChannel);
+							}
+							else if (StageSource->IsA<UDMMaterialStageThroughputLayerBlend>())
+							{
+								const int32 OutputChannel = ValueType == EDMValueType::VT_ColorAtlas
+									? FDMMaterialStageConnectorChannel::FOURTH_CHANNEL
+									: FDMMaterialStageConnectorChannel::WHOLE_CHANNEL;
+
+								FScopedTransaction Transaction(LOCTEXT("SetStageInputMask", "Set Material Designer Mask Source"));
+								Stage->Modify();
+								Stage->ChangeInput_NewLocalValue(UDMMaterialStageThroughputLayerBlend::InputMaskSource, FDMMaterialStageConnectorChannel::WHOLE_CHANNEL, ValueType, OutputChannel);
+							}
+							else
+							{
+								FScopedTransaction Transaction(LOCTEXT("SetStageSource", "Set Material Designer Stage Source"));
+								Stage->Modify();
+								Stage->ChangeSource_NewLocalValue(ValueType);
+							}
+
+							if (TSharedPtr<SDMSlot> SlotWidget = MenuContext->GetSlotWidget().Pin())
+							{
+								SlotWidget->InvalidateComponentEditWidget();
+							}
+						})
+				)
+			));
+		}
+	}
+
+	void GenerateChangeSourceMenu_GlobalValues(UToolMenu* InMenu)
+	{
+		if (!ensure(IsValid(InMenu)))
+		{
+			return;
+		}
+
+		UDMMenuContext* MenuContext = InMenu->FindContext<UDMMenuContext>();
+		if (!ensure(MenuContext))
+		{
+			return;
+		}
+
+		UDynamicMaterialModel* MaterialModel = MenuContext->GetModel();
+		if (!ensure(MaterialModel))
+		{
+			return;
+		}
+
+		const TArray<UDMMaterialValue*>& Values = MaterialModel->GetValues();
+
+		if (Values.IsEmpty())
+		{
+			return;
+		}
+
+		for (UDMMaterialValue* Value : Values)
+		{
+			if (!IsValid(Value))
+			{
+				continue;
+			}
+
+			InMenu->AddMenuEntry(NAME_None, FToolMenuEntry::InitMenuEntry(NAME_None,
+				Value->GetDescription(),
+				LOCTEXT("ChangeSourceValueSourceTooltip2", "Change the source of this stage to this Material Value."),
+				FSlateIcon(),
+				FUIAction(FExecuteAction::CreateWeakLambda(
+					MenuContext,
+					[MenuContext, Value]()
+					{
+						UDMMaterialStage* const Stage = MenuContext->GetStage();
+						if (!Stage)
+						{
+							return;
+						}
+
+						UDMMaterialStageSource* const StageSource = Stage->GetSource();
+						if (!StageSource)
+						{
+							return;
+						}
+
+						if (StageSource->IsA<UDMMaterialStageBlend>())
+						{
+							const int32 OutputChannel = Value->GetType() == EDMValueType::VT_ColorAtlas
+								? FDMMaterialStageConnectorChannel::THREE_CHANNELS
+								: FDMMaterialStageConnectorChannel::WHOLE_CHANNEL;
+
+							FScopedTransaction Transaction(LOCTEXT("SetStageInputBase", "Set Material Designer Base Source"));
+							Stage->Modify();
+							Stage->ChangeInput_Value(UDMMaterialStageBlend::InputB, FDMMaterialStageConnectorChannel::WHOLE_CHANNEL,
+								Value, OutputChannel);
+						}
+						else if (StageSource->IsA<UDMMaterialStageThroughputLayerBlend>())
+						{
+							const int32 OutputChannel = Value->GetType() == EDMValueType::VT_ColorAtlas
+								? FDMMaterialStageConnectorChannel::FOURTH_CHANNEL
+								: FDMMaterialStageConnectorChannel::WHOLE_CHANNEL;
+
+							FScopedTransaction Transaction(LOCTEXT("SetStageInputMask", "Set Material Designer Mask Source"));
+							Stage->Modify();
+							Stage->ChangeInput_Value(UDMMaterialStageThroughputLayerBlend::InputMaskSource, FDMMaterialStageConnectorChannel::WHOLE_CHANNEL,
+								Value, OutputChannel);
+						}
+						else
+						{
+							FScopedTransaction Transaction(LOCTEXT("SetStageSource", "Set Material Designer Stage Source"));
+							Stage->Modify();
+							Stage->ChangeSource_Value(Value);
+						}
+
+						if (TSharedPtr<SDMSlot> SlotWidget = MenuContext->GetSlotWidget().Pin())
+						{
+							SlotWidget->InvalidateComponentEditWidget();
+						}
+					})
+				)
+			));
+		}
+	}
+
+	void GenerateChangeSourceMenu_NewGlobalValues(UToolMenu* InMenu)
+	{
+		if (!ensure(IsValid(InMenu)))
+		{
+			return;
+		}
+
+		UDMMenuContext* MenuContext = InMenu->FindContext<UDMMenuContext>();
+		if (!ensure(MenuContext))
+		{
+			return;
+		}
+
+		UDMMaterialStage* Stage = MenuContext->GetStage();
+		if (!Stage)
+		{
+			return;
+		}
+
+		for (EDMValueType ValueType : UDMValueDefinitionLibrary::GetValueTypes())
+		{
+			static const FText NameTooltipFormat = LOCTEXT("ChangeSourceNewValueSourceTooltipTemplate", "Add a new {0} Value and use it as the source of this stage.");
+
+			const FText Name = UDMValueDefinitionLibrary::GetValueDefinition(ValueType).GetDisplayName();
+			const FText FormattedTooltip = FText::Format(NameTooltipFormat, Name);
+
+			InMenu->AddMenuEntry(NAME_None, FToolMenuEntry::InitMenuEntry(NAME_None,
+				Name,
+				FormattedTooltip,
+				FSlateIcon(),
+				FUIAction(
+					FExecuteAction::CreateWeakLambda(
+						MenuContext,
+						[ValueType, MenuContext]()
+						{
+							UDMMaterialStage* const Stage = MenuContext->GetStage();
+							if (!Stage)
+							{
+								return;
+							}
+
+							UDMMaterialStageSource* const StageSource = Stage->GetSource();
+							if (!StageSource)
+							{
+								return;
+							}
+
+							if (StageSource->IsA<UDMMaterialStageBlend>())
+							{
+								const int32 OutputChannel = ValueType == EDMValueType::VT_ColorAtlas
+									? FDMMaterialStageConnectorChannel::THREE_CHANNELS
+									: FDMMaterialStageConnectorChannel::WHOLE_CHANNEL;
+
+								FScopedTransaction Transaction(LOCTEXT("SetStageInputBase", "Set Material Designer Base Source"));
+								Stage->Modify();
+								Stage->ChangeInput_NewValue(UDMMaterialStageBlend::InputB, FDMMaterialStageConnectorChannel::WHOLE_CHANNEL, ValueType, OutputChannel);
+							}
+							else if (StageSource->IsA<UDMMaterialStageThroughputLayerBlend>())
+							{
+								const int32 OutputChannel = ValueType == EDMValueType::VT_ColorAtlas
+									? FDMMaterialStageConnectorChannel::FOURTH_CHANNEL
+									: FDMMaterialStageConnectorChannel::WHOLE_CHANNEL;
+
+								FScopedTransaction Transaction(LOCTEXT("SetStageInputMask", "Set Material Designer Mask Source"));
+								Stage->Modify();
+								Stage->ChangeInput_NewValue(UDMMaterialStageThroughputLayerBlend::InputMaskSource, FDMMaterialStageConnectorChannel::WHOLE_CHANNEL, ValueType, OutputChannel);
+							}
+							else
+							{
+								FScopedTransaction Transaction(LOCTEXT("SetStageSource", "Set Material Designer Stage Source"));
+								Stage->Modify();
+								Stage->ChangeSource_NewValue(ValueType);
+							}
+
+							if (TSharedPtr<SDMSlot> SlotWidget = MenuContext->GetSlotWidget().Pin())
+							{
+								SlotWidget->InvalidateComponentEditWidget();
+							}
+						})
+				)
+			));
+		}
+	}
+
+	void GenerateChangeSourceMenu_Slot_Properties(UToolMenu* InMenu, UDMMaterialSlot* InSlot)
+	{
+		if (!ensure(IsValid(InMenu)))
+		{
+			return;
+		}
+
+		UDMMenuContext* MenuContext = InMenu->FindContext<UDMMenuContext>();
+		if (!ensure(MenuContext))
+		{
+			return;
+		}
+
+		UDMMaterialStage* Stage = MenuContext->GetStage();
+		if (!Stage)
+		{
+			return;
+		}
+
+		UDMMaterialLayerObject* Layer = Stage->GetLayer();
+		if (!Layer)
+		{
+			return;
+		}
+
+		if (!IsValid(InSlot) || Layer->GetSlot() != InSlot)
+		{
+			return;
+		}
+
+		UDynamicMaterialModel* MaterialModel = MenuContext->GetModel();
+		if (!MaterialModel)
+		{
+			return;
+		}
+
+		UDynamicMaterialModelEditorOnlyData* ModelEditorOnlyData = UDynamicMaterialModelEditorOnlyData::Get(MaterialModel);
+		if (!ModelEditorOnlyData)
+		{
+			return;
+		}
+
+		for (EDMMaterialPropertyType Property : ModelEditorOnlyData->GetMaterialPropertiesForSlot(InSlot))
+		{
+			UDMMaterialProperty* PropertyObj = ModelEditorOnlyData->GetMaterialProperty(Property);
+
+			if (ensure(PropertyObj))
+			{
+				InMenu->AddMenuEntry(NAME_None, FToolMenuEntry::InitMenuEntry(NAME_None,
+					PropertyObj->GetDescription(),
+					LOCTEXT("ChangeSourceSlotSourceTooltip3", "Change the source of this stage to the output from this Material Slot's Property."),
+					FSlateIcon(),
+					FUIAction(FExecuteAction::CreateWeakLambda(
+						MenuContext,
+						[MenuContext, Property]()
+						{
+							UDMMaterialStage* const Stage = MenuContext->GetStage();
+							if (!Stage)
+							{
+								return;
+							}
+
+							UDMMaterialStageSource* const StageSource = Stage->GetSource();
+							if (!StageSource)
+							{
+								return;
+							}
+
+							UDMMaterialLayerObject* Layer = Stage->GetLayer();
+							if (!Layer)
+							{
+								return;
+							}
+
+							UDMMaterialSlot* Slot = Layer->GetSlot();
+							if (!Slot)
+							{
+								return;
+							}
+
+							if (StageSource->IsA<UDMMaterialStageBlend>())
+							{
+								FScopedTransaction Transaction(LOCTEXT("SetStageInputBase", "Set Material Designer Base Source"));
+								Stage->Modify();
+								Stage->ChangeInput_Slot(UDMMaterialStageBlend::InputB, FDMMaterialStageConnectorChannel::WHOLE_CHANNEL,
+									Slot, Property, 0, FDMMaterialStageConnectorChannel::WHOLE_CHANNEL);
+							}
+							else if (StageSource->IsA<UDMMaterialStageThroughputLayerBlend>())
+							{
+								FScopedTransaction Transaction(LOCTEXT("SetStageInputMask", "Set Material Designer Mask Source"));
+								Stage->Modify();
+								Stage->ChangeInput_Slot(UDMMaterialStageThroughputLayerBlend::InputMaskSource, FDMMaterialStageConnectorChannel::WHOLE_CHANNEL,
+									Slot, Property, 0, FDMMaterialStageConnectorChannel::WHOLE_CHANNEL);
+							}
+							else
+							{
+								FScopedTransaction Transaction(LOCTEXT("SetStageSource", "Set Material Designer Stage Source"));
+								Stage->Modify();
+								Stage->ChangeSource_Slot(Slot, Property);
+							}
+
+							if (TSharedPtr<SDMSlot> SlotWidget = MenuContext->GetSlotWidget().Pin())
+							{
+								SlotWidget->InvalidateComponentEditWidget();
+							}
+						})
+					)
+				));
+			}
+		}
+	}
+
+	void GenerateChangeSourceMenu_Slots(UToolMenu* InMenu)
+	{
+		if (!ensure(IsValid(InMenu)))
+		{
+			return;
+		}
+
+		UDMMenuContext* MenuContext = InMenu->FindContext<UDMMenuContext>();
+		if (!ensure(MenuContext))
+		{
+			return;
+		}
+
+		UDMMaterialStage* Stage = MenuContext->GetStage();
+		if (!Stage)
+		{
+			return;
+		}
+
+		UDMMaterialLayerObject* Layer = Stage->GetLayer();
+		if (!Layer)
+		{
+			return;
+		}
+
+		UDMMaterialSlot* Slot = Layer->GetSlot();
+		if (!Slot)
+		{
+			return;
+		}
+
+		UDynamicMaterialModel* MaterialModel = MenuContext->GetModel();
+		if (!MaterialModel)
+		{
+			return;
+		}
+
+		UDynamicMaterialModelEditorOnlyData* ModelEditorOnlyData = UDynamicMaterialModelEditorOnlyData::Get(MaterialModel);
+		if (!MaterialModel)
+		{
+			return;
+		}
+
+		const TArray<UDMMaterialSlot*>& Slots = ModelEditorOnlyData->GetSlots();
+		if (Slots.Num() <= 1)
+		{
+			return;
+		}
+
+		for (UDMMaterialSlot* SlotIter : Slots)
+		{
+			if (SlotIter == Slot)
+			{
+				continue;
+			}
+
+			if (SlotIter->GetLayers().IsEmpty())
+			{
+				continue;
+			}
+
+			TArray<EDMMaterialPropertyType> SlotProperties = ModelEditorOnlyData->GetMaterialPropertiesForSlot(SlotIter);
+			if (SlotProperties.IsEmpty())
+			{
+				continue;
+			}
+
+			if (SlotProperties.Num() == 1)
+			{
+				static const FText SlotNameFormatTemplate = LOCTEXT("SlotAndProperty", "{0} [{1}]");
+
+				UDMMaterialProperty* PropertyObj = ModelEditorOnlyData->GetMaterialProperty(SlotProperties[0]);
+
+				if (ensure(PropertyObj))
+				{
+					InMenu->AddMenuEntry(NAME_None, FToolMenuEntry::InitMenuEntry(NAME_None,
+						FText::Format(SlotNameFormatTemplate, SlotIter->GetDescription(), PropertyObj->GetDescription()),
+						LOCTEXT("ChangeSourceSlotSourceTooltip3", "Change the source of this stage to the output from this Material Slot's Property."),
+						FSlateIcon(),
+						FUIAction(FExecuteAction::CreateWeakLambda(
+							SlotIter,
+							[SlotIter, SlotProperty = SlotProperties[0], MenuContext]()
+							{
+								if (!IsValid(MenuContext))
+								{
+									return;
+								}
+
+								UDMMaterialStage* const Stage = MenuContext->GetStage();
+								if (!Stage)
+								{
+									return;
+								}
+
+								UDMMaterialStageSource* const StageSource = Stage->GetSource();
+								if (!StageSource)
+								{
+									return;
+								}
+
+								if (StageSource->IsA<UDMMaterialStageBlend>())
+								{
+									FScopedTransaction Transaction(LOCTEXT("SetStageInputBase", "Set Material Designer Base Source"));
+									Stage->Modify();
+									Stage->ChangeInput_Slot(UDMMaterialStageBlend::InputB, FDMMaterialStageConnectorChannel::WHOLE_CHANNEL,
+										SlotIter, SlotProperty, 0, FDMMaterialStageConnectorChannel::WHOLE_CHANNEL);
+								}
+								else if (StageSource->IsA<UDMMaterialStageThroughputLayerBlend>())
+								{
+									FScopedTransaction Transaction(LOCTEXT("SetStageInputMask", "Set Material Designer Mask Source"));
+									Stage->Modify();
+									Stage->ChangeInput_Slot(UDMMaterialStageThroughputLayerBlend::InputMaskSource, FDMMaterialStageConnectorChannel::WHOLE_CHANNEL,
+										SlotIter, SlotProperty, 0, FDMMaterialStageConnectorChannel::WHOLE_CHANNEL);
+								}
+								else
+								{
+									FScopedTransaction Transaction(LOCTEXT("SetStageSource", "Set Material Designer Stage Source"));
+									Stage->Modify();
+									Stage->ChangeSource_Slot(SlotIter, SlotProperty);
+								}
+
+								if (TSharedPtr<SDMSlot> SlotWidget = MenuContext->GetSlotWidget().Pin())
+								{
+									SlotWidget->InvalidateComponentEditWidget();
+								}
+							})
+						)
+					));
+				}
+			}
+			else
+			{
+				FToolMenuSection& NewSection = InMenu->AddSection("ChangeSourceSlotTooltip", LOCTEXT("ChangeSourceSlot", "Change Source Slot"));
+				NewSection.AddDynamicEntry(NAME_None,
+					FNewToolMenuSectionDelegate::CreateWeakLambda(
+						SlotIter,
+						[SlotIter](FToolMenuSection& InSection)
+						{
+							InSection.AddSubMenu(NAME_None,
+							SlotIter->GetDescription(),
+							LOCTEXT("ChangeSourceSlotTooltip", "Change the source of this stage to the output from another Material Slot."),
+							FNewToolMenuDelegate::CreateLambda([SlotIter](UToolMenu* InMenu) { GenerateChangeSourceMenu_Slot_Properties(InMenu, SlotIter); })
+							);
+						}));
+			}
+		}
+	}
+
+	void GenerateChangeSourceMenu_Gradients(UToolMenu* const InMenu)
+	{
+		if (!ensure(IsValid(InMenu)))
+		{
+			return;
+		}
+
+		UDMMenuContext* MenuContext = InMenu->FindContext<UDMMenuContext>();
+		if (!ensure(MenuContext))
+		{
+			return;
+		}
+
+		const TArray<TStrongObjectPtr<UClass>>& Gradients = UDMMaterialStageGradient::GetAvailableGradients();
+		if (Gradients.IsEmpty())
+		{
+			return;
+		}
+
+		for (TStrongObjectPtr<UClass> GradientClass : Gradients)
+		{
+			UDMMaterialStageGradient* GradientCDO = Cast<UDMMaterialStageGradient>(GradientClass->GetDefaultObject());
+
+			if (ensure(GradientCDO))
+			{
+				const FText MenuName = GradientCDO->GetDescription();
+
+				InMenu->AddMenuEntry(NAME_None, FToolMenuEntry::InitMenuEntry(NAME_None,
+					MenuName,
+					LOCTEXT("ChangeSourceGradientTooltip", "Change the source of this stage to a Material Gradient."),
+					FSlateIcon(),
+					FUIAction(FExecuteAction::CreateWeakLambda(
+						MenuContext,
+						[MenuContext, GradientClass]()
+						{
+							UDMMaterialStage* const Stage = MenuContext->GetStage();
+							if (!Stage)
+							{
+								return;
+							}
+
+							UDMMaterialStageSource* const StageSource = Stage->GetSource();
+							if (!StageSource)
+							{
+								return;
+							}
+
+							if (StageSource->IsA<UDMMaterialStageBlend>())
+							{
+								FScopedTransaction Transaction(LOCTEXT("SetStageInputBase", "Set Material Designer Base Source"));
+								Stage->Modify();
+								Stage->ChangeInput_Gradient(UDMMaterialStageBlend::InputB, FDMMaterialStageConnectorChannel::WHOLE_CHANNEL, GradientClass.Get(), FDMMaterialStageConnectorChannel::WHOLE_CHANNEL);
+							}
+							else if (StageSource->IsA<UDMMaterialStageThroughputLayerBlend>())
+							{
+								FScopedTransaction Transaction(LOCTEXT("SetStageInputMask", "Set Material Designer Mask Source"));
+								Stage->Modify();
+								Stage->ChangeInput_Gradient(UDMMaterialStageThroughputLayerBlend::InputMaskSource, FDMMaterialStageConnectorChannel::WHOLE_CHANNEL, GradientClass.Get(), FDMMaterialStageConnectorChannel::WHOLE_CHANNEL);
+							}
+							else
+							{
+								FScopedTransaction Transaction(LOCTEXT("SetStageSource", "Set Material Designer Stage Source"));
+								Stage->Modify();
+								Stage->ChangeSource_Gradient(GradientClass.Get());
+							}
+
+							if (TSharedPtr<SDMSlot> SlotWidget = MenuContext->GetSlotWidget().Pin())
+							{
+								SlotWidget->InvalidateComponentEditWidget();
+							}
+						})
+					)
+				));
+			}
+		}
+	}
+
+	void ChangeSourceToTextureSampleFromContext(UDMMenuContext* InMenuContext)
+	{
+		if (!IsValid(InMenuContext))
+		{
+			return;
+		}
+
+		UDMMaterialStage* const Stage = InMenuContext->GetStage();
+
+		if (!Stage)
+		{
+			return;
+		}
+
+		UDMMaterialStageSource* const StageSource = Stage->GetSource();
+
+		if (!StageSource)
+		{
+			return;
+		}
+
+		if (StageSource->IsA<UDMMaterialStageBlend>())
+		{
+			FScopedTransaction Transaction(LOCTEXT("SetStageInputBase", "Set Material Designer Base Source"));
+			Stage->Modify();
+			Stage->ChangeInput_Expression(UDMMaterialStageBlend::InputB, FDMMaterialStageConnectorChannel::WHOLE_CHANNEL,
+				UDMMaterialStageExpressionTextureSample::StaticClass(), 0, FDMMaterialStageConnectorChannel::THREE_CHANNELS);
+		}
+		else if (StageSource->IsA<UDMMaterialStageThroughputLayerBlend>())
+		{
+			FScopedTransaction Transaction(LOCTEXT("SetStageInputMask", "Set Material Designer Mask Source"));
+			Stage->Modify();
+			Stage->ChangeInput_Expression(UDMMaterialStageThroughputLayerBlend::InputMaskSource, FDMMaterialStageConnectorChannel::WHOLE_CHANNEL,
+				UDMMaterialStageExpressionTextureSample::StaticClass(), 0, FDMMaterialStageConnectorChannel::WHOLE_CHANNEL);
+		}
+		else
+		{
+			FScopedTransaction Transaction(LOCTEXT("SetStageSource", "Set Material Designer Stage Source"));
+			Stage->Modify();
+			Stage->ChangeSource_Expression(TSubclassOf<UDMMaterialStageExpression>(UDMMaterialStageExpressionTextureSample::StaticClass()));
+		}
+
+		if (TSharedPtr<SDMSlot> SlotWidget = InMenuContext->GetSlotWidget().Pin())
+		{
+			SlotWidget->InvalidateComponentEditWidget();
+		}
+	}
+
+	void ChangeSourceToSolidColorRGBFromContext(UDMMenuContext* InMenuContext)
+	{
+		if (!IsValid(InMenuContext))
+		{
+			return;
+		}
+
+		UDMMaterialStage* const Stage = InMenuContext->GetStage();
+
+		if (!Stage)
+		{
+			return;
+		}
+
+		UDMMaterialStageSource* const StageSource = Stage->GetSource();
+
+		if (!StageSource)
+		{
+			return;
+		}
+
+		if (StageSource->IsA<UDMMaterialStageBlend>())
+		{
+			FScopedTransaction Transaction(LOCTEXT("SetStageInputBase", "Set Material Designer Base Source"));
+			Stage->Modify();
+			Stage->ChangeInput_NewLocalValue(UDMMaterialStageBlend::InputB, FDMMaterialStageConnectorChannel::WHOLE_CHANNEL,
+				EDMValueType::VT_Float3_RGB, FDMMaterialStageConnectorChannel::THREE_CHANNELS);
+		}
+		else if (StageSource->IsA<UDMMaterialStageThroughputLayerBlend>())
+		{
+			FScopedTransaction Transaction(LOCTEXT("SetStageInputMask", "Set Material Designer Mask Source"));
+			Stage->Modify();
+			Stage->ChangeInput_NewLocalValue(UDMMaterialStageThroughputLayerBlend::InputMaskSource, FDMMaterialStageConnectorChannel::WHOLE_CHANNEL,
+				EDMValueType::VT_Float3_RGB, FDMMaterialStageConnectorChannel::WHOLE_CHANNEL);
+		}
+		else
+		{
+			FScopedTransaction Transaction(LOCTEXT("SetStageSource", "Set Material Designer Stage Source"));
+			Stage->Modify();
+			Stage->ChangeSource_NewLocalValue(EDMValueType::VT_Float3_RGB);
+		}
+
+		if (TSharedPtr<SDMSlot> SlotWidget = InMenuContext->GetSlotWidget().Pin())
+		{
+			SlotWidget->InvalidateComponentEditWidget();
+		}
+	}
+
+	void ChangeSourceToColorAtlasFromContext(UDMMenuContext* InMenuContext)
+	{
+		if (!IsValid(InMenuContext))
+		{
+			return;
+		}
+
+		UDMMaterialStage* const Stage = InMenuContext->GetStage();
+
+		if (!Stage)
+		{
+			return;
+		}
+
+		UDMMaterialStageSource* const StageSource = Stage->GetSource();
+
+		if (!StageSource)
+		{
+			return;
+		}
+
+		if (StageSource->IsA<UDMMaterialStageBlend>())
+		{
+			FScopedTransaction Transaction(LOCTEXT("SetStageInputBase", "Set Material Designer Base Source"));
+			Stage->Modify();
+			Stage->ChangeInput_NewLocalValue(UDMMaterialStageBlend::InputB, FDMMaterialStageConnectorChannel::WHOLE_CHANNEL,
+				EDMValueType::VT_ColorAtlas, FDMMaterialStageConnectorChannel::THREE_CHANNELS);
+
+		}
+		else if (StageSource->IsA<UDMMaterialStageThroughputLayerBlend>())
+		{
+			FScopedTransaction Transaction(LOCTEXT("SetStageInputMask", "Set Material Designer Mask Source"));
+			Stage->Modify();
+			Stage->ChangeInput_NewLocalValue(UDMMaterialStageThroughputLayerBlend::InputMaskSource, FDMMaterialStageConnectorChannel::WHOLE_CHANNEL,
+				EDMValueType::VT_ColorAtlas, FDMMaterialStageConnectorChannel::FOURTH_CHANNEL);
+		}
+		else
+		{
+			FScopedTransaction Transaction(LOCTEXT("SetStageSource", "Set Material Designer Stage Source"));
+			Stage->Modify();
+			Stage->ChangeSource_NewLocalValue(EDMValueType::VT_ColorAtlas);
+		}
+
+		if (TSharedPtr<SDMSlot> SlotWidget = InMenuContext->GetSlotWidget().Pin())
+		{
+			SlotWidget->InvalidateComponentEditWidget();
+		}
+	}
+
+	void ChangeSourceToTextureSampleEdgeColorFromContext(UDMMenuContext* InMenuContext)
+	{
+		if (!IsValid(InMenuContext))
+		{
+			return;
+		}
+
+		UDMMaterialStageSource* const StageSource = InMenuContext->GetStageSource();
+
+		if (!StageSource)
+		{
+			return;
+		}
+
+		UDMMaterialStage* const Stage = InMenuContext->GetStage();
+
+		if (!Stage)
+		{
+			return;
+		}
+
+		UDMMaterialLayerObject* const Layer = Stage->GetLayer();
+
+		if (!Layer)
+		{
+			return;
+		}
+
+		if (StageSource->IsA<UDMMaterialStageBlend>())
+		{
+			FScopedTransaction Transaction(LOCTEXT("SetStageInputBase", "Set Material Designer Base Source"));
+			Stage->Modify();
+			Stage->ChangeInput_Expression(UDMMaterialStageBlend::InputB, FDMMaterialStageConnectorChannel::WHOLE_CHANNEL,
+				UDMMaterialStageExpressionTextureSampleEdgeColor::StaticClass(), 0, FDMMaterialStageConnectorChannel::THREE_CHANNELS);
+		}
+		else if (StageSource->IsA<UDMMaterialStageThroughputLayerBlend>())
+		{
+			FScopedTransaction Transaction(LOCTEXT("SetStageInputMask", "Set Material Designer Mask Source"));
+			Stage->Modify();
+			Stage->ChangeInput_Expression(UDMMaterialStageThroughputLayerBlend::InputMaskSource, FDMMaterialStageConnectorChannel::WHOLE_CHANNEL,
+				UDMMaterialStageExpressionTextureSampleEdgeColor::StaticClass(), 0, FDMMaterialStageConnectorChannel::WHOLE_CHANNEL);
+		}
+		else
+		{
+			FScopedTransaction Transaction(LOCTEXT("SetStageInput", "Set Material Designer Source"));
+			Stage->Modify();
+			Stage->ChangeSource_Expression(TSubclassOf<UDMMaterialStageExpression>(UDMMaterialStageExpressionTextureSampleEdgeColor::StaticClass()));
+		}
+
+		if (UDMMaterialStage* MaskStage = Layer->GetStage(EDMMaterialLayerStage::Mask))
+		{
+			MaskStage->SetEnabled(false);
+		}
+
+		if (TSharedPtr<SDMSlot> SlotWidget = InMenuContext->GetSlotWidget().Pin())
+		{
+			SlotWidget->InvalidateComponentEditWidget();
+		}
+	}
+
+	void ChangeSourceToSceneTextureFromContext(UDMMenuContext* InMenuContext)
+	{
+		if (!IsValid(InMenuContext))
+		{
+			return;
+		}
+
+		UDMMaterialStageSource* const StageSource = InMenuContext->GetStageSource();
+
+		if (!StageSource)
+		{
+			return;
+		}
+
+		UDMMaterialStage* const Stage = InMenuContext->GetStage();
+
+		if (!Stage)
+		{
+			return;
+		}
+
+		UDMMaterialLayerObject* const Layer = Stage->GetLayer();
+
+		if (!Layer)
+		{
+			return;
+		}
+
+		if (StageSource->IsA<UDMMaterialStageBlend>())
+		{
+			FScopedTransaction Transaction(LOCTEXT("SetStageInputBase", "Set Material Designer Base Source"));
+			Stage->Modify();
+			Stage->ChangeInput_Expression(UDMMaterialStageBlend::InputB, FDMMaterialStageConnectorChannel::WHOLE_CHANNEL,
+				UDMMaterialStageExpressionSceneTexture::StaticClass(), 0, FDMMaterialStageConnectorChannel::THREE_CHANNELS);
+
+			if (Layer->GetStageType(Stage) == EDMMaterialLayerStage::Base)
+			{
+				if (UDMMaterialStage* MaskStage = Layer->GetStage(EDMMaterialLayerStage::Mask, true))
+				{
+					MaskStage->ChangeInput_Expression(UDMMaterialStageThroughputLayerBlend::InputMaskSource, FDMMaterialStageConnectorChannel::WHOLE_CHANNEL,
+						UDMMaterialStageExpressionSceneTexture::StaticClass(), 0, FDMMaterialStageConnectorChannel::FOURTH_CHANNEL);
+				}
+			}
+		}
+		else if (StageSource->IsA<UDMMaterialStageThroughputLayerBlend>())
+		{
+			FScopedTransaction Transaction(LOCTEXT("SetStageInputMask", "Set Material Designer Mask Source"));
+			Stage->Modify();
+			Stage->ChangeInput_Expression(UDMMaterialStageThroughputLayerBlend::InputMaskSource, FDMMaterialStageConnectorChannel::WHOLE_CHANNEL,
+				UDMMaterialStageExpressionSceneTexture::StaticClass(), 0, FDMMaterialStageConnectorChannel::FOURTH_CHANNEL);
+		}
+		else
+		{
+			FScopedTransaction Transaction(LOCTEXT("SetStageInput", "Set Material Designer Source"));
+			Stage->Modify();
+			Stage->ChangeSource_Expression(UDMMaterialStageExpressionSceneTexture::StaticClass());
+		}
+
+		if (TSharedPtr<SDMSlot> SlotWidget = InMenuContext->GetSlotWidget().Pin())
+		{
+			SlotWidget->InvalidateComponentEditWidget();
+		}
+	}
+
+	bool CanChangeSourceToSceneTextureFromContext(UDMMenuContext* InMenuContext)
+	{
+		if (IsValid(InMenuContext))
+		{
+			if (UDMMaterialStageSource* const StageSource = InMenuContext->GetStageSource())
+			{
+				if (UDMMaterialStage* const Stage = InMenuContext->GetStage())
+				{
+					if (UDMMaterialLayerObject* const Layer = Stage->GetLayer())
+					{
+						if (UDMMaterialSlot* Slot = Layer->GetSlot())
+						{
+							if (UDynamicMaterialModelEditorOnlyData* EditorOnlyData = Slot->GetMaterialModelEditorOnlyData())
+							{
+								return EditorOnlyData->GetDomain() == EMaterialDomain::MD_PostProcess;
+							}
+						}
+					}
+				}
+			}
+		}
+
+		return false;
+	}
+
+	void ChangeSourceToMaterialFunctionFromContext(UDMMenuContext* InMenuContext)
+	{
+		if (!IsValid(InMenuContext))
+		{
+			return;
+		}
+
+		UDMMaterialStage* const Stage = InMenuContext->GetStage();
+
+		if (!Stage)
+		{
+			return;
+		}
+
+		UDMMaterialStageSource* const StageSource = Stage->GetSource();
+
+		if (!StageSource)
+		{
+			return;
+		}
+
+		if (StageSource->IsA<UDMMaterialStageBlend>())
+		{
+			FScopedTransaction Transaction(LOCTEXT("SetStageInputBase", "Set Material Designer Base Source"));
+			Stage->Modify();
+			Stage->ChangeInput_MaterialFunction(UDMMaterialStageBlend::InputB, FDMMaterialStageConnectorChannel::WHOLE_CHANNEL);
+		}
+		else if (StageSource->IsA<UDMMaterialStageThroughputLayerBlend>())
+		{
+			FScopedTransaction Transaction(LOCTEXT("SetStageInputMask", "Set Material Designer Mask Source"));
+			Stage->Modify();
+			Stage->ChangeInput_MaterialFunction(UDMMaterialStageThroughputLayerBlend::InputMaskSource, FDMMaterialStageConnectorChannel::WHOLE_CHANNEL);
+		}
+		else
+		{
+			FScopedTransaction Transaction(LOCTEXT("SetStageSource", "Set Material Designer Stage Source"));
+			Stage->Modify();
+			Stage->ChangeSource_MaterialFunction();
+		}
+
+		if (TSharedPtr<SDMSlot> SlotWidget = InMenuContext->GetSlotWidget().Pin())
+		{
+			SlotWidget->InvalidateComponentEditWidget();
+		}
+	}
+
+	void CreateChangeMaterialStageSource(FToolMenuSection& InSection)
+	{
+		UDMMenuContext* MenuContext = InSection.FindContext<UDMMenuContext>();
+
+		if (!ensure(MenuContext))
+		{
+			return;
+		}
+
+		UDMMaterialSlot* Slot = MenuContext->GetSlot();
+
+		UDynamicMaterialModel* MaterialModel = MenuContext->GetModel();
+
+		if (!ensure(MaterialModel))
+		{
+			return;
+		}
+
+		UDynamicMaterialModelEditorOnlyData* ModelEditorOnlyData = UDynamicMaterialModelEditorOnlyData::Get(MaterialModel);
+
+		if (!ensure(ModelEditorOnlyData))
+		{
+			return;
+		}
+
+		const TArray<TStrongObjectPtr<UClass>>& Gradients = UDMMaterialStageGradient::GetAvailableGradients();
+
+		bool bHasValidSlot = false;
+
+		if constexpr (bAdvancedSlotsEnabled)
+		{
+			const TArray<UDMMaterialSlot*>& Slots = ModelEditorOnlyData->GetSlots();
+
+			for (UDMMaterialSlot* SlotIter : Slots)
+			{
+				if (SlotIter == Slot)
+				{
+					continue;
+				}
+
+				if (SlotIter->GetLayers().IsEmpty())
+				{
+					continue;
+				}
+
+				bHasValidSlot = true;
+				break;
+			}
+		}
+
+		InSection.AddMenuEntry("TextureSample",
+			LOCTEXT("TextureSample", "Texture"),
+			LOCTEXT("ChangeSourceTextureSampleTooltip", "Change the source of this stage to a texture."),
+			FSlateIcon(),
+			FUIAction(
+				FExecuteAction::CreateStatic(
+					&ChangeSourceToTextureSampleFromContext,
+					MenuContext
+				)
+			)
+		);
+
+		InSection.AddMenuEntry("SolidColor",
+			LOCTEXT("ChangeSourceColorRGB", "Solid Color"),
+			LOCTEXT("ChangeSourceColorRGBTooltip", "Change the source of this stage to a Solid Color."),
+			FSlateIcon(),
+			FUIAction(
+				FExecuteAction::CreateStatic(
+					&ChangeSourceToSolidColorRGBFromContext,
+					MenuContext
+				)
+			)
+		);
+
+		InSection.AddMenuEntry("ColorAtlas",
+			LOCTEXT("ChangeSourceColorAtlas", "Color Atlas"),
+			LOCTEXT("ChangeSourceColorAtlasTooltip", "Change the source of this stage to a Color Atlas."),
+			FSlateIcon(),
+			FUIAction(
+				FExecuteAction::CreateStatic(
+					&ChangeSourceToColorAtlasFromContext,
+					MenuContext
+				)
+			)
+		);
+
+		InSection.AddMenuEntry("TextureSample_EdgeColor",
+			LOCTEXT("AddTextureSampleEgdeColor", "Texture Edge Color"),
+			LOCTEXT("ChangeSourceTextureSampleTooltip", "Change the source of this stage to the edge color of a texture."),
+			FSlateIcon(),
+			FUIAction(
+				FExecuteAction::CreateStatic(
+					&ChangeSourceToTextureSampleEdgeColorFromContext,
+					MenuContext
+				)
+			)
+		);
+
+		InSection.AddMenuEntry("SceneTexture",
+			LOCTEXT("AddSceneTexture", "Scene Texture"),
+			LOCTEXT("ChangeSourceSceneTextureTooltip", "Change the source of this stage to Scene Texture in post process materials."),
+			FSlateIcon(),
+			FUIAction(
+				FExecuteAction::CreateStatic(
+					&ChangeSourceToSolidColorRGBFromContext,
+					MenuContext
+				),
+				FCanExecuteAction::CreateStatic(
+					&CanChangeSourceToSceneTextureFromContext,
+					MenuContext
+				)
+			)
+		);
+
+		if constexpr (UE::DynamicMaterialEditor::bGlobalValuesEnabled)
+		{
+			const TArray<UDMMaterialValue*>& Values = MaterialModel->GetValues();
+
+			InSection.AddSubMenu("NewLocalValue",
+				LOCTEXT("ChangeSourceNewLocalValue", "New Local Value"),
+				LOCTEXT("ChangeSourceNewLocalValueTooltip", "Add a new local Material Value and use it as the source of this stage."),
+				FNewToolMenuDelegate::CreateStatic(&GenerateChangeSourceMenu_NewLocalValues)
+			);
+
+			if (!Values.IsEmpty())
+			{
+				InSection.AddSubMenu("GlobalValue",
+					LOCTEXT("ChangeSourceValue", "Global Value"),
+					LOCTEXT("ChangeSourceValueTooltip", "Change the source of this stage to a global Material Value."),
+					FNewToolMenuDelegate::CreateStatic(&GenerateChangeSourceMenu_GlobalValues)
+				);
+			}
+
+			InSection.AddSubMenu("NewLocalValue",
+				LOCTEXT("ChangeSourceNewValue", "New Global Value"),
+				LOCTEXT("ChangeSourceNewValueTooltip", "Add a new global Material Value and use it as the source of this stage."),
+				FNewToolMenuDelegate::CreateStatic(&GenerateChangeSourceMenu_NewGlobalValues)
+			);
+		}
+
+		if (!Gradients.IsEmpty())
+		{
+			InSection.AddSubMenu("Gradient",
+				LOCTEXT("ChangeSourceGradient", "Gradient"),
+				LOCTEXT("ChangeSourceGradientTooltip", "Change the source of this stage to a Material Gradient."),
+				FNewToolMenuDelegate::CreateStatic(&GenerateChangeSourceMenu_Gradients)
+			);
+		}
+
+		InSection.AddMenuEntry("MaterialFunction",
+			LOCTEXT("ChangeSourceMaterialFunction", "Material Function"),
+			LOCTEXT("ChangeSourceMaterialFunctionTooltip", "Change the source of this stage to a Material Function."),
+			FSlateIcon(),
+			FUIAction(
+				FExecuteAction::CreateStatic(
+					&ChangeSourceToMaterialFunctionFromContext,
+					MenuContext
+				)
+			)
+		);
+
+		if constexpr (UE::DynamicMaterialEditor::bAdvancedSlotsEnabled)
+		{
+			if (bHasValidSlot)
+			{
+				InSection.AddDynamicEntry(
+					NAME_None,
+					FNewToolMenuSectionDelegate::CreateLambda(
+						[](FToolMenuSection& InSection)
+						{
+							InSection.AddSubMenu("SlotOutput",
+							LOCTEXT("ChangeSourceSlot", "Slot Output"),
+							LOCTEXT("ChangeSourceSlotTooltip", "Change the source of this stage to the output from another Material Slot."),
+							FNewToolMenuDelegate::CreateStatic(&GenerateChangeSourceMenu_Slots)
+							);
+						}));
+			}
+		}
+	}
+}
+
+TSharedRef<SWidget> FDMMaterialStageSourceMenus::MakeChangeSourceMenu(const TSharedPtr<SDMSlot>& InSlotWidget, const TSharedPtr<SDMStage>& InStageWidget)
+{
+	using namespace UE::DynamicMaterialEditor::Private;
+
+	if (!UToolMenus::Get()->IsMenuRegistered(ChangeStageSourceMenuName))
+	{
+		UToolMenu* const NewToolMenu = UDMMenuContext::GenerateContextMenuDefault(ChangeStageSourceMenuName);
+		if (!NewToolMenu)
+		{
+			return SNullWidget::NullWidget;
+		}
+
+		FToolMenuSection& NewSection = NewToolMenu->AddSection("ChangeMaterialStage", LOCTEXT("ChangeStageSource", "Change Stage Source"));
+		NewSection.AddDynamicEntry(NAME_None, FNewToolMenuSectionDelegate::CreateStatic(&CreateChangeMaterialStageSource));
+	}
+
+	FToolMenuContext MenuContext(UDMMenuContext::CreateStage(InSlotWidget, InStageWidget));
+
+	return UToolMenus::Get()->GenerateWidget(ChangeStageSourceMenuName, MenuContext);
+}
+
+void FDMMaterialStageSourceMenus::CreateSourceMenuTree(TFunction<void(EDMExpressionMenu Menu, TArray<UDMMaterialStageExpression*>& SubmenuExpressionList)> Callback, const TArray<TStrongObjectPtr<UClass>>& AllExpressions)
+{
+	TMap<EDMExpressionMenu, TArray<UDMMaterialStageExpression*>> MenuMap;
+
+	for (const TStrongObjectPtr<UClass>& Class : AllExpressions)
+	{
+		TSubclassOf<UDMMaterialStageExpression> ExpressionClass = Class.Get();
+		if (!ExpressionClass.Get())
+		{
+			continue;
+		}
+
+		UDMMaterialStageExpression* ExpressionCDO = Cast<UDMMaterialStageExpression>(ExpressionClass->GetDefaultObject(true));
+		if (!ExpressionCDO)
+		{
+			continue;
+		}
+
+		const TArray<EDMExpressionMenu>& Menus = ExpressionCDO->GetMenus();
+		for (EDMExpressionMenu Menu : Menus)
+		{
+			MenuMap.FindOrAdd(Menu).Add(ExpressionCDO);
+		}
+	}
+
+	auto CreateMenu = [&Callback, &MenuMap](EDMExpressionMenu Menu)
+	{
+		TArray<UDMMaterialStageExpression*>* ExpressionListPtr = MenuMap.Find(Menu);
+		if (!ExpressionListPtr || ExpressionListPtr->Num() == 0)
+		{
+			return;
+		}
+
+		Callback(Menu, *ExpressionListPtr);
+	};
+
+	CreateMenu(EDMExpressionMenu::Texture);
+	CreateMenu(EDMExpressionMenu::Math);
+	CreateMenu(EDMExpressionMenu::Geometry);
+	CreateMenu(EDMExpressionMenu::Object);
+	CreateMenu(EDMExpressionMenu::WorldSpace);
+	CreateMenu(EDMExpressionMenu::Time);
+	CreateMenu(EDMExpressionMenu::Camera);
+	CreateMenu(EDMExpressionMenu::Particle);
+	CreateMenu(EDMExpressionMenu::Decal);
+	CreateMenu(EDMExpressionMenu::Landscape);
+	CreateMenu(EDMExpressionMenu::Other);
+}
+
+#undef LOCTEXT_NAMESPACE
diff --git a/Engine/Plugins/Experimental/DynamicMaterial/Source/DynamicMaterialEditor/Private/Menus/DMMaterialStageSourceMenus.h b/Engine/Plugins/Experimental/DynamicMaterial/Source/DynamicMaterialEditor/Private/Menus/DMMaterialStageSourceMenus.h
new file mode 100644
--- /dev/null
+++ b/Engine/Plugins/Experimental/DynamicMaterial/Source/DynamicMaterialEditor/Private/Menus/DMMaterialStageSourceMenus.h
@@ -0,0 +1,22 @@
+// Copyright Epic Games, Inc. All Rights Reserved.
+
+#pragma once
+
+#include "Templates/SharedPointerFwd.h"
+#include "UObject/StrongObjectPtr.h"
+
+class SDMSlot;
+class SDMStage;
+class SWidget;
+class UDMMaterialStageExpression;
+enum class EDMExpressionMenu : uint8;
+
+class FDMMaterialStageSourceMenus
+{
+public:
+	/** Generate right click menu for changing sources */
+	static TSharedRef<SWidget> MakeChangeSourceMenu(const TSharedPtr<SDMSlot>& InSlotWidget, const TSharedPtr<SDMStage>& InStageWidget);
+
+	static void CreateSourceMenuTree(TFunction<void(EDMExpressionMenu Menu, TArray<UDMMaterialStageExpression*>& SubmenuExpressionList)> Callback, 
+		const TArray<TStrongObjectPtr<UClass>>& AllExpressions);
+};
diff --git a/Engine/Plugins/Experimental/DynamicMaterial/Source/DynamicMaterialEditor/Private/Menus/DMMenuContext.cpp b/Engine/Plugins/Experimental/DynamicMaterial/Source/DynamicMaterialEditor/Private/Menus/DMMenuContext.cpp
new file mode 100644
--- /dev/null
+++ b/Engine/Plugins/Experimental/DynamicMaterial/Source/DynamicMaterialEditor/Private/Menus/DMMenuContext.cpp
@@ -0,0 +1,154 @@
+// Copyright Epic Games, Inc. All Rights Reserved.
+
+#include "DMMenuContext.h"
+#include "Components/DMMaterialLayer.h"
+#include "Components/DMMaterialSlot.h"
+#include "Components/DMMaterialStage.h"
+#include "Components/DMMaterialStageBlend.h"
+#include "Components/DMMaterialStageSource.h"
+#include "Model/DynamicMaterialModel.h"
+#include "Model/DynamicMaterialModelEditorOnlyData.h"
+#include "Slate/SDMSlot.h"
+#include "Slate/SDMStage.h"
+#include "ToolMenu.h"
+#include "ToolMenus.h"
+
+UDMMenuContext* UDMMenuContext::Create(const TWeakPtr<SDMSlot>& InSlotWidget, const TWeakPtr<SDMStage>& InStageWidget, UDMMaterialLayerObject* InLayerObject)
+{
+	UDMMenuContext* Context = NewObject<UDMMenuContext>();
+	Context->SlotWidgetWeak = InSlotWidget;
+	Context->StageWidgetWeak = InStageWidget;
+	Context->LayerObjectWeak = InLayerObject;
+	return Context;
+}
+
+UToolMenu* UDMMenuContext::GenerateContextMenu(const FName InMenuName, const TWeakPtr<SDMSlot>& InSlotWidget, const TWeakPtr<SDMStage>& InStageWidget, UDMMaterialLayerObject* InLayerObject)
+{
+	UToolMenu* NewMenu = UToolMenus::Get()->RegisterMenu(InMenuName, NAME_None, EMultiBoxType::Menu, false);
+	if (!NewMenu)
+	{
+		return nullptr;
+	}
+
+	if (!IsValid(InLayerObject))
+	{
+		if (TSharedPtr<SDMStage> StageWidget = InStageWidget.Pin())
+		{
+			if (UDMMaterialStage* Stage = StageWidget->GetStage())
+			{
+				InLayerObject = Stage->GetLayer();
+			}
+		}
+	}
+
+	NewMenu->bToolBarForceSmallIcons = true;
+	NewMenu->bShouldCloseWindowAfterMenuSelection = true;
+	NewMenu->bCloseSelfOnly = true;
+
+	TSharedPtr<FUICommandList> CommandList = nullptr;
+
+	if (TSharedPtr<SDMSlot> Slot = InSlotWidget.Pin())
+	{
+		if (TSharedPtr<SDMEditor> Editor = Slot->GetEditorWidget())
+		{
+			CommandList = Editor->GetCommandList();
+		}
+	}
+
+	NewMenu->Context = FToolMenuContext(CommandList, TSharedPtr<FExtender>(), Create(InSlotWidget, InStageWidget, InLayerObject));
+
+	return NewMenu;
+}
+
+UDMMenuContext* UDMMenuContext::CreateEmpty()
+{
+	return Create(nullptr, nullptr, nullptr);
+}
+
+UDMMenuContext* UDMMenuContext::CreateSlot(const TWeakPtr<SDMSlot>& InSlotWidget)
+{
+	return Create(InSlotWidget, nullptr, nullptr);
+}
+
+UDMMenuContext* UDMMenuContext::CreateLayer(const TWeakPtr<SDMSlot>& InSlotWidget, UDMMaterialLayerObject* InLayerObject)
+{
+	return Create(InSlotWidget, nullptr, InLayerObject);
+}
+
+UDMMenuContext* UDMMenuContext::CreateStage(const TWeakPtr<SDMSlot>& InSlotWidget, const TWeakPtr<SDMStage>& InStageWidget)
+{
+	return Create(InSlotWidget, InStageWidget, nullptr);
+}
+
+UToolMenu* UDMMenuContext::GenerateContextMenuDefault(const FName InMenuName)
+{
+	return GenerateContextMenu(InMenuName, nullptr, nullptr, nullptr);
+}
+
+UToolMenu* UDMMenuContext::GenerateContextMenuSlot(const FName InMenuName, const TWeakPtr<SDMSlot>& InSlotWidget)
+{
+	return GenerateContextMenu(InMenuName, InSlotWidget, nullptr, nullptr);
+}
+
+UToolMenu* UDMMenuContext::GenerateContextMenuLayer(const FName InMenuName, const TWeakPtr<SDMSlot>& InSlotWidget, UDMMaterialLayerObject* InLayerObject)
+{
+	return GenerateContextMenu(InMenuName, InSlotWidget, nullptr, InLayerObject);
+}
+
+UToolMenu* UDMMenuContext::GenerateContextMenuStage(const FName InMenuName, const TWeakPtr<SDMSlot>& InSlotWidget, const TWeakPtr<SDMStage>& InStageWidget)
+{
+	return GenerateContextMenu(InMenuName, InSlotWidget, InStageWidget, nullptr);
+}
+
+UDMMaterialSlot* UDMMenuContext::GetSlot() const
+{
+	if (TSharedPtr<SDMSlot> SlotWidget = SlotWidgetWeak.Pin())
+	{
+		return SlotWidget->GetSlot();
+	}
+
+	return nullptr;
+}
+
+UDynamicMaterialModel* UDMMenuContext::GetModel() const
+{
+	if (const UDMMaterialSlot* const Slot = GetSlot())
+	{
+		if (UDynamicMaterialModelEditorOnlyData* ModelEditorOnlyData = Slot->GetMaterialModelEditorOnlyData())
+		{
+			return ModelEditorOnlyData->GetMaterialModel();
+		}
+	}
+
+	return nullptr;
+}
+
+UDMMaterialStage* UDMMenuContext::GetStage() const
+{
+	if (TSharedPtr<SDMStage> StageWidget = StageWidgetWeak.Pin())
+	{
+		return StageWidget->GetStage();
+	}
+
+	return nullptr;
+}
+
+UDMMaterialStageSource* UDMMenuContext::GetStageSource() const
+{
+	if (UDMMaterialStage* Stage = GetStage())
+	{
+		return Stage->GetSource();
+	}
+
+	return nullptr;
+}
+
+UDMMaterialStageBlend* UDMMenuContext::GetStageSourceAsBlend() const
+{
+	return Cast<UDMMaterialStageBlend>(GetStageSource());
+}
+
+const UDMMaterialLayerObject* UDMMenuContext::GetLayer() const
+{
+	return LayerObjectWeak.Get();
+}
diff --git a/Engine/Plugins/Experimental/DynamicMaterial/Source/DynamicMaterialEditor/Private/Menus/DMMenuContext.h b/Engine/Plugins/Experimental/DynamicMaterial/Source/DynamicMaterialEditor/Private/Menus/DMMenuContext.h
new file mode 100644
--- /dev/null
+++ b/Engine/Plugins/Experimental/DynamicMaterial/Source/DynamicMaterialEditor/Private/Menus/DMMenuContext.h
@@ -0,0 +1,61 @@
+// Copyright Epic Games, Inc. All Rights Reserved.
+
+#pragma once
+
+#include "Templates/SharedPointer.h"
+#include "DMMenuContext.generated.h"
+
+class SDMSlot;
+class SDMStage;
+class UDMMaterialLayerObject;
+class UDMMaterialSlot;
+class UDMMaterialStage;
+class UDMMaterialStageBlend;
+class UDMMaterialStageSource;
+class UDynamicMaterialModel;
+class UToolMenu;
+
+UCLASS(BlueprintType)
+class DYNAMICMATERIALEDITOR_API UDMMenuContext : public UObject
+{
+	GENERATED_BODY()
+
+public:
+	static UDMMenuContext* CreateEmpty();
+	static UDMMenuContext* CreateSlot(const TWeakPtr<SDMSlot>& InSlotWidget);
+	static UDMMenuContext* CreateLayer(const TWeakPtr<SDMSlot>& InSlotWidget, UDMMaterialLayerObject* InLayerObject);
+	static UDMMenuContext* CreateStage(const TWeakPtr<SDMSlot>& InSlotWidget, const TWeakPtr<SDMStage>& InStageWidget);
+
+	static UToolMenu* GenerateContextMenuDefault(const FName InMenuName);
+	static UToolMenu* GenerateContextMenuSlot(const FName InMenuName, const TWeakPtr<SDMSlot>& InSlotWidget);
+	static UToolMenu* GenerateContextMenuLayer(const FName InMenuName, const TWeakPtr<SDMSlot>& InSlotWidget, UDMMaterialLayerObject* InLayerObject);
+	static UToolMenu* GenerateContextMenuStage(const FName InMenuName, const TWeakPtr<SDMSlot>& InSlotWidget, const TWeakPtr<SDMStage>& InStageWidget);
+
+	const TWeakPtr<SDMSlot>& GetSlotWidget() const { return SlotWidgetWeak; }
+	const TWeakPtr<SDMStage>& GetStageWidget() const { return StageWidgetWeak; }
+
+	UFUNCTION(BlueprintCallable, Category = "Material Designer Extensions")
+	UDMMaterialSlot* GetSlot() const;
+
+	UFUNCTION(BlueprintCallable, Category = "Material Designer Extensions")
+	UDynamicMaterialModel* GetModel() const;
+
+	UFUNCTION(BlueprintCallable, Category = "Material Designer Extensions")
+	UDMMaterialStage* GetStage() const;
+
+	UFUNCTION(BlueprintCallable, Category = "Material Designer Extensions")
+	UDMMaterialStageSource* GetStageSource() const;
+
+	UFUNCTION(BlueprintCallable, Category = "Material Designer Extensions")
+	UDMMaterialStageBlend* GetStageSourceAsBlend() const;
+
+	const UDMMaterialLayerObject* GetLayer() const;
+
+protected:
+	TWeakPtr<SDMSlot> SlotWidgetWeak;
+	TWeakPtr<SDMStage> StageWidgetWeak;
+	TWeakObjectPtr<UDMMaterialLayerObject> LayerObjectWeak;
+
+	static UDMMenuContext* Create(const TWeakPtr<SDMSlot>& InSlotWidget, const TWeakPtr<SDMStage>& InStageWidget, UDMMaterialLayerObject* InLayerObject);
+	static UToolMenu* GenerateContextMenu(const FName InMenuName, const TWeakPtr<SDMSlot>& InSlotWidget, const TWeakPtr<SDMStage>& InStageWidget, UDMMaterialLayerObject* InLayerObject);
+};
diff --git a/Engine/Plugins/Experimental/DynamicMaterial/Source/DynamicMaterialEditor/Private/Menus/DMToolBarMenus.cpp b/Engine/Plugins/Experimental/DynamicMaterial/Source/DynamicMaterialEditor/Private/Menus/DMToolBarMenus.cpp
new file mode 100644
--- /dev/null
+++ b/Engine/Plugins/Experimental/DynamicMaterial/Source/DynamicMaterialEditor/Private/Menus/DMToolBarMenus.cpp
@@ -0,0 +1,368 @@
+// Copyright Epic Games, Inc. All Rights Reserved.
+
+#include "DMToolBarMenus.h"
+#include "ContentBrowserModule.h"
+#include "DMBlueprintFunctionLibrary.h"
+#include "DynamicMaterialEditorModule.h"
+#include "DynamicMaterialEditorSettings.h"
+#include "IContentBrowserSingleton.h"
+#include "ISinglePropertyView.h"
+#include "Material/DynamicMaterialInstance.h"
+#include "Materials/Material.h"
+#include "Menus/DMMenuContext.h"
+#include "Model/DynamicMaterialModel.h"
+#include "Model/DynamicMaterialModelEditorOnlyData.h"
+#include "PropertyEditorModule.h"
+#include "Slate/SDMEditor.h"
+#include "Slate/SDMSlot.h"
+#include "ToolMenu.h"
+#include "ToolMenus.h"
+#include "Widgets/SWidget.h"
+
+#define LOCTEXT_NAMESPACE "FDMToolBarMenus"
+
+namespace UE::DynamicMaterialEditor::Private
+{
+	static FName ToolBarEditorLayoutMenuName = TEXT("MaterialDesigner.EditorLayout");
+	static FName ToolBarPreviewOptionsSectionName = TEXT("PreviewOptions");
+	static FName ToolBarTooltipOptionsSectionName = TEXT("TooltipOptions");
+	static FName ToolBarExportSectionName = TEXT("Export");
+
+	void OpenMaterialEditorFromContext(UDMMenuContext* InMenuContext)
+	{
+		if (IsValid(InMenuContext))
+		{
+			if (UDynamicMaterialModelEditorOnlyData* const ModelEditorOnlyData = UDynamicMaterialModelEditorOnlyData::Get(InMenuContext->GetModel()))
+			{
+				ModelEditorOnlyData->OpenMaterialEditor();
+			}
+		}
+	}
+
+	void ExportMaterialInstanceFromInstance(TWeakObjectPtr<UDynamicMaterialInstance> InMaterialInstanceWeak)
+	{
+		if (UDynamicMaterialInstance* MaterialInstance = InMaterialInstanceWeak.Get())
+		{
+			FSaveAssetDialogConfig SaveAssetDialogConfig;
+			SaveAssetDialogConfig.DialogTitleOverride = LOCTEXT("SaveAssetDialogTitle", "Save Asset As");
+			SaveAssetDialogConfig.DefaultPath = "/Game";
+			SaveAssetDialogConfig.DefaultAssetName = MaterialInstance->GetName();
+			SaveAssetDialogConfig.ExistingAssetPolicy = ESaveAssetDialogExistingAssetPolicy::AllowButWarn;
+
+			FContentBrowserModule& ContentBrowserModule = FModuleManager::LoadModuleChecked<FContentBrowserModule>("ContentBrowser");
+			FString SaveObjectPath = ContentBrowserModule.Get().CreateModalSaveAssetDialog(SaveAssetDialogConfig);
+
+			UDMBlueprintFunctionLibrary::ExportMaterialInstance(MaterialInstance->GetMaterialModel(), SaveObjectPath);
+		}
+	}
+
+	void ExportMaterialModelFromModel(TWeakObjectPtr<UDynamicMaterialModel> InMaterialModelWeak)
+	{
+		if (UDynamicMaterialModel* MaterialModel = InMaterialModelWeak.Get())
+		{
+			UMaterial* GeneratedMaterial = MaterialModel->GetGeneratedMaterial();
+
+			if (!GeneratedMaterial)
+			{
+				UE_LOG(LogDynamicMaterialEditor, Error, TEXT("Failed to find a generated material to export."));
+				return;
+			}
+
+			FSaveAssetDialogConfig SaveAssetDialogConfig;
+			SaveAssetDialogConfig.DialogTitleOverride = LOCTEXT("SaveAssetDialogTitle", "Save Asset As");
+			SaveAssetDialogConfig.DefaultPath = "/Game";
+			SaveAssetDialogConfig.DefaultAssetName = GeneratedMaterial->GetName();
+			SaveAssetDialogConfig.ExistingAssetPolicy = ESaveAssetDialogExistingAssetPolicy::AllowButWarn;
+
+			FContentBrowserModule& ContentBrowserModule = FModuleManager::LoadModuleChecked<FContentBrowserModule>("ContentBrowser");
+			const FString SaveObjectPath = ContentBrowserModule.Get().CreateModalSaveAssetDialog(SaveAssetDialogConfig);
+
+			if (SaveObjectPath.Len() == 0)
+			{
+				return;
+			}
+
+			UDMBlueprintFunctionLibrary::ExportGeneratedMaterial(MaterialModel, SaveObjectPath);
+		}
+	}
+
+	void AddToolBarBoolOptionMenuEntry(FToolMenuSection& InSection, const FName& InPropertyName, const FUIAction InAction)
+	{
+		const FProperty* const OptionProperty = UDynamicMaterialEditorSettings::StaticClass()->FindPropertyByName(InPropertyName);
+
+		if (ensure(OptionProperty))
+		{
+			InSection.AddMenuEntry(NAME_None,
+				OptionProperty->GetDisplayNameText(),
+				OptionProperty->GetToolTipText(),
+				FSlateIcon(),
+				InAction, EUserInterfaceActionType::ToggleButton);
+		}
+	}
+
+	void AddToolBarIntOptionMenuEntry(FToolMenuSection& InSection, const FName& InPropertyName, 
+		TAttribute<bool> InIsEnabledAttribute = TAttribute<bool>(),
+		TAttribute<EVisibility> InVisibilityAttribute = TAttribute<EVisibility>())
+	{
+		InSection.AddDynamicEntry(NAME_None,
+			FNewToolMenuSectionDelegate::CreateLambda(
+				[InPropertyName, InIsEnabledAttribute, InVisibilityAttribute](FToolMenuSection& InSection)
+				{
+					const FProperty* const OptionProperty = UDynamicMaterialEditorSettings::StaticClass()->FindPropertyByName(InPropertyName);
+					FText DisplayName = FText::GetEmpty();
+					FText Tooltip = FText::GetEmpty();
+
+					if (ensure(OptionProperty))
+					{
+						DisplayName = OptionProperty->GetDisplayNameText();
+						Tooltip = OptionProperty->GetToolTipText();
+					}
+
+					FPropertyEditorModule& PropertyEditor = FModuleManager::Get().LoadModuleChecked<FPropertyEditorModule>(TEXT("PropertyEditor"));
+
+					FSinglePropertyParams SinglePropertyParams;
+					SinglePropertyParams.NamePlacement = EPropertyNamePlacement::Hidden;
+
+					TSharedRef<ISinglePropertyView> SinglePropertyView = PropertyEditor.CreateSingleProperty(UDynamicMaterialEditorSettings::Get(), InPropertyName, SinglePropertyParams).ToSharedRef();
+					SinglePropertyView->SetToolTipText(Tooltip);
+					SinglePropertyView->SetEnabled(InIsEnabledAttribute);
+					SinglePropertyView->SetVisibility(InVisibilityAttribute);
+
+					InSection.AddEntry(FToolMenuEntry::InitWidget(NAME_None,
+						SNew(SBox)
+						.HAlign(HAlign_Fill)
+						[
+							SNew(SBox)
+								.WidthOverride(80.0f)
+								.HAlign(HAlign_Right)
+								[
+									SinglePropertyView
+								]
+						],
+						DisplayName));
+				})
+		);
+	}
+
+	void AddToolbarExportMenu(UToolMenu* InMenu)
+	{
+		if (!IsValid(InMenu) || InMenu->ContainsSection(ToolBarExportSectionName))
+		{
+			return;
+		}
+
+		const UDMMenuContext* const MenuContext = InMenu->FindContext<UDMMenuContext>();
+
+		if (!MenuContext)
+		{
+			return;
+		}
+
+		UDynamicMaterialModel* const MaterialModel = MenuContext->GetModel();
+
+		if (!MaterialModel)
+		{
+			return;
+		}
+
+		UDynamicMaterialInstance* Instance = MaterialModel->GetDynamicMaterialInstance();
+
+		if (!Instance)
+		{
+			return;
+		}
+
+		const bool bAllowInstanceExport = IsValid(Instance->GetOuter()) && !Instance->GetOuter()->IsA<UPackage>();
+		const bool bAllowMaterialExport = IsValid(MaterialModel->GetGeneratedMaterial());
+
+		if (!bAllowInstanceExport && !bAllowMaterialExport)
+		{
+			return;
+		}
+
+		FToolMenuSection& NewSection = InMenu->AddSection("Export", LOCTEXT("ExportSection", "Export Material"));
+
+		if (bAllowInstanceExport)
+		{
+			NewSection.AddMenuEntry(NAME_None,
+				LOCTEXT("ExportMaterialInstance", "Export Material Designer Instance"),
+				LOCTEXT("ExportMaterialInstanceTooltip", "Export the material instance to an asset."),
+				FSlateIcon(),
+				FUIAction(FExecuteAction::CreateStatic(
+					&UE::DynamicMaterialEditor::Private::ExportMaterialInstanceFromInstance,
+					TWeakObjectPtr<UDynamicMaterialInstance>(Instance)
+				))
+			);
+		}
+
+		if (bAllowMaterialExport)
+		{
+			NewSection.AddMenuEntry(NAME_None,
+				LOCTEXT("ExportGeneratedMaterial", "Export Generated Material"),
+				LOCTEXT("ExportGeneratedMaterialTooltip", "Export the generated material to an asset."),
+				FSlateIcon(),
+				FUIAction(FExecuteAction::CreateStatic(
+					&UE::DynamicMaterialEditor::Private::ExportMaterialModelFromModel,
+					TWeakObjectPtr<UDynamicMaterialModel>(MaterialModel))
+				)
+			);
+		}
+	}
+
+	void AddToolBarTooltipOptionsSection(UToolMenu* InMenu)
+	{
+		if (!IsValid(InMenu) || InMenu->ContainsSection(ToolBarTooltipOptionsSectionName))
+		{
+			return;
+		}
+
+		FToolMenuSection& NewSection = InMenu->AddSection(ToolBarTooltipOptionsSectionName, LOCTEXT("TooltipOptionsSection", "Tooltip Options"));
+
+		AddToolBarBoolOptionMenuEntry(NewSection,
+			GET_MEMBER_NAME_CHECKED(UDynamicMaterialEditorSettings, bShowTooltipPreview),
+			FUIAction(
+				FExecuteAction::CreateLambda(
+					[]()
+					{
+						UDynamicMaterialEditorSettings* Settings = UDynamicMaterialEditorSettings::Get();
+						Settings->bShowTooltipPreview = !Settings->bShowTooltipPreview;
+						Settings->SaveConfig();
+					}),
+				FCanExecuteAction(),
+				FIsActionChecked::CreateLambda(
+					[]()
+					{
+						return UDynamicMaterialEditorSettings::Get()->bShowTooltipPreview;
+					})
+			)
+		);
+
+		TAttribute<bool> AreTooltipPreviewsEnabledAttribute = TAttribute<bool>::CreateLambda(
+			[]()
+			{
+				return UDynamicMaterialEditorSettings::Get()->bShowTooltipPreview;
+			});
+
+		AddToolBarIntOptionMenuEntry(
+			NewSection,
+			GET_MEMBER_NAME_CHECKED(UDynamicMaterialEditorSettings, TooltipTextureSize),
+			AreTooltipPreviewsEnabledAttribute
+		);
+	}
+
+	void AddToolBarPreviewOptionsSection(UToolMenu* InMenu)
+	{
+		if (!IsValid(InMenu) || InMenu->ContainsSection(ToolBarPreviewOptionsSectionName))
+		{
+			return;
+		}
+
+		FToolMenuSection& NewSection = InMenu->AddSection(ToolBarPreviewOptionsSectionName, LOCTEXT("PreviewOptionsSection", "Preview Options"));
+
+		AddToolBarIntOptionMenuEntry(NewSection,
+			GET_MEMBER_NAME_CHECKED(UDynamicMaterialEditorSettings, LayerPreviewSize)
+		);
+
+		NewSection.AddSeparator(NAME_None);
+
+		AddToolBarIntOptionMenuEntry(NewSection,
+			GET_MEMBER_NAME_CHECKED(UDynamicMaterialEditorSettings, DetailsPreviewSize)
+		);
+
+		NewSection.AddSeparator(NAME_None);
+	}
+
+	void AddToolBarAdvancedSection(UToolMenu* InMenu)
+	{
+		FToolMenuSection& NewSection = InMenu->AddSection("AdvancedSettings", LOCTEXT("AdvancedSettingsSection", "Advanced Settings"));
+
+		InMenu->AddDynamicSection(NAME_None, FNewToolMenuDelegate::CreateStatic(&AddToolBarPreviewOptionsSection));
+
+		InMenu->AddDynamicSection(NAME_None, FNewToolMenuDelegate::CreateStatic(&AddToolBarTooltipOptionsSection));
+
+		NewSection.AddSeparator(NAME_None);
+
+		NewSection.AddMenuEntry(NAME_None,
+			LOCTEXT("ResetAllSettingsToDefaults", "Reset All To Defaults"),
+			LOCTEXT("ResetAllSettingsToDefaultsTooltip", "Resets all the Material Designer settings to their default values."),
+			FSlateIcon(),
+			FUIAction(FExecuteAction::CreateUObject(
+				UDynamicMaterialEditorSettings::Get(), 
+				&UDynamicMaterialEditorSettings::ResetAllLayoutSettings)
+			)
+		);
+	}
+
+	void AddToolBarEditorLayoutMenu(UToolMenu* InMenu)
+	{
+		AddToolbarExportMenu(InMenu);
+
+		FToolMenuSection& NewSection = InMenu->AddSection("MaterialDesigner", LOCTEXT("MaterialDesignerSection", "MaterialDesigner"));
+
+		UDMMenuContext* MenuContext = InMenu->FindContext<UDMMenuContext>();
+
+		NewSection.AddMenuEntry(NAME_None,
+			LOCTEXT("OpenInUEMaterialEditor", "Open in UE Material Editor..."),
+			LOCTEXT("OpenInUEMaterialEditorTooltip", "Opens the currently editing generated Material Designer Instance material in the Unreal Engine material editor."),
+			FSlateIcon(),
+			FUIAction(FExecuteAction::CreateStatic(
+				&OpenMaterialEditorFromContext,
+				MenuContext
+			))
+		);
+
+		NewSection.AddSubMenu(
+			"AdvancedSettings",
+			LOCTEXT("AdvancedSettingsSubMenu", "Advanced Settings"),
+			LOCTEXT("AdvancedSettingsSubMenu_ToolTip", "Display advanced Material Designer settings"),
+			FNewToolMenuDelegate::CreateStatic(&AddToolBarAdvancedSection)
+		);
+
+		NewSection.AddMenuEntry(NAME_None,
+			LOCTEXT("OpenSettings", "Material Designer Settings..."),
+			LOCTEXT("OpenSettingsTooltip", "Opens the Material Designer settings window."),
+			FSlateIcon(FAppStyle::GetAppStyleSetName(), "FoliageEditMode.Settings"),
+			FUIAction(FExecuteAction::CreateUObject(
+				UDynamicMaterialEditorSettings::Get(),
+				&UDynamicMaterialEditorSettings::OpenEditorSettingsWindow)
+			)
+		);
+	}
+}
+
+TSharedRef<SWidget> FDMToolBarMenus::MakeEditorLayoutMenu(const TSharedPtr<SDMEditor>& InEditorWidget)
+{
+	using namespace UE::DynamicMaterialEditor::Private;
+
+	if (!UToolMenus::Get()->IsMenuRegistered(ToolBarEditorLayoutMenuName))
+	{
+		UToolMenu* const NewToolMenu = UDMMenuContext::GenerateContextMenuDefault(ToolBarEditorLayoutMenuName);
+
+		if (!NewToolMenu)
+		{
+			return SNullWidget::NullWidget;
+		}
+
+		FToolMenuSection& NewSection = NewToolMenu->AddDynamicSection(
+			"MaterialDesigner", 
+			FNewToolMenuDelegate::CreateStatic(&AddToolBarEditorLayoutMenu)
+		);
+	}
+
+	TSharedPtr<SDMSlot> SlotWidget;
+	if (InEditorWidget.IsValid())
+	{
+		TArray<TSharedPtr<SDMSlot>> SlotWidgets = InEditorWidget->GetSlotWidgets();
+
+		if (SlotWidgets.IsEmpty() == false)
+		{
+			SlotWidget = SlotWidgets[0];
+		}
+	}
+
+	FToolMenuContext MenuContext(FDynamicMaterialEditorModule::Get().GetCommandList(), TSharedPtr<FExtender>(), UDMMenuContext::CreateSlot(SlotWidget));
+
+	return UToolMenus::Get()->GenerateWidget(ToolBarEditorLayoutMenuName, MenuContext);
+}
+
+#undef LOCTEXT_NAMESPACE
diff --git a/Engine/Plugins/Experimental/DynamicMaterial/Source/DynamicMaterialEditor/Private/Menus/DMToolBarMenus.h b/Engine/Plugins/Experimental/DynamicMaterial/Source/DynamicMaterialEditor/Private/Menus/DMToolBarMenus.h
new file mode 100644
--- /dev/null
+++ b/Engine/Plugins/Experimental/DynamicMaterial/Source/DynamicMaterialEditor/Private/Menus/DMToolBarMenus.h
@@ -0,0 +1,27 @@
+// Copyright Epic Games, Inc. All Rights Reserved.
+
+#pragma once
+
+#include "Layout/Visibility.h"
+#include "Misc/Attribute.h"
+#include "Templates/SharedPointer.h"
+
+class SDMEditor;
+class SWidget;
+class UToolMenu;
+struct FToolMenuSection;
+struct FUIAction;
+
+class FDMToolBarMenus
+{
+public:
+	static TSharedRef<SWidget> MakeEditorLayoutMenu(const TSharedPtr<SDMEditor>& InEditorWidget = nullptr);
+
+	static void AddPreviewOptionsSection(UToolMenu* InMenu);
+	static void AddTooltipOptionsSection(UToolMenu* InMenu);
+	static void AddAdvancedSection(UToolMenu* InMenu);
+
+	/** Helper functions to add property menu entries. */
+	static void AddBoolOptionMenuEntry(FToolMenuSection& InSection, const FName& InPropertyName, const FUIAction InAction);
+	static void AddIntOptionMenuEntry(FToolMenuSection& InSection, const FName& InPropertyName, TAttribute<bool> InIsEnabledAttribute = TAttribute<bool>(), TAttribute<EVisibility> InVisibilityAttribute = TAttribute<EVisibility>());
+};
