[MaterialDesigner]

diff --git a/Engine/Plugins/Experimental/DynamicMaterial/Source/DynamicMaterialEditor/Private/Model/DynamicMaterialModelEditorOnlyData.cpp b/Engine/Plugins/Experimental/DynamicMaterial/Source/DynamicMaterialEditor/Private/Model/DynamicMaterialModelEditorOnlyData.cpp
new file mode 100644
--- /dev/null
+++ b/Engine/Plugins/Experimental/DynamicMaterial/Source/DynamicMaterialEditor/Private/Model/DynamicMaterialModelEditorOnlyData.cpp
@@ -0,0 +1,1433 @@
+// Copyright Epic Games, Inc. All Rights Reserved.
+
+#include "Model/DynamicMaterialModelEditorOnlyData.h"
+#include "AssetRegistry/AssetRegistryModule.h"
+#include "AssetToolsModule.h"
+#include "Components/DMMaterialParameter.h"
+#include "Components/DMMaterialProperty.h"
+#include "Components/DMMaterialSlot.h"
+#include "Components/DMMaterialStage.h"
+#include "Components/DMMaterialStageBlend.h"
+#include "Components/DMMaterialStageThroughputLayerBlend.h"
+#include "Components/DMMaterialSubStage.h"
+#include "Components/DMMaterialValue.h"
+#include "Components/DMTextureUV.h"
+#include "Components/MaterialProperties/DMMPAmbientOcclusion.h"
+#include "Components/MaterialProperties/DMMPAnisotropy.h"
+#include "Components/MaterialProperties/DMMPBaseColor.h"
+#include "Components/MaterialProperties/DMMPEmissiveColor.h"
+#include "Components/MaterialProperties/DMMPMetallic.h"
+#include "Components/MaterialProperties/DMMPNormal.h"
+#include "Components/MaterialProperties/DMMPOpacity.h"
+#include "Components/MaterialProperties/DMMPOpacityMask.h"
+#include "Components/MaterialProperties/DMMPPixelDepthOffset.h"
+#include "Components/MaterialProperties/DMMPRefraction.h"
+#include "Components/MaterialProperties/DMMPRoughness.h"
+#include "Components/MaterialProperties/DMMPSpecular.h"
+#include "Components/MaterialProperties/DMMPTangent.h"
+#include "Components/MaterialProperties/DMMPWorldPositionOffset.h"
+#include "Components/MaterialStageBlends/DMMSBNormal.h"
+#include "Components/MaterialStageExpressions/DMMSETextureSample.h"
+#include "Components/MaterialStageInputs/DMMSIExpression.h"
+#include "Components/MaterialStageInputs/DMMSIValue.h"
+#include "Components/MaterialValues/DMMaterialValueFloat1.h"
+#include "Components/MaterialValues/DMMaterialValueTexture.h"
+#include "CoreGlobals.h"
+#include "DMComponentPath.h"
+#include "DMDefs.h"
+#include "DMMaterialFunctionLibrary.h"
+#include "DynamicMaterialEditorModule.h"
+#include "DynamicMaterialEditorSettings.h"
+#include "DynamicMaterialModule.h"
+#include "Factories/MaterialFactoryNew.h"
+#include "FileHelpers.h"
+#include "IAssetTools.h"
+#include "Components/DMMaterialLayer.h"
+#include "Material/DynamicMaterialInstance.h"
+#include "Materials/Material.h"
+#include "Materials/MaterialExpressionAppendVector.h"
+#include "Materials/MaterialExpressionMaterialFunctionCall.h"
+#include "Materials/MaterialExpressionMax.h"
+#include "Materials/MaterialExpressionMultiply.h"
+#include "Materials/MaterialFunction.h"
+#include "Misc/Guid.h"
+#include "Model/DynamicMaterialModel.h"
+#include "Model/DMMaterialBuildState.h"
+#include "Model/DMMaterialBuildUtils.h"
+#include "UObject/Package.h"
+
+#define LOCTEXT_NAMESPACE "MaterialDesignerModel"
+
+const FString UDynamicMaterialModelEditorOnlyData::SlotsPathToken       = FString(TEXT("Slots"));
+const FString UDynamicMaterialModelEditorOnlyData::RGBSlotPathToken     = FString(TEXT("RGBSlot"));
+const FString UDynamicMaterialModelEditorOnlyData::OpacitySlotPathToken = FString(TEXT("OpacitySlot"));
+const FString UDynamicMaterialModelEditorOnlyData::PropertiesPathToken  = FString(TEXT("Properties"));
+
+UDynamicMaterialModelEditorOnlyData* UDynamicMaterialModelEditorOnlyData::Get(UDynamicMaterialModel* InModel)
+{
+	if (InModel)
+	{
+		return Get(InModel->GetEditorOnlyData());
+	}
+
+	return nullptr;
+}
+
+UDynamicMaterialModelEditorOnlyData* UDynamicMaterialModelEditorOnlyData::Get(TWeakObjectPtr<UDynamicMaterialModel> InModelWeak)
+{
+	return Get(InModelWeak.Get());
+}
+
+UDynamicMaterialModelEditorOnlyData* UDynamicMaterialModelEditorOnlyData::Get(const TScriptInterface<IDynamicMaterialModelEditorOnlyDataInterface>& InInterface)
+{
+	return Cast<UDynamicMaterialModelEditorOnlyData>(InInterface.GetObject());
+}
+
+UDynamicMaterialModelEditorOnlyData* UDynamicMaterialModelEditorOnlyData::Get(IDynamicMaterialModelEditorOnlyDataInterface* InInterface)
+{
+	return Cast<UDynamicMaterialModelEditorOnlyData>(InInterface);
+}
+
+UDynamicMaterialModelEditorOnlyData::UDynamicMaterialModelEditorOnlyData()
+{
+	State = EDMState::Idle;
+
+	BlendMode = EBlendMode::BLEND_Translucent;
+	ShadingModel = EDMMaterialShadingModel::Unlit;
+	bPixelAnimationFlag = false;
+
+	Properties.Emplace(EDMMaterialPropertyType::BaseColor,           CreateDefaultSubobject<UDMMaterialPropertyBaseColor>(          "MaterialProperty_BaseColor"));
+	Properties.Emplace(EDMMaterialPropertyType::EmissiveColor,       CreateDefaultSubobject<UDMMaterialPropertyEmissiveColor>(      "MaterialProperty_EmissiveColor"));
+	Properties.Emplace(EDMMaterialPropertyType::Opacity,             CreateDefaultSubobject<UDMMaterialPropertyOpacity>(            "MaterialProperty_Opacity"));
+	Properties.Emplace(EDMMaterialPropertyType::OpacityMask,         CreateDefaultSubobject<UDMMaterialPropertyOpacityMask>(        "MaterialProperty_OpacityMask"));
+	Properties.Emplace(EDMMaterialPropertyType::Metallic,            CreateDefaultSubobject<UDMMaterialPropertyMetallic>(           "MaterialProperty_Metallic"));
+	Properties.Emplace(EDMMaterialPropertyType::Specular,            CreateDefaultSubobject<UDMMaterialPropertySpecular>(           "MaterialProperty_Specular"));
+	Properties.Emplace(EDMMaterialPropertyType::Roughness,           CreateDefaultSubobject<UDMMaterialPropertyRoughness>(          "MaterialProperty_Roughness"));
+	Properties.Emplace(EDMMaterialPropertyType::Anisotropy,          CreateDefaultSubobject<UDMMaterialPropertyAnisotropy>(         "MaterialProperty_Anisotropy"));
+	Properties.Emplace(EDMMaterialPropertyType::Normal,              CreateDefaultSubobject<UDMMaterialPropertyNormal>(             "MaterialProperty_Normal"));
+	Properties.Emplace(EDMMaterialPropertyType::Tangent,             CreateDefaultSubobject<UDMMaterialPropertyTangent>(            "MaterialProperty_Tangent"));
+	Properties.Emplace(EDMMaterialPropertyType::WorldPositionOffset, CreateDefaultSubobject<UDMMaterialPropertyWorldPositionOffset>("MaterialProperty_WorldPositionOffset"));
+	Properties.Emplace(EDMMaterialPropertyType::Refraction,          CreateDefaultSubobject<UDMMaterialPropertyRefraction>(         "MaterialProperty_Refraction"));
+	Properties.Emplace(EDMMaterialPropertyType::AmbientOcclusion,    CreateDefaultSubobject<UDMMaterialPropertyAmbientOcclusion>(   "MaterialProperty_AmbientOcclusion"));
+	Properties.Emplace(EDMMaterialPropertyType::PixelDepthOffset,    CreateDefaultSubobject<UDMMaterialPropertyPixelDepthOffset>(   "MaterialProperty_PixelDepthOffset"));
+
+	Properties.Emplace(EDMMaterialPropertyType::Custom1, UDMMaterialProperty::CreateCustomMaterialPropertyDefaultSubobject(this, EDMMaterialPropertyType::Custom1, "MaterialProperty_Custom1"));
+	Properties.Emplace(EDMMaterialPropertyType::Custom2, UDMMaterialProperty::CreateCustomMaterialPropertyDefaultSubobject(this, EDMMaterialPropertyType::Custom2, "MaterialProperty_Custom2"));
+	Properties.Emplace(EDMMaterialPropertyType::Custom3, UDMMaterialProperty::CreateCustomMaterialPropertyDefaultSubobject(this, EDMMaterialPropertyType::Custom3, "MaterialProperty_Custom3"));
+	Properties.Emplace(EDMMaterialPropertyType::Custom4, UDMMaterialProperty::CreateCustomMaterialPropertyDefaultSubobject(this, EDMMaterialPropertyType::Custom4, "MaterialProperty_Custom4"));
+
+	for (const TPair<EDMMaterialPropertyType, TObjectPtr<UDMMaterialProperty>>& Pair : Properties)
+	{
+		Pair.Value->SetComponentState(EDMComponentLifetimeState::Added);
+	}
+}
+
+UMaterial* UDynamicMaterialModelEditorOnlyData::GetGeneratedMaterial() const
+{
+	return IsValid(MaterialModel) ? MaterialModel->DynamicMaterial : nullptr;
+}
+
+void UDynamicMaterialModelEditorOnlyData::ResetData()
+{
+	if (GUndo)
+	{
+		Modify();
+	}
+
+	BlendMode = EBlendMode::BLEND_Translucent;
+	ShadingModel = EDMMaterialShadingModel::Unlit;
+	Slots.Empty();
+	Expressions.Empty();
+}
+
+void UDynamicMaterialModelEditorOnlyData::CreateMaterial()
+{
+	if (!IsValid(MaterialModel))
+	{
+		return;
+	}
+
+	if (FDynamicMaterialModule::IsMaterialExportEnabled() == false)
+	{
+		UMaterialFactoryNew* MaterialFactory = NewObject<UMaterialFactoryNew>();
+		check(MaterialFactory);
+
+		MaterialModel->DynamicMaterial = Cast<UMaterial>(MaterialFactory->FactoryCreateNew(
+			UMaterial::StaticClass(),
+			this,
+			NAME_None,
+			RF_DuplicateTransient | RF_TextExportTransient,
+			nullptr,
+			GWarn
+		));
+	}
+	else
+	{
+		FString MaterialBaseName = GetName() + "-" + FGuid::NewGuid().ToString();
+		const FString FullName = "/Game/DynamicMaterials/" + MaterialBaseName;
+		UPackage* Package = CreatePackage(*FullName);
+
+		UMaterialFactoryNew* MaterialFactory = NewObject<UMaterialFactoryNew>();
+		check(MaterialFactory);
+
+		MaterialModel->DynamicMaterial = Cast<UMaterial>(MaterialFactory->FactoryCreateNew(
+			UMaterial::StaticClass(),
+			Package,
+			*MaterialBaseName,
+			RF_DuplicateTransient | RF_TextExportTransient | RF_Standalone | RF_Public,
+			nullptr,
+			GWarn
+		));
+
+		FAssetRegistryModule::AssetCreated(MaterialModel->DynamicMaterial);
+		Package->FullyLoad();
+	}
+
+	MaterialModel->DynamicMaterial->bOutputTranslucentVelocity = true;
+	MaterialModel->DynamicMaterial->bEnableResponsiveAA = true;
+}
+
+void UDynamicMaterialModelEditorOnlyData::BuildMaterial(bool bInDirtyAssets)
+{
+	if (State != EDMState::Idle)
+	{
+		checkNoEntry();
+		return;
+	}
+
+	if (!IsValid(MaterialModel))
+	{
+		return;
+	}
+
+	UE_LOG(LogDynamicMaterialEditor, Display, TEXT("Building Material Designer Material (%s)..."), *MaterialModel->GetFullName());
+
+	if (!IsValid(MaterialModel->DynamicMaterial))
+	{
+		CreateMaterial();
+	}
+
+	State = EDMState::Building;
+	Expressions.Empty();
+	MaterialModel->DynamicMaterial->MaterialDomain = Domain;
+	MaterialModel->DynamicMaterial->BlendMode = BlendMode;
+	MaterialModel->DynamicMaterial->bHasPixelAnimation = bPixelAnimationFlag;
+
+	switch (ShadingModel)
+	{
+		case EDMMaterialShadingModel::DefaultLit:
+			MaterialModel->DynamicMaterial->SetShadingModel(EMaterialShadingModel::MSM_DefaultLit);
+			break;
+
+		case EDMMaterialShadingModel::Unlit:
+			MaterialModel->DynamicMaterial->SetShadingModel(EMaterialShadingModel::MSM_Unlit);
+			break;
+
+		default:
+			checkNoEntry();
+			break;
+	}
+
+	TSharedRef<FDMMaterialBuildState> BuildState = CreateBuildState(MaterialModel->DynamicMaterial, bInDirtyAssets);
+
+	const bool bIsPostProcessMaterial = Domain == EMaterialDomain::MD_PostProcess;
+
+	/**
+	 * Process slots to build base material inputs.
+	 */
+	for (const TPair<EDMMaterialPropertyType, TObjectPtr<UDMMaterialProperty>>& Pair : Properties)
+	{
+		if (UE::DynamicMaterialEditor::IsCustomMaterialProperty(Pair.Key))
+		{
+			continue;
+		}
+
+		if (bIsPostProcessMaterial && Pair.Key != EDMMaterialPropertyType::EmissiveColor)
+		{
+			continue;
+		}
+
+		// For now we don't have channel remapping!
+		FExpressionInput* MaterialPropertyPtr = BuildState->GetMaterialProperty(Pair.Key);
+
+		if (!MaterialPropertyPtr)
+		{
+			continue;
+		}
+
+		MaterialPropertyPtr->Expression = nullptr;
+		MaterialPropertyPtr->OutputIndex = 0;
+
+		UMaterialExpression* LastPropertyExpression = nullptr;
+		UDMMaterialSlot* Slot = GetSlotForMaterialProperty(Pair.Key);
+
+		if (Slot && !Slot->GetLayers().IsEmpty())
+		{
+			Slot->GenerateExpressions(BuildState);
+
+			if (BuildState->GetSlotExpressions(Slot).IsEmpty())
+			{
+				continue;
+			}
+
+			LastPropertyExpression = BuildState->GetLastSlotPropertyExpression(Slot, Pair.Key);
+		}
+
+		if (!LastPropertyExpression)
+		{
+			continue;
+		}
+
+		MaterialPropertyPtr->Expression = LastPropertyExpression;
+
+		if (Pair.Value->GetInputConnectionMap().Channels.IsEmpty() == false)
+		{
+			MaterialPropertyPtr->OutputIndex = Pair.Value->GetInputConnectionMap().Channels[0].OutputIndex;
+		}
+		else
+		{
+			MaterialPropertyPtr->OutputIndex = 0;
+		}
+	}
+
+	if (!bIsPostProcessMaterial)
+	{
+		/**
+		 * Generate opacity input based on base/emissive if it doesn't already have an input.
+		 */
+		EDMMaterialPropertyType GenerateOpacityInput = EDMMaterialPropertyType::None;
+
+		if (BlendMode == EBlendMode::BLEND_Translucent)
+		{
+			if (GetSlotForMaterialProperty(EDMMaterialPropertyType::Opacity) == nullptr)
+			{
+				GenerateOpacityInput = EDMMaterialPropertyType::Opacity;
+			}
+		}
+		else if (BlendMode == EBlendMode::BLEND_Masked)
+		{
+			if (GetSlotForMaterialProperty(EDMMaterialPropertyType::OpacityMask) == nullptr)
+			{
+				GenerateOpacityInput = EDMMaterialPropertyType::OpacityMask;
+			}
+		}
+
+		if (GenerateOpacityInput != EDMMaterialPropertyType::None)
+		{
+			UDMMaterialSlot* OpacitySlot = nullptr;
+			EDMMaterialPropertyType OpacityProperty = EDMMaterialPropertyType::None;
+
+			if (UDMMaterialSlot* BaseColorSlot = GetSlotForMaterialProperty(EDMMaterialPropertyType::BaseColor))
+			{
+				OpacitySlot = BaseColorSlot;
+				OpacityProperty = EDMMaterialPropertyType::BaseColor;
+			}
+			else if (UDMMaterialSlot* EmissiveSlot = GetSlotForMaterialProperty(EDMMaterialPropertyType::EmissiveColor))
+			{
+				OpacitySlot = EmissiveSlot;
+				OpacityProperty = EDMMaterialPropertyType::EmissiveColor;
+			}
+
+			if (OpacitySlot)
+			{
+				UMaterialExpression* OpacityOutputNode;
+				int32 OutputIndex;
+				int32 OutputChannel;
+				GenerateOpacityExpressions(BuildState, OpacitySlot, OpacityProperty, OpacityOutputNode, OutputIndex, OutputChannel);
+
+				if (OpacityOutputNode)
+				{
+					if (FExpressionInput* OpacityPropertyPtr = BuildState->GetMaterialProperty(GenerateOpacityInput))
+					{
+						OpacityPropertyPtr->Expression = OpacityOutputNode;
+						OpacityPropertyPtr->OutputIndex = 0;
+						OpacityPropertyPtr->Mask = 0;
+
+						if (OutputChannel != FDMMaterialStageConnectorChannel::WHOLE_CHANNEL)
+						{
+							OpacityPropertyPtr->Mask = 1;
+							OpacityPropertyPtr->MaskR = !!(OutputChannel & FDMMaterialStageConnectorChannel::FIRST_CHANNEL);
+							OpacityPropertyPtr->MaskG = !!(OutputChannel & FDMMaterialStageConnectorChannel::SECOND_CHANNEL);
+							OpacityPropertyPtr->MaskB = !!(OutputChannel & FDMMaterialStageConnectorChannel::THIRD_CHANNEL);
+							OpacityPropertyPtr->MaskA = !!(OutputChannel & FDMMaterialStageConnectorChannel::FOURTH_CHANNEL);
+						}
+					}
+				}
+			}
+		}
+
+		/**
+		 * Apply global opacity slider
+		 */
+		FExpressionInput* OpacityPropertyPtr = nullptr;
+
+		if (BlendMode == BLEND_Translucent)
+		{
+			OpacityPropertyPtr = BuildState->GetMaterialProperty(EDMMaterialPropertyType::Opacity);
+		}
+		else if (BlendMode == BLEND_Masked)
+		{
+			OpacityPropertyPtr = BuildState->GetMaterialProperty(EDMMaterialPropertyType::OpacityMask);
+		}
+
+		if (OpacityPropertyPtr != nullptr && IsValid(MaterialModel) && IsValid(MaterialModel->GlobalOpacityValue))
+		{
+			MaterialModel->GlobalOpacityValue->GenerateExpression(BuildState);
+
+			UMaterialExpression* GlobalOpacityExpression = BuildState->GetLastValueExpression(MaterialModel->GlobalOpacityValue);
+
+			if (OpacityPropertyPtr->Expression)
+			{
+				UMaterialExpressionMultiply* OpacityMultiply = BuildState->GetBuildUtils().CreateExpression<UMaterialExpressionMultiply>(UE_DM_NodeComment_Default);
+				OpacityMultiply->A.Expression = OpacityPropertyPtr->Expression;
+				OpacityMultiply->A.Mask = OpacityPropertyPtr->Mask;
+				OpacityMultiply->A.MaskR = OpacityPropertyPtr->MaskR;
+				OpacityMultiply->A.MaskG = OpacityPropertyPtr->MaskG;
+				OpacityMultiply->A.MaskB = OpacityPropertyPtr->MaskB;
+				OpacityMultiply->A.MaskA = OpacityPropertyPtr->MaskA;
+				OpacityMultiply->A.OutputIndex = OpacityPropertyPtr->OutputIndex;
+
+				OpacityMultiply->B.Expression = GlobalOpacityExpression;
+				OpacityMultiply->B.SetMask(1, 1, 0, 0, 0);
+				OpacityMultiply->B.OutputIndex = 0;
+
+				OpacityPropertyPtr->Expression = OpacityMultiply;
+				OpacityPropertyPtr->SetMask(1, 1, 0, 0, 0);
+			}
+		}
+	}
+
+	/**
+	 * Apply output processors.
+	 */
+	for (const TPair<EDMMaterialPropertyType, TObjectPtr<UDMMaterialProperty>>& Pair : Properties)
+	{
+		if (UE::DynamicMaterialEditor::IsCustomMaterialProperty(Pair.Key))
+		{
+			continue;
+		}
+
+		UMaterialFunctionInterface* OutputProcessor = Pair.Value->GetOutputProcessor();
+
+		if (!OutputProcessor)
+		{
+			continue;
+		}
+
+		FExpressionInput* MaterialPropertyPtr = BuildState->GetMaterialProperty(Pair.Key);
+
+		if (!MaterialPropertyPtr)
+		{
+			continue;
+		}
+
+		UMaterialExpressionMaterialFunctionCall* MFC = FDMMaterialFunctionLibrary::Get().MakeExpression(
+			BuildState->GetDynamicMaterial(),
+			OutputProcessor,
+			UE_DM_NodeComment_Default
+		);
+
+		TArrayView<FExpressionInput*> Inputs = MFC->GetInputsView();
+
+		if (Inputs.IsEmpty())
+		{
+			continue;
+		}
+
+		UMaterialExpression* LastPropertyExpression = MaterialPropertyPtr->Expression;
+
+		if (!LastPropertyExpression)
+		{
+			LastPropertyExpression = Pair.Value->GetDefaultInput(BuildState);
+
+			if (!LastPropertyExpression)
+			{
+				continue;
+			}
+		}
+
+		LastPropertyExpression->ConnectExpression(Inputs[0], MaterialPropertyPtr->OutputIndex);
+		MFC->ConnectExpression(MaterialPropertyPtr, 0);
+
+		MaterialPropertyPtr->OutputIndex = 0;
+	}
+
+	State = EDMState::Idle;
+
+	if (IsValid(MaterialModel) && IsValid(MaterialModel->DynamicMaterialInstance))
+	{
+		MaterialModel->DynamicMaterialInstance->OnMaterialBuilt(MaterialModel);
+	}
+
+	OnMaterialBuiltDelegate.Broadcast(MaterialModel);
+}
+
+void UDynamicMaterialModelEditorOnlyData::RequestMaterialBuild()
+{
+	if (!HasAnyFlags(RF_ClassDefaultObject | RF_ArchetypeObject))
+	{
+		FDynamicMaterialEditorModule::Get().AddBuildRequest(this, /* Dirty Packages */ !GIsEditorLoadingPackage);
+	}
+}
+
+void UDynamicMaterialModelEditorOnlyData::OnValueListUpdate()
+{
+	OnValueListUpdateDelegate.Broadcast(MaterialModel);
+}
+
+void UDynamicMaterialModelEditorOnlyData::LoadDeprecatedModelData_Base(bool bInCreateMaterialPackage, EBlendMode InBlendMode, 
+	EDMMaterialShadingModel InShadingModel)
+{
+	bCreateMaterialPackage = bInCreateMaterialPackage;
+
+	if (IsValid(MaterialModel->DynamicMaterial))
+	{
+		BlendMode = MaterialModel->DynamicMaterial->BlendMode;
+
+		ShadingModel = MaterialModel->DynamicMaterial->GetShadingModels().HasShadingModel(EMaterialShadingModel::MSM_DefaultLit)
+			? EDMMaterialShadingModel::DefaultLit
+			: EDMMaterialShadingModel::Unlit;
+	}
+	else if (IsValid(MaterialModel->DynamicMaterialInstance))
+	{
+		BlendMode = MaterialModel->DynamicMaterialInstance->BlendMode;
+
+		ShadingModel = MaterialModel->DynamicMaterialInstance->GetShadingModels().HasShadingModel(EMaterialShadingModel::MSM_DefaultLit)
+			? EDMMaterialShadingModel::DefaultLit
+			: EDMMaterialShadingModel::Unlit;
+	}
+	else
+	{
+		BlendMode = InBlendMode;
+		ShadingModel = InShadingModel;
+	}
+}
+
+void UDynamicMaterialModelEditorOnlyData::LoadDeprecatedModelData_Expressions(TArray<TObjectPtr<UMaterialExpression>>& InExpressions)
+{
+	// Expressions should be parented to the material.
+	Expressions = InExpressions;
+	InExpressions.Empty();
+}
+
+void UDynamicMaterialModelEditorOnlyData::LoadDeprecatedModelData_Properties(TMap<EDMMaterialPropertyType, TObjectPtr<UObject>>& InProperties)
+{
+	for (const TPair<EDMMaterialPropertyType, TObjectPtr<UObject>>& Pair : InProperties)
+	{
+		UDMMaterialProperty* InProperty = Cast<UDMMaterialProperty>(Pair.Value);
+
+		if (!InProperty)
+		{
+			continue;
+		}
+
+		if (const TObjectPtr<UDMMaterialProperty>* ThisPropertyPtr = Properties.Find(Pair.Key))
+		{
+			if (UDMMaterialProperty* ThisProperty = *ThisPropertyPtr)
+			{
+				ThisProperty->LoadDeprecatedModelData(InProperty);
+			}
+		}
+	}
+
+	InProperties.Empty();
+}
+
+void UDynamicMaterialModelEditorOnlyData::LoadDeprecatedModelData_Slots(TArray<TObjectPtr<UObject>>& InSlots)
+{
+	Slots.Empty();
+	Slots.Reserve(InSlots.Num());
+
+	for (TObjectPtr<UObject>& SlotObject : InSlots)
+	{
+		UDMMaterialSlot* Slot = Cast<UDMMaterialSlot>(SlotObject);
+
+		if (!Slot)
+		{
+			continue;
+		}
+
+		Slots.Add(Slot);
+
+		// Reparent from the model to the editor only data
+		Slot->Rename(nullptr, this, UE::DynamicMaterial::RenameFlags);
+	}
+
+	InSlots.Empty();
+}
+
+void UDynamicMaterialModelEditorOnlyData::LoadDeprecatedModelData_PropertySlotMap(TMap<EDMMaterialPropertyType, TObjectPtr<UObject>>& InPropertySlotMap)
+{
+	PropertySlotMap.Empty();
+	PropertySlotMap.Reserve(InPropertySlotMap.Num());
+
+	for (const TPair<EDMMaterialPropertyType, TObjectPtr<UObject>>& Pair : InPropertySlotMap)
+	{
+		UDMMaterialSlot* Slot = Cast<UDMMaterialSlot>(Pair.Value);
+
+		if (!Slot)
+		{
+			continue;
+		}
+
+		PropertySlotMap.Add(Pair.Key, Slot);
+	}
+
+	InPropertySlotMap.Empty();
+}
+
+void UDynamicMaterialModelEditorOnlyData::LoadDeprecatedModelData(UDynamicMaterialModel* InMaterialModel)
+{
+	MaterialModel = InMaterialModel;
+	check(InMaterialModel);
+
+	PRAGMA_DISABLE_DEPRECATION_WARNINGS
+
+	LoadDeprecatedModelData_Base(MaterialModel->bCreateMaterialPackage, MaterialModel->BlendMode, MaterialModel->ShadingModel);
+	LoadDeprecatedModelData_Expressions(MaterialModel->Expressions);
+	LoadDeprecatedModelData_Properties(InMaterialModel->Properties);
+	LoadDeprecatedModelData_Slots(InMaterialModel->Slots);
+	LoadDeprecatedModelData_PropertySlotMap(InMaterialModel->PropertySlotMap);
+
+	PRAGMA_ENABLE_DEPRECATION_WARNINGS
+}
+
+TSharedRef<FDMMaterialBuildState> UDynamicMaterialModelEditorOnlyData::CreateBuildState(UMaterial* InMaterialToBuild, bool bInDirtyAssets) const
+{
+	check(MaterialModel);
+	check(InMaterialToBuild);
+
+	return MakeShared<FDMMaterialBuildState>(InMaterialToBuild, MaterialModel, bInDirtyAssets);
+}
+
+UDMMaterialComponent* UDynamicMaterialModelEditorOnlyData::GetSubComponentByPath(FDMComponentPath& InPath,
+	const FDMComponentPathSegment& InPathSegment) const
+{
+	if (InPathSegment.GetToken() == SlotsPathToken)
+	{
+		int32 SlotIndex;
+
+		if (InPathSegment.GetParameter(SlotIndex))
+		{
+			if (Slots.IsValidIndex(SlotIndex))
+			{
+				return Slots[SlotIndex]->GetComponentByPath(InPath);
+			}
+		}
+
+		return nullptr;
+	}
+
+	// @TODO This will probably need to change when other slot types are opened up.
+	if (InPathSegment.GetToken() == RGBSlotPathToken)
+	{
+		if (Slots.IsValidIndex(0))
+		{
+			return Slots[0]->GetComponentByPath(InPath);
+		}
+
+		return nullptr;
+	}
+
+	// @TODO This will probably need to change when other slot types are opened up.
+	if (InPathSegment.GetToken() == OpacitySlotPathToken)
+	{
+		if (Slots.IsValidIndex(1))
+		{
+			return Slots[1]->GetComponentByPath(InPath);
+		}
+
+		return nullptr;
+	}
+
+	if (InPathSegment.GetToken() == PropertiesPathToken)
+	{
+		FString PropertyStr;
+
+		if (InPathSegment.GetParameter(PropertyStr))
+		{
+			UEnum* PropertyEnum = StaticEnum<EDMMaterialPropertyType>();
+			int64 IntValue = PropertyEnum->GetValueByNameString(PropertyStr);
+
+			if (IntValue != INDEX_NONE)
+			{
+				EDMMaterialPropertyType EnumValue = static_cast<EDMMaterialPropertyType>(IntValue);
+
+				if (const TObjectPtr<UDMMaterialProperty>* PropertyPtr = Properties.Find(EnumValue))
+				{
+					return (*PropertyPtr)->GetComponentByPath(InPath);
+				}
+			}
+		}
+	}
+
+	return nullptr;
+}
+
+TSharedRef<IDMMaterialBuildStateInterface> UDynamicMaterialModelEditorOnlyData::CreateBuildStateInterface(UMaterial* InMaterialToBuild) const
+{
+	return CreateBuildState(InMaterialToBuild);
+}
+
+void UDynamicMaterialModelEditorOnlyData::DoBuild_Implementation(bool bInDirtyAssets)
+{
+	BuildMaterial(bInDirtyAssets);
+}
+
+void UDynamicMaterialModelEditorOnlyData::SetDomain(TEnumAsByte<EMaterialDomain> InDomain)
+{
+	if (Domain == InDomain)
+	{
+		return;
+	}
+
+	Domain = InDomain;
+
+	if (Domain == EMaterialDomain::MD_PostProcess)
+	{
+		// Post process domain only supports a single slot - emissive color. 1 is the opacity slot.
+		if (Slots.IsValidIndex(1))
+		{
+			RemoveSlot(1);
+		}
+
+		SetShadingModel(EDMMaterialShadingModel::Unlit);
+		SetBlendMode(EBlendMode::BLEND_Opaque);
+
+		// Setting it to unlit will set these to base color. Need to reset them to emissive.
+		if (Slots.IsValidIndex(0))
+		{
+			for (UDMMaterialLayerObject* Layer : Slots[0]->GetLayers())
+			{
+				if (GUndo)
+				{
+					Layer->Modify();
+				}
+
+				Layer->SetMaterialProperty(EDMMaterialPropertyType::EmissiveColor);
+			}
+		}
+	}
+
+	RequestMaterialBuild();
+}
+
+void UDynamicMaterialModelEditorOnlyData::SetBlendMode(TEnumAsByte<EBlendMode> InBlendMode)
+{
+	if (BlendMode == InBlendMode)
+	{
+		return;
+	}
+
+	BlendMode = InBlendMode;
+
+	if (Slots.IsValidIndex(1))
+	{
+		// 1 is the opacity slot
+		if (InBlendMode == EBlendMode::BLEND_Opaque)
+		{
+			RemoveSlot(1);
+		}
+		else
+		{
+			for (UDMMaterialLayerObject* Layer : Slots[0]->GetLayers())
+			{
+				if (GUndo)
+				{
+					Layer->Modify();
+				}
+
+				switch (InBlendMode)
+				{
+					case EBlendMode::BLEND_Masked:
+						Layer->SetMaterialProperty(EDMMaterialPropertyType::OpacityMask);
+						break;
+
+					default:
+						Layer->SetMaterialProperty(EDMMaterialPropertyType::Opacity);
+						break;
+				}
+			}
+		}
+	}
+
+	RequestMaterialBuild();
+}
+
+void UDynamicMaterialModelEditorOnlyData::SetShadingModel(EDMMaterialShadingModel InShadingModel)
+{
+	if (ShadingModel == InShadingModel)
+	{
+		return;
+	}
+
+	ShadingModel = InShadingModel;
+
+	if (Slots.IsValidIndex(0))
+	{
+		for (UDMMaterialLayerObject* Layer : Slots[0]->GetLayers())
+		{
+			if (GUndo)
+			{
+				Layer->Modify();
+			}
+
+			switch (ShadingModel)
+			{
+				case EDMMaterialShadingModel::Unlit:
+					Layer->SetMaterialProperty(EDMMaterialPropertyType::EmissiveColor);
+					break;
+
+				default:
+					Layer->SetMaterialProperty(EDMMaterialPropertyType::BaseColor);
+					break;
+			}
+		}
+	}
+
+	RequestMaterialBuild();
+}
+
+void UDynamicMaterialModelEditorOnlyData::OpenMaterialEditor() const
+{
+	if (!IsValid(MaterialModel) || !IsValid(MaterialModel->DynamicMaterial))
+	{
+		return;
+	}
+
+	IAssetTools& AssetTools = FAssetToolsModule::GetModule().Get();
+	AssetTools.OpenEditorForAssets({MaterialModel->DynamicMaterial});
+}
+
+TMap<EDMMaterialPropertyType, UDMMaterialProperty*> UDynamicMaterialModelEditorOnlyData::GetMaterialProperties() const
+{
+	TMap<EDMMaterialPropertyType, UDMMaterialProperty*> LocalProperties;
+
+	LocalProperties.Reserve(Properties.Num());
+	for (const TPair<EDMMaterialPropertyType, TObjectPtr<UDMMaterialProperty>>& Pair : Properties)
+	{
+		LocalProperties.Add(Pair.Key, Pair.Value.Get());
+	}
+
+	return LocalProperties;
+}
+
+UDMMaterialProperty* UDynamicMaterialModelEditorOnlyData::GetMaterialProperty(EDMMaterialPropertyType MaterialProperty) const
+{
+	TObjectPtr<UDMMaterialProperty> const* PropertyObjPtr = Properties.Find(MaterialProperty);
+
+	if (PropertyObjPtr)
+	{
+		return *PropertyObjPtr;
+	}
+
+	return nullptr;
+}
+
+UDMMaterialSlot* UDynamicMaterialModelEditorOnlyData::GetSlot(int32 Index) const
+{
+	if (!Slots.IsValidIndex(Index))
+	{
+		return nullptr;
+	}
+
+	return Slots[Index];
+}
+
+UDMMaterialSlot* UDynamicMaterialModelEditorOnlyData::AddSlot()
+{
+	EDMMaterialPropertyType NewSlotProperty = EDMMaterialPropertyType::None;
+
+	if (Slots.IsEmpty())
+	{
+		const FDMUpdateGuard Guard;
+
+		switch (ShadingModel)
+		{
+			case EDMMaterialShadingModel::DefaultLit:
+				NewSlotProperty = EDMMaterialPropertyType::BaseColor;
+				break;
+
+			case EDMMaterialShadingModel::Unlit:
+				NewSlotProperty = EDMMaterialPropertyType::EmissiveColor;
+				break;
+
+			default:
+				checkNoEntry()
+					break;
+		}
+	}
+	else if (Domain == EMaterialDomain::MD_PostProcess)
+	{
+		return nullptr;
+	}
+	else
+	{
+		for (int64 FirstAvailableEnum = static_cast<int64>(EDMMaterialPropertyType::EmissiveColor) + 1; FirstAvailableEnum < static_cast<int64>(EDMMaterialPropertyType::Any); ++FirstAvailableEnum)
+		{
+			EDMMaterialPropertyType FirstAvailableEnumActual = static_cast<EDMMaterialPropertyType>(FirstAvailableEnum);
+
+			if (GetSlotForMaterialProperty(FirstAvailableEnumActual) == nullptr)
+			{
+				NewSlotProperty = FirstAvailableEnumActual;
+				break;
+			}
+		}
+
+		if (NewSlotProperty == EDMMaterialPropertyType::None)
+		{
+			return nullptr;
+		}
+	}
+
+	UDMMaterialSlot* NewSlot = NewObject<UDMMaterialSlot>(this, NAME_None, RF_Transactional);
+	check(NewSlot);
+
+	NewSlot->SetIndex(Slots.Num());
+	Slots.Add(NewSlot);
+	NewSlot->SetComponentState(EDMComponentLifetimeState::Added);
+
+	NewSlot->GetOnConnectorsUpdateDelegate().AddUObject(this, &UDynamicMaterialModelEditorOnlyData::OnSlotConnectorsUpdated);
+	NewSlot->Update(EDMUpdateType::Structure);
+
+	UDMMaterialStage* DefaultStage = UDMMaterialStageBlend::CreateStage(UDMMaterialStageBlendNormal::StaticClass());
+	check(DefaultStage);
+
+	DefaultStage->SetBeingEdited(true);
+
+	UDMMaterialStage* MaskStage = UDMMaterialStageThroughputLayerBlend::CreateStage();
+	check(MaskStage);
+
+	NewSlot->AddLayerWithMask(NewSlotProperty, DefaultStage, MaskStage);
+	NewSlot->SetEditingLayers(true);
+
+	AssignMaterialPropertyToSlot(NewSlotProperty, NewSlot);
+
+	UDMMaterialStageInputExpression* BaseInputExpression = DefaultStage->ChangeInput_Expression(UDMMaterialStageBlendNormal::InputB, FDMMaterialStageConnectorChannel::WHOLE_CHANNEL,
+		UDMMaterialStageExpressionTextureSample::StaticClass(), 0, FDMMaterialStageConnectorChannel::WHOLE_CHANNEL);
+
+	switch (NewSlotProperty)
+	{
+		case EDMMaterialPropertyType::Opacity:
+		case EDMMaterialPropertyType::OpacityMask:
+			if (UDMMaterialStageExpressionTextureSample* BaseInputTextureSample = Cast<UDMMaterialStageExpressionTextureSample>(BaseInputExpression->GetMaterialStageExpression()))
+			{
+				BaseInputTextureSample->SetClampTextureEnabled(true);
+
+				if (UDMMaterialStage* BaseTextureInputStage = BaseInputExpression->GetSubStage())
+				{
+					const TArray<UDMMaterialStageInput*> BaseTextureStageInputs = BaseTextureInputStage->GetInputs();
+
+					for (UDMMaterialStageInput* BaseTextureStageInput : BaseTextureStageInputs)
+					{
+						if (UDMMaterialStageInputValue* BaseTextureInputValue = Cast<UDMMaterialStageInputValue>(BaseTextureStageInput))
+						{
+							if (UDMMaterialValueTexture* BaseTextureValue = Cast<UDMMaterialValueTexture>(BaseTextureInputValue->GetValue()))
+							{
+								BaseTextureValue->SetDefaultValue(UDynamicMaterialEditorSettings::Get()->DefaultOpacitySlotMask.LoadSynchronous());
+								BaseTextureValue->ApplyDefaultValue();
+								break;
+							}
+						}
+					}
+				}
+			}
+
+			{
+				const TArray<UDMMaterialStageInput*> MaskStageInputs = MaskStage->GetInputs();
+
+				for (UDMMaterialStageInput* MaskStageInput : MaskStageInputs)
+				{
+					if (UDMMaterialStageInputExpression* MaskInputExpression = Cast<UDMMaterialStageInputExpression>(MaskStageInput))
+					{
+						if (UDMMaterialStageExpressionTextureSample* MaskInputTextureSample = Cast<UDMMaterialStageExpressionTextureSample>(MaskInputExpression->GetMaterialStageExpression()))
+						{
+							MaskInputTextureSample->SetClampTextureEnabled(true);
+
+							if (UDMMaterialStage* MaskTextureInputStage = MaskInputExpression->GetSubStage())
+							{
+								const TArray<UDMMaterialStageInput*> MaskTextureStageInputs = MaskTextureInputStage->GetInputs();
+
+								for (UDMMaterialStageInput* MaskTextureStageInput : MaskTextureStageInputs)
+								{
+									if (UDMMaterialStageInputValue* MaskTextureInputValue = Cast<UDMMaterialStageInputValue>(MaskTextureStageInput))
+									{
+										if (UDMMaterialValueTexture* MaskTextureValue = Cast<UDMMaterialValueTexture>(MaskTextureInputValue->GetValue()))
+										{
+											MaskTextureValue->SetDefaultValue(UDynamicMaterialEditorSettings::Get()->DefaultOpaqueTexture.LoadSynchronous());
+											MaskTextureValue->ApplyDefaultValue();
+											break;
+										}
+									}
+								}
+							}
+						}
+
+						break;
+					}
+				}
+			}
+
+			break;
+
+		default:
+			// Do nothing
+			break;
+	}
+
+	OnSlotListUpdateDelegate.Broadcast(MaterialModel);
+
+	RequestMaterialBuild();
+
+	return NewSlot;
+}
+
+void UDynamicMaterialModelEditorOnlyData::RemoveSlot(int32 Index)
+{
+	UDMMaterialSlot* Slot = GetSlot(Index);
+	check(Slot);
+	
+	if (GUndo)
+	{
+		Slot->Modify(/* Always mark dirty */ false);
+	}
+
+	TArray<EDMMaterialPropertyType> SlotProperties = GetMaterialPropertiesForSlot(Slot);
+
+	for (EDMMaterialPropertyType MaterialProperty : SlotProperties)
+	{
+		UnassignMaterialProperty(MaterialProperty);
+	}
+
+	const EDMMaterialPropertyType* Key = PropertySlotMap.FindKey(Slot);
+
+	if (Key)
+	{
+		PropertySlotMap.Remove(*Key);
+	}
+
+	Slots.RemoveAt(Index);
+	Slot->GetOnConnectorsUpdateDelegate().RemoveAll(this);
+	Slot->SetComponentState(EDMComponentLifetimeState::Removed);
+
+	RequestMaterialBuild();
+
+	OnSlotListUpdateDelegate.Broadcast(MaterialModel);
+}
+
+UDMMaterialSlot* UDynamicMaterialModelEditorOnlyData::GetSlotForMaterialProperty(EDMMaterialPropertyType Property) const
+{
+	TObjectPtr<UDMMaterialSlot> const* SlotPtr = PropertySlotMap.Find(Property);
+
+	if (!SlotPtr)
+	{
+		return nullptr;
+	}
+
+	return *SlotPtr;
+}
+
+TArray<EDMMaterialPropertyType> UDynamicMaterialModelEditorOnlyData::GetMaterialPropertiesForSlot(UDMMaterialSlot* Slot) const
+{
+	TArray<EDMMaterialPropertyType> OutProperties;
+
+	for (const TPair<EDMMaterialPropertyType, TObjectPtr<UDMMaterialSlot>>& Pair : PropertySlotMap)
+	{
+		if (Pair.Value == Slot)
+		{
+			OutProperties.Add(Pair.Key);
+		}
+	}
+
+	return OutProperties;
+}
+
+void UDynamicMaterialModelEditorOnlyData::AssignMaterialPropertyToSlot(EDMMaterialPropertyType Property, UDMMaterialSlot* Slot)
+{
+	if (!Slot)
+	{
+		UnassignMaterialProperty(Property);
+		return;
+	}
+
+	check(Properties.Contains(Property));
+
+	PropertySlotMap.FindOrAdd(Property) = Slot;
+	Properties[Property]->ResetInputConnectionMap();
+	Slot->GetOnPropertiesUpdateDelegate().Broadcast(Slot);
+
+	RequestMaterialBuild();
+}
+
+void UDynamicMaterialModelEditorOnlyData::UnassignMaterialProperty(EDMMaterialPropertyType Property)
+{
+	TObjectPtr<UDMMaterialSlot>* SlotPtr = PropertySlotMap.Find(Property);
+	check(SlotPtr);
+
+	PropertySlotMap.Remove(Property);
+	(*SlotPtr)->GetOnPropertiesUpdateDelegate().Broadcast(*SlotPtr);
+
+	RequestMaterialBuild();
+}
+
+bool UDynamicMaterialModelEditorOnlyData::IsPixelAnimationFlagSet() const
+{
+	return bPixelAnimationFlag;
+}
+
+void UDynamicMaterialModelEditorOnlyData::SetPixelAnimationFlag(bool bInFlagValue)
+{
+	if (bPixelAnimationFlag == bInFlagValue)
+	{
+		return;
+	}
+
+	bPixelAnimationFlag = bInFlagValue;
+
+	if (MaterialModel)
+	{
+		if (UMaterial* Material = MaterialModel->GetGeneratedMaterial())
+		{
+			if (GUndo)
+			{
+				Material->Modify();
+			}
+
+			Material->bHasPixelAnimation = bPixelAnimationFlag;
+		}
+	}
+
+	RequestMaterialBuild();
+}
+
+void UDynamicMaterialModelEditorOnlyData::NotifyPostChange(const FPropertyChangedEvent& PropertyChangedEvent, class FEditPropertyChain* PropertyThatChanged)
+{
+	RequestMaterialBuild();
+}
+
+void UDynamicMaterialModelEditorOnlyData::PostLoad()
+{
+	Super::PostLoad();
+
+	SetFlags(RF_Transactional);
+
+	ReinitComponents();
+}
+
+void UDynamicMaterialModelEditorOnlyData::PostEditUndo()
+{
+	Super::PostEditUndo();
+
+	RequestMaterialBuild();
+}
+
+void UDynamicMaterialModelEditorOnlyData::PostEditImport()
+{
+	Super::PostEditImport();
+
+	PostEditorDuplicate();
+	ReinitComponents();
+	RequestMaterialBuild();
+}
+
+void UDynamicMaterialModelEditorOnlyData::PostDuplicate(bool bDuplicateForPIE)
+{
+	Super::PostDuplicate(bDuplicateForPIE);
+
+	PostEditorDuplicate();
+	ReinitComponents();
+	RequestMaterialBuild();
+}
+
+void UDynamicMaterialModelEditorOnlyData::OnValueUpdated(UDMMaterialValue* InValue, EDMUpdateType InUpdateType)
+{
+	check(InValue);
+
+	if (InValue->GetMaterialModel() != MaterialModel)
+	{
+		return;
+	}
+
+	if (InUpdateType == EDMUpdateType::Value && IsValid(MaterialModel) && IsValid(MaterialModel->DynamicMaterialInstance))
+	{
+		InValue->SetMIDParameter(MaterialModel->DynamicMaterialInstance);
+	}
+
+	OnValueUpdateDelegate.Broadcast(MaterialModel, InValue);
+
+	// Non-exported materials have their values update via settings parameters
+	// Exported materials need to be rebuilt to update the main material.
+	const bool bMaterialInDifferentPackage = MaterialModel->DynamicMaterial ? MaterialModel->DynamicMaterial->GetPackage() != GetPackage() : true;
+
+	if (InUpdateType == EDMUpdateType::Structure || bMaterialInDifferentPackage)
+	{
+		RequestMaterialBuild();
+	}
+}
+
+void UDynamicMaterialModelEditorOnlyData::OnTextureUVUpdated(UDMTextureUV* InTextureUV)
+{
+	check(InTextureUV);
+
+	if (InTextureUV->GetMaterialModel() != MaterialModel)
+	{
+		return;
+	}
+
+	if (IsValid(MaterialModel) && IsValid(MaterialModel->DynamicMaterialInstance))
+	{
+		InTextureUV->SetMIDParameters(MaterialModel->DynamicMaterialInstance);
+	}
+
+	OnTextureUVUpdateDelegate.Broadcast(MaterialModel, InTextureUV);
+
+	// Non-exported materials have their values update via settings parameters
+	// Exported materials need to be rebuilt to update the main material.
+	if (MaterialModel->DynamicMaterial && MaterialModel->DynamicMaterial->GetPackage() != GetPackage())
+	{
+		RequestMaterialBuild();
+	}
+}
+
+void UDynamicMaterialModelEditorOnlyData::SaveEditor()
+{
+	UEditorLoadingAndSavingUtils::SavePackages({GetPackage()}, false);
+
+	if (IsValid(MaterialModel) && IsValid(MaterialModel->DynamicMaterial))
+	{
+		if (MaterialModel->DynamicMaterial->GetPackage() != GetPackage())
+		{
+			UEditorLoadingAndSavingUtils::SavePackages({MaterialModel->DynamicMaterial->GetPackage()}, false);
+		}
+	}
+}
+
+FString UDynamicMaterialModelEditorOnlyData::GetMaterialAssetPath() const
+{
+	return FPaths::GetPath(GetPackage()->GetPathName());
+}
+
+FString UDynamicMaterialModelEditorOnlyData::GetMaterialAssetName() const
+{
+	return GetName() + "_Mat";
+}
+
+FString UDynamicMaterialModelEditorOnlyData::GetMaterialPackageName(const FString& MaterialBaseName) const
+{
+	return GetPackage()->GetName() + "_Mat";
+}
+
+void UDynamicMaterialModelEditorOnlyData::OnSlotConnectorsUpdated(UDMMaterialSlot* Slot)
+{
+	check(Slot);
+
+	RequestMaterialBuild();
+		
+	TArray<EDMMaterialPropertyType> SlotProperties = GetMaterialPropertiesForSlot(Slot);
+
+	for (EDMMaterialPropertyType Property : SlotProperties)
+	{
+		Properties[Property]->ResetInputConnectionMap();
+	}
+}
+
+void UDynamicMaterialModelEditorOnlyData::GenerateOpacityExpressions(const TSharedRef<FDMMaterialBuildState>& InBuildState, UDMMaterialSlot* InFromSlot,
+	EDMMaterialPropertyType InFromProperty, UMaterialExpression*& OutExpression, int32& OutOutputIndex, int32& OutOutputChannel) const
+{
+	const TArray<TObjectPtr<UDMMaterialLayerObject>>& SlotLayers = InFromSlot->GetLayers();
+	OutExpression = nullptr;
+
+	for (const TObjectPtr<UDMMaterialLayerObject>& Layer : SlotLayers)
+	{
+		// Although we are working with masks, if the base is disabled, this is handled by the GenerateExpressions
+		// of the LayerBlend code (to multiply alpha together, instead of maxing it).
+		if (Layer->GetMaterialProperty() != InFromProperty || !Layer->IsEnabled() || !Layer->IsStageEnabled(EDMMaterialLayerStage::Base))
+		{
+			continue;
+		}
+
+		UDMMaterialStage* BaseStage = Layer->GetStage(EDMMaterialLayerStage::Base);
+		UDMMaterialStage* MaskStage = Layer->GetStage(EDMMaterialLayerStage::Mask);
+
+		MaskStage->GenerateExpressions(InBuildState);
+		UDMMaterialStageThroughputLayerBlend* LayerBlend = Cast<UDMMaterialStageThroughputLayerBlend>(MaskStage->GetSource());
+
+		if (!LayerBlend)
+		{
+			continue;
+		}
+
+		UMaterialExpression* MaskOutputExpression;
+		int32 MaskOutputIndex;
+		int32 MaskOutputChannel;
+		LayerBlend->GetMaskOutput(InBuildState, MaskOutputExpression, MaskOutputIndex, MaskOutputChannel);
+
+		if (!MaskOutputExpression)
+		{
+			continue;
+		}
+
+		if (LayerBlend->UsePremultiplyAlpha())
+		{
+			if (UDMMaterialStageSource* Source = BaseStage->GetSource())
+			{
+				UMaterialExpression* LayerAlphaOutputExpression;
+				int32 LayerAlphaOutputIndex;
+				int32 LayerAlphaOutputChannel;
+
+				Source->GetMaskAlphaBlendNode(InBuildState, LayerAlphaOutputExpression, LayerAlphaOutputIndex, LayerAlphaOutputChannel);
+
+				if (LayerAlphaOutputExpression)
+				{
+					UMaterialExpressionMultiply* AlphaMultiply = InBuildState->GetBuildUtils().CreateExpression<UMaterialExpressionMultiply>(UE_DM_NodeComment_Default);
+
+					AlphaMultiply->A.Expression = MaskOutputExpression;
+					AlphaMultiply->A.OutputIndex = MaskOutputIndex;
+					AlphaMultiply->A.Mask = 0;
+
+					if (MaskOutputChannel != FDMMaterialStageConnectorChannel::WHOLE_CHANNEL)
+					{
+						AlphaMultiply->A.Mask = 1;
+						AlphaMultiply->A.MaskR = !!(MaskOutputChannel & FDMMaterialStageConnectorChannel::FIRST_CHANNEL);
+						AlphaMultiply->A.MaskG = !!(MaskOutputChannel & FDMMaterialStageConnectorChannel::SECOND_CHANNEL);
+						AlphaMultiply->A.MaskB = !!(MaskOutputChannel & FDMMaterialStageConnectorChannel::THIRD_CHANNEL);
+						AlphaMultiply->A.MaskA = !!(MaskOutputChannel & FDMMaterialStageConnectorChannel::FOURTH_CHANNEL);
+					}
+
+					AlphaMultiply->B.Expression = LayerAlphaOutputExpression;
+					AlphaMultiply->B.OutputIndex = LayerAlphaOutputIndex;
+					AlphaMultiply->B.Mask = 0;
+
+					if (LayerAlphaOutputChannel != FDMMaterialStageConnectorChannel::WHOLE_CHANNEL)
+					{
+						AlphaMultiply->B.Mask = 1;
+						AlphaMultiply->B.MaskR = !!(LayerAlphaOutputChannel & FDMMaterialStageConnectorChannel::FIRST_CHANNEL);
+						AlphaMultiply->B.MaskG = !!(LayerAlphaOutputChannel & FDMMaterialStageConnectorChannel::SECOND_CHANNEL);
+						AlphaMultiply->B.MaskB = !!(LayerAlphaOutputChannel & FDMMaterialStageConnectorChannel::THIRD_CHANNEL);
+						AlphaMultiply->B.MaskA = !!(LayerAlphaOutputChannel & FDMMaterialStageConnectorChannel::FOURTH_CHANNEL);
+					}
+
+					MaskOutputExpression = AlphaMultiply;
+					MaskOutputIndex = 0;
+					MaskOutputChannel = FDMMaterialStageConnectorChannel::WHOLE_CHANNEL;
+				}
+			}
+		}
+
+		if (OutExpression == nullptr)
+		{
+			OutExpression = MaskOutputExpression;
+			
+			// The first output will use the node's output info.
+			OutOutputIndex = MaskOutputIndex;
+			OutOutputChannel = MaskOutputChannel;
+			continue;
+		}
+
+		UMaterialExpressionMax* Max = InBuildState->GetBuildUtils().CreateExpression<UMaterialExpressionMax>(UE_DM_NodeComment_Default);
+		check(Max);
+
+		Max->A.Expression = OutExpression;
+		Max->A.OutputIndex = OutOutputIndex;
+		Max->A.Mask = 0;
+
+		if (OutOutputChannel != FDMMaterialStageConnectorChannel::WHOLE_CHANNEL)
+		{
+			Max->A.Mask = 1;
+			Max->A.MaskR = !!(OutOutputChannel & FDMMaterialStageConnectorChannel::FIRST_CHANNEL);
+			Max->A.MaskG = !!(OutOutputChannel & FDMMaterialStageConnectorChannel::SECOND_CHANNEL);
+			Max->A.MaskB = !!(OutOutputChannel & FDMMaterialStageConnectorChannel::THIRD_CHANNEL);
+			Max->A.MaskA = !!(OutOutputChannel & FDMMaterialStageConnectorChannel::FOURTH_CHANNEL);
+		}
+
+		Max->B.Expression = MaskOutputExpression;
+		Max->B.OutputIndex = MaskOutputIndex;
+		Max->B.Mask = 0;
+
+		if (MaskOutputChannel != FDMMaterialStageConnectorChannel::WHOLE_CHANNEL)
+		{
+			Max->B.Mask = 1;
+			Max->B.MaskR = !!(MaskOutputChannel & FDMMaterialStageConnectorChannel::FIRST_CHANNEL);
+			Max->B.MaskG = !!(MaskOutputChannel & FDMMaterialStageConnectorChannel::SECOND_CHANNEL);
+			Max->B.MaskB = !!(MaskOutputChannel & FDMMaterialStageConnectorChannel::THIRD_CHANNEL);
+			Max->B.MaskA = !!(MaskOutputChannel & FDMMaterialStageConnectorChannel::FOURTH_CHANNEL);
+		}
+
+		OutExpression = Max;
+
+		// If we have to combine, it will use the Max node's output info
+		OutOutputIndex = 0;
+		OutOutputChannel = FDMMaterialStageConnectorChannel::WHOLE_CHANNEL;
+	}
+}
+
+void UDynamicMaterialModelEditorOnlyData::ReinitComponents()
+{
+	for (int32 SlotIdx = 0; SlotIdx < Slots.Num(); ++SlotIdx)
+	{
+		if (IsValid(Slots[SlotIdx]))
+		{
+			Slots[SlotIdx]->GetOnConnectorsUpdateDelegate().AddUObject(this, &UDynamicMaterialModelEditorOnlyData::OnSlotConnectorsUpdated);
+		}
+		else
+		{
+			Slots.RemoveAt(SlotIdx);
+			--SlotIdx;
+		}
+	}
+}
+
+void UDynamicMaterialModelEditorOnlyData::PostEditorDuplicate()
+{
+	if (GUndo)
+	{
+		Modify();
+	}
+
+	for (const TPair<EDMMaterialPropertyType, TObjectPtr<UDMMaterialProperty>>& Pair : Properties)
+	{
+		if (UDMMaterialProperty* Property = Pair.Value.Get())
+		{
+			if (GUndo)
+			{
+				Property->Modify();
+			}
+
+			Property->PostEditorDuplicate(MaterialModel, nullptr);
+		}
+	}
+
+	for (UDMMaterialSlot* Slot : Slots)
+	{
+		if (GUndo)
+		{
+			Slot->Modify();
+		}
+
+		Slot->PostEditorDuplicate(MaterialModel, nullptr);
+	}
+
+	PropertySlotMap.Empty();
+
+	for (UDMMaterialSlot* Slot : Slots)
+	{
+		const TArray<TObjectPtr<UDMMaterialLayerObject>>& SlotLayers = Slot->GetLayers();
+
+		for (const TObjectPtr<UDMMaterialLayerObject>& Layer : SlotLayers)
+		{
+			EDMMaterialPropertyType Property = Layer->GetMaterialProperty();
+			TObjectPtr<UDMMaterialSlot>* SlotPtr = PropertySlotMap.Find(Property);
+
+			if (!SlotPtr)
+			{
+				PropertySlotMap.Emplace(Property, Slot);
+			}
+			else
+			{
+				check(*SlotPtr == Slot);
+			}
+		}
+	}
+}
+
+#undef LOCTEXT_NAMESPACE
diff --git a/Engine/Plugins/Experimental/DynamicMaterial/Source/DynamicMaterialEditor/Private/Slate/Layers/SDMLayerEffectsItem.cpp b/Engine/Plugins/Experimental/DynamicMaterial/Source/DynamicMaterialEditor/Private/Slate/Layers/SDMLayerEffectsItem.cpp
new file mode 100644
--- /dev/null
+++ b/Engine/Plugins/Experimental/DynamicMaterial/Source/DynamicMaterialEditor/Private/Slate/Layers/SDMLayerEffectsItem.cpp
@@ -0,0 +1,392 @@
+// Copyright Epic Games, Inc. All Rights Reserved.
+
+#include "Slate/Layers/SDMLayerEffectsItem.h"
+#include "Components/DMMaterialEffect.h"
+#include "Components/DMMaterialEffectStack.h"
+#include "Components/DMMaterialStage.h"
+#include "DragDrop/DMLayerEffectsDragDropOperation.h"
+#include "DynamicMaterialEditorStyle.h"
+#include "ScopedTransaction.h"
+#include "Slate/SDMStage.h"
+#include "SlateOptMacros.h"
+#include "Widgets/Images/SImage.h"
+#include "Widgets/Input/SButton.h"
+#include "Widgets/Layout/SExpandableArea.h"
+
+#define LOCTEXT_NAMESPACE "SDMLayerEffectsItem"
+
+BEGIN_SLATE_FUNCTION_BUILD_OPTIMIZATION
+
+FDMEffectsLayerItem::FDMEffectsLayerItem(UDMMaterialEffect* InMaterialEffect)
+{
+	MaterialEffectWeak = InMaterialEffect;
+}
+
+void SDMLayerEffectsItem::Construct(const FArguments& InArgs, const TSharedRef<STableViewBase>& InEffectsView, 
+	const TSharedPtr<FDMEffectsLayerItem>& InLayerItem)
+{
+	EffectsViewWeak = InEffectsView;
+	LayerItem = InLayerItem;
+	bEffectSelected = false;
+
+	OnEffectSelected = InArgs._OnEffectSelected;
+
+	STableRow<TSharedPtr<FDMEffectsLayerItem>>::Construct(
+		STableRow<TSharedPtr<FDMEffectsLayerItem>>::FArguments()
+		.Padding(2.0f)
+		.ShowSelection(true)
+		.ToolTipText(this, &SDMLayerEffectsItem::GetToolTipText)
+		.Style(FDynamicMaterialEditorStyle::Get(), "EffectsView.Row")
+		.OnPaintDropIndicator(this, &SDMLayerEffectsItem::OnLayerItemPaintDropIndicator)
+		.OnCanAcceptDrop(this, &SDMLayerEffectsItem::OnLayerItemCanAcceptDrop)
+		.OnDragDetected(this, &SDMLayerEffectsItem::OnLayerItemDragDetected)
+		.OnDragEnter(this, &SDMLayerEffectsItem::OnLayerItemDragEnter)
+		.OnDragLeave(this, &SDMLayerEffectsItem::OnLayerItemDragLeave)
+		.OnAcceptDrop(this, &SDMLayerEffectsItem::OnLayerItemAcceptDrop)
+		, InEffectsView
+	);
+
+	SetContent(CreateMainContent());
+}
+
+TSharedRef<SWidget> SDMLayerEffectsItem::CreateMainContent()
+{
+	return
+		SNew(SVerticalBox)
+		+ SVerticalBox::Slot()
+		[
+			SNew(SHorizontalBox)
+			.ToolTipText(this, &SDMLayerEffectsItem::GetToolTipText)
+			+ SHorizontalBox::Slot()
+			.AutoWidth()
+			.HAlign(HAlign_Left)
+			.VAlign(VAlign_Center)
+			.Padding(5.0f, 0.0f, 0.0f, 0.0f)
+			[
+				CreateLayerBypassButton()
+			]
+			+ SHorizontalBox::Slot()
+			.FillWidth(1.0f)
+			.HAlign(HAlign_Left)
+			.VAlign(VAlign_Center)
+			.Padding(10.0f, 0.0f, 0.0f, 0.0f)
+			[
+				SNew(STextBlock)
+				.Text(this, &SDMLayerEffectsItem::GetLayerHeaderText)
+			]
+			+ SHorizontalBox::Slot()
+			.AutoWidth()
+			.HAlign(HAlign_Right)
+			.VAlign(VAlign_Center)
+			.Padding(0.0f, 0.0f, 10.0f, 0.0f)
+			[
+				CreateLayerRemoveButton()
+			]
+		];
+}
+
+TSharedRef<SWidget> SDMLayerEffectsItem::CreateLayerBypassButton()
+{
+	return 
+		SNew(SButton)
+		.ContentPadding(4.0f)
+		.ButtonStyle(FDynamicMaterialEditorStyle::Get(), "HoverHintOnly")
+		.ToolTipText(LOCTEXT("LayerBypassTooltip", "Toggle the bypassing of this layer."))
+		.OnClicked(this, &SDMLayerEffectsItem::OnLayerBypassButtonClick)
+		[
+			SNew(SImage)
+			.DesiredSizeOverride(FVector2D(16.0f))
+			.Image(this, &SDMLayerEffectsItem::GetLayerBypassButtonImage)
+		];
+}
+
+TSharedRef<SWidget> SDMLayerEffectsItem::CreateLayerRemoveButton()
+{
+	return 
+		SNew(SButton)
+		.ContentPadding(2.0f)
+		.ButtonStyle(FDynamicMaterialEditorStyle::Get(), "HoverHintOnly")
+		.ToolTipText(LOCTEXT("RemoveEffect", "Remove Effect"))
+		.OnClicked(this, &SDMLayerEffectsItem::OnLayerRemoveButtonClick)
+		[
+			SNew(SImage)
+			.Image(FAppStyle::Get().GetBrush("Icons.Delete"))
+			.DesiredSizeOverride(FVector2D(12.0f))
+		];
+}
+
+END_SLATE_FUNCTION_BUILD_OPTIMIZATION
+
+FCursorReply SDMLayerEffectsItem::OnCursorQuery(const FGeometry& InMyGeometry, const FPointerEvent& InCursorEvent) const
+{
+	if (InCursorEvent.IsMouseButtonDown(EKeys::LeftMouseButton))
+	{
+		return FCursorReply::Cursor(EMouseCursor::GrabHandClosed);
+	}
+	return STableRow<TSharedPtr<FDMEffectsLayerItem>>::OnCursorQuery(InMyGeometry, InCursorEvent);
+}
+
+int32 SDMLayerEffectsItem::OnLayerItemPaintDropIndicator(EItemDropZone InItemDropZone, const FPaintArgs& Args, const FGeometry& AllottedGeometry, const FSlateRect& MyCullingRect, FSlateWindowElementList& OutDrawElements, int32 LayerId, const FWidgetStyle& InWidgetStyle, bool bParentEnabled) const
+{
+	const FSlateBrush* DropIndicatorBrush = GetDropIndicatorBrush(InItemDropZone);
+	static float OffsetX = 10.0f;
+	FVector2D Offset(OffsetX * GetIndentLevel(), 0.0f);
+
+	FSlateDrawElement::MakeBox(
+		OutDrawElements,
+		LayerId++,
+		AllottedGeometry.ToPaintGeometry(FVector2D(AllottedGeometry.GetLocalSize() - Offset), FSlateLayoutTransform(Offset)),
+		DropIndicatorBrush,
+		ESlateDrawEffect::None,
+		DropIndicatorBrush->GetTint(InWidgetStyle) * InWidgetStyle.GetColorAndOpacityTint()
+	);
+
+	return LayerId;
+}
+
+TOptional<EItemDropZone> SDMLayerEffectsItem::OnLayerItemCanAcceptDrop(const FDragDropEvent& InDragDropEvent, EItemDropZone InDropZone, 
+	TSharedPtr<FDMEffectsLayerItem> InSlotLayer) const
+{
+	if (InSlotLayer.IsValid())
+	{
+		TSharedPtr<FDMLayerEffectsDragDropOperation> DragDropOperation = InDragDropEvent.GetOperationAs<FDMLayerEffectsDragDropOperation>();
+		if (DragDropOperation.IsValid())
+		{
+			return InDropZone;
+		}
+	}
+
+	return TOptional<EItemDropZone>();
+}
+
+void SDMLayerEffectsItem::OnLayerItemDragEnter(const FDragDropEvent& InDragDropEvent)
+{
+
+}
+
+void SDMLayerEffectsItem::OnLayerItemDragLeave(const FDragDropEvent& InDragDropEvent)
+{
+
+}
+
+FReply SDMLayerEffectsItem::OnLayerItemDragDetected(const FGeometry& InMyGeometry, const FPointerEvent& InMouseEvent)
+{
+	const bool bShouldDuplicate = InMouseEvent.IsAltDown();
+
+	TSharedRef<FDMLayerEffectsDragDropOperation> DragDropOperation = MakeShared<FDMLayerEffectsDragDropOperation>(SharedThis(this), bShouldDuplicate);
+
+	return FReply::Handled().BeginDragDrop(DragDropOperation);
+}
+
+FReply SDMLayerEffectsItem::OnLayerItemAcceptDrop(const FDragDropEvent& InDragDropEvent, EItemDropZone InDropZone,
+	TSharedPtr<FDMEffectsLayerItem> InSlotLayer)
+{
+	if (!LayerItem.IsValid())
+	{
+		return FReply::Handled();
+	}
+
+	UDMMaterialEffect* MaterialEffect = LayerItem->MaterialEffectWeak.Get();
+
+	if (!MaterialEffect)
+	{
+		return FReply::Handled();
+	}
+
+	UDMMaterialEffectStack* MaterialEffectStack = MaterialEffect->GetEffectStack();
+
+	if (!MaterialEffectStack)
+	{
+		return FReply::Handled();
+	}
+
+	TSharedPtr<FDMLayerEffectsDragDropOperation> DragDropOperation = InDragDropEvent.GetOperationAs<FDMLayerEffectsDragDropOperation>();
+
+	if (!DragDropOperation.IsValid())
+	{
+		return FReply::Handled();
+	}
+
+	TSharedPtr<SDMLayerEffectsItem> DraggedWidget = DragDropOperation->GetLayerItemWidget();
+
+	if (!DraggedWidget.IsValid())
+	{
+		return FReply::Handled();
+	}
+
+	TSharedPtr<FDMEffectsLayerItem> DraggedItem = DraggedWidget->GetLayerItem();
+
+	if (!DraggedItem.IsValid())
+	{
+		return FReply::Handled();
+	}
+
+	UDMMaterialEffect* DraggedMaterialEffect = DraggedItem->MaterialEffectWeak.Get();
+
+	if (!IsValid(DraggedMaterialEffect) || DraggedMaterialEffect == MaterialEffect)
+	{
+		return FReply::Handled();
+	}
+
+	UDMMaterialEffectStack* DraggedMaterialEffectStack = DraggedMaterialEffect->GetEffectStack();
+
+	if (!IsValid(DraggedMaterialEffectStack) || DraggedMaterialEffectStack != MaterialEffectStack)
+	{
+		return FReply::Handled();
+	}
+
+	const int32 MaterialEffectIndex = MaterialEffect->FindIndex();
+
+	if (MaterialEffectIndex == INDEX_NONE)
+	{
+		return FReply::Handled();
+	}
+
+	const int32 DraggedMaterialEffectIndex = DraggedMaterialEffect->FindIndex();
+
+	if (DraggedMaterialEffectIndex != INDEX_NONE)
+	{
+		return FReply::Handled();
+	}
+
+	if (InDropZone == EItemDropZone::AboveItem)
+	{
+		MaterialEffectStack->MoveEffect(DraggedMaterialEffect, MaterialEffectIndex);
+	}
+	else
+	{
+		MaterialEffectStack->MoveEffect(DraggedMaterialEffect, MaterialEffectIndex + 1);
+	}
+
+	return FReply::Handled();
+}
+
+FText SDMLayerEffectsItem::GetToolTipText() const
+{
+	if (LayerItem.IsValid())
+	{
+		if (UDMMaterialEffect* MaterialEffect = LayerItem->MaterialEffectWeak.Get())
+		{
+			return MaterialEffect->GetEffectDescription();
+		}
+	}
+
+	return FText::GetEmpty();
+}
+
+FText SDMLayerEffectsItem::GetLayerHeaderText() const
+{
+	if (LayerItem.IsValid())
+	{
+		if (UDMMaterialEffect* MaterialEffect = LayerItem->MaterialEffectWeak.Get())
+		{
+			return MaterialEffect->GetEffectName();
+		}
+	}
+
+	return FText::GetEmpty();
+}
+
+void SDMLayerEffectsItem::OnEffectClick(const FPointerEvent& InMouseEvent, const TSharedRef<SDMStage>& InStageWidget)
+{
+	UDMMaterialStage* const ClickedStage = InStageWidget->GetStage();
+
+	if (!ClickedStage)
+	{
+		return;
+	}
+
+	const FKey MouseButton = InMouseEvent.GetEffectingButton();
+
+	if (MouseButton == EKeys::LeftMouseButton)
+	{
+		if (TSharedPtr<STableViewBase> TableViewBase = EffectsViewWeak.Pin())
+		{
+			TSharedPtr<SListView<TSharedPtr<FDMEffectsLayerItem>>> LayerView = StaticCastSharedPtr<SListView<TSharedPtr<FDMEffectsLayerItem>>>(TableViewBase);
+			LayerView->SetSelection(LayerItem);
+
+			// Trigger the unselected event
+			for (TSharedPtr<FDMEffectsLayerItem> LayerItemIter : LayerView->GetItems())
+			{
+				if (LayerItemIter.IsValid() && LayerItemIter != LayerItem)
+				{
+					if (TSharedPtr<ITableRow> Widget = LayerView->WidgetFromItem(LayerItemIter))
+					{
+						TSharedPtr<SDMLayerEffectsItem> ItemWidget = StaticCastSharedPtr<SDMLayerEffectsItem>(Widget);
+
+						if (ItemWidget->IsEffectSelected())
+						{
+							OnEffectSelected.ExecuteIfBound(false, ItemWidget.ToSharedRef());
+						}
+					}
+				}
+			}
+		}
+
+		OnEffectSelected.ExecuteIfBound(true, SharedThis(this));
+	}
+}
+
+bool SDMLayerEffectsItem::OnEffectCanAcceptDrop(const FDragDropEvent& InDragDropEvent, const TSharedRef<SDMStage>& InStageWidget) const
+{
+	return true;
+}
+
+FReply SDMLayerEffectsItem::OnEffectAcceptDrop(const FDragDropEvent& InDragDropEvent, const TSharedRef<SDMStage>& InStageWidget)
+{
+	return FReply::Handled();
+}
+
+FReply SDMLayerEffectsItem::OnLayerBypassButtonClick()
+{
+	if (LayerItem.IsValid())
+	{
+		if (UDMMaterialEffect* Effect = LayerItem->MaterialEffectWeak.Get())
+		{
+			FScopedTransaction Transaction(LOCTEXT("ToggleEffectEnabled", "Material Designer Toggle Effect"));
+			Effect->Modify();
+			Effect->SetEnabled(!Effect->IsEnabled());
+		}
+	}
+
+	return FReply::Handled();
+}
+
+const FSlateBrush* SDMLayerEffectsItem::GetLayerBypassButtonImage() const
+{
+	static const FSlateIcon VisibleIcon = FSlateIcon(FAppStyle::GetAppStyleSetName(), "Icons.Visible");
+	static const FSlateIcon InvisibleIcon = FSlateIcon(FAppStyle::GetAppStyleSetName(), "Icons.Hidden");
+
+	if (LayerItem.IsValid())
+	{
+		if (UDMMaterialEffect* Effect = LayerItem->MaterialEffectWeak.Get())
+		{
+			if (Effect->IsEnabled())
+			{
+				return VisibleIcon.GetIcon();
+			}
+		}
+	}
+
+	return InvisibleIcon.GetIcon();
+}
+
+FReply SDMLayerEffectsItem::OnLayerRemoveButtonClick()
+{
+	if (LayerItem.IsValid())
+	{
+		if (UDMMaterialEffect* Effect = LayerItem->MaterialEffectWeak.Get())
+		{
+			if (UDMMaterialEffectStack* EffectStack = Effect->GetEffectStack())
+			{
+				FScopedTransaction Transaction(LOCTEXT("RemoveEffect", "Material Designer Remove Effect"));
+				EffectStack->Modify();
+				EffectStack->RemoveEffect(Effect);
+			}
+		}
+	}
+
+	return FReply::Handled();
+}
+
+#undef LOCTEXT_NAMESPACE
diff --git a/Engine/Plugins/Experimental/DynamicMaterial/Source/DynamicMaterialEditor/Private/Slate/Layers/SDMLayerEffectsItem.h b/Engine/Plugins/Experimental/DynamicMaterial/Source/DynamicMaterialEditor/Private/Slate/Layers/SDMLayerEffectsItem.h
new file mode 100644
--- /dev/null
+++ b/Engine/Plugins/Experimental/DynamicMaterial/Source/DynamicMaterialEditor/Private/Slate/Layers/SDMLayerEffectsItem.h
@@ -0,0 +1,86 @@
+// Copyright Epic Games, Inc. All Rights Reserved.
+
+#pragma once
+
+#include "Components/DMMaterialSlot.h"
+#include "Slate/SDMStage.h"
+#include "Widgets/Views/SListView.h"
+#include "Widgets/Views/STableRow.h"
+
+class SBox;
+class SDMSlot;
+class SDMStage;
+class UDMMaterialEffect;
+class UDMMaterialSlot;
+class UDMMaterialStage;
+struct FSlateBrush;
+
+struct FDMEffectsLayerItem
+{
+	TWeakObjectPtr<UDMMaterialEffect> MaterialEffectWeak;
+
+	FDMEffectsLayerItem(UDMMaterialEffect* InMaterialEffect);
+};
+
+/**
+ * Material Slot Layer Effects Item
+ */
+class SDMLayerEffectsItem : public STableRow<TSharedPtr<FDMEffectsLayerItem>>
+{
+public:
+	DECLARE_DELEGATE_TwoParams(FDMOnEffectSelected, const bool, const TSharedRef<SDMLayerEffectsItem>& /** EffectsItemWidget */);
+
+	SLATE_BEGIN_ARGS(SDMLayerEffectsItem)
+	{}
+		SLATE_EVENT(FDMOnEffectSelected, OnEffectSelected)
+	SLATE_END_ARGS()
+
+	virtual ~SDMLayerEffectsItem() override = default;
+
+	void Construct(const FArguments& InArgs, const TSharedRef<STableViewBase>& InEffectsView, const TSharedPtr<FDMEffectsLayerItem>& InLayerItem);
+
+	const TSharedPtr<FDMEffectsLayerItem>& GetLayerItem() const { return LayerItem; }
+
+	FORCEINLINE bool IsEffectSelected() const { return bEffectSelected; }
+	FORCEINLINE void SetEffectSelected(const bool bInSelected) { bEffectSelected = bInSelected; }
+
+protected:
+	TWeakPtr<STableViewBase> EffectsViewWeak;
+	TSharedPtr<FDMEffectsLayerItem> LayerItem;
+
+	FDMOnEffectSelected OnEffectSelected;
+
+	bool bEffectSelected;
+
+	TSharedRef<SWidget> CreateMainContent();
+
+	TSharedRef<SWidget> CreateLayerBypassButton();
+	TSharedRef<SWidget> CreateLayerRemoveButton();
+
+	int32 OnLayerItemPaintDropIndicator(EItemDropZone InItemDropZone, const FPaintArgs& Args, const FGeometry& AllottedGeometry, 
+		const FSlateRect& MyCullingRect, FSlateWindowElementList& OutDrawElements, int32 LayerId, const FWidgetStyle& InWidgetStyle, 
+		bool bParentEnabled) const;
+	TOptional<EItemDropZone> OnLayerItemCanAcceptDrop(const FDragDropEvent& InDragDropEvent, EItemDropZone InDropZone, 
+		TSharedPtr<FDMEffectsLayerItem> InSlotLayer) const;
+	void OnLayerItemDragEnter(const FDragDropEvent& InDragDropEvent);
+	void OnLayerItemDragLeave(const FDragDropEvent& InDragDropEvent);
+	FReply OnLayerItemDragDetected(const FGeometry& InMyGeometry, const FPointerEvent& InMouseEvent);
+	FReply OnLayerItemAcceptDrop(const FDragDropEvent& InDragDropEvent, EItemDropZone InDropZone, TSharedPtr<FDMEffectsLayerItem> InSlotLayer);
+
+	FText GetToolTipText() const;
+	FText GetLayerHeaderText() const;
+
+	//~ Begin SWidget
+	virtual FCursorReply OnCursorQuery(const FGeometry& InMyGeometry, const FPointerEvent& InCursorEvent) const override;
+	//~ End SWidget
+
+	void OnEffectClick(const FPointerEvent& InMouseEvent, const TSharedRef<SDMStage>& InStageWidget);
+
+	bool OnEffectCanAcceptDrop(const FDragDropEvent& InDragDropEvent, const TSharedRef<SDMStage>& InStageWidget) const;
+	FReply OnEffectAcceptDrop(const FDragDropEvent& InDragDropEvent, const TSharedRef<SDMStage>& InStageWidget);
+
+	FReply OnLayerBypassButtonClick();
+	const FSlateBrush* GetLayerBypassButtonImage() const;
+
+	FReply OnLayerRemoveButtonClick();
+};
diff --git a/Engine/Plugins/Experimental/DynamicMaterial/Source/DynamicMaterialEditor/Private/Slate/Layers/SDMLayerEffectsView.cpp b/Engine/Plugins/Experimental/DynamicMaterial/Source/DynamicMaterialEditor/Private/Slate/Layers/SDMLayerEffectsView.cpp
new file mode 100644
--- /dev/null
+++ b/Engine/Plugins/Experimental/DynamicMaterial/Source/DynamicMaterialEditor/Private/Slate/Layers/SDMLayerEffectsView.cpp
@@ -0,0 +1,298 @@
+// Copyright Epic Games, Inc. All Rights Reserved.
+
+#include "Slate/Layers/SDMLayerEffectsView.h"
+#include "Components/DMMaterialEffect.h"
+#include "Components/DMMaterialEffectStack.h"
+#include "Components/DMMaterialLayer.h"
+#include "Slate/Layers/SDMLayerEffectsItem.h"
+#include "Slate/Layers/SDMSlotLayerItem.h"
+#include "Slate/SDMSlot.h"
+
+#define LOCTEXT_NAMESPACE "SDMLayerEffectsView"
+
+void SDMLayerEffectsView::Construct(const FArguments& InArgs, const TSharedPtr<SDMSlot>& InSlotWidget, UDMMaterialEffectStack* InMaterialEffectStack)
+{
+	SlotWidgetWeak = InSlotWidget;
+	MaterialEffectStackWeak = InMaterialEffectStack;
+	OnEffectSelected = InArgs._OnEffectSelected;
+
+	RebuildList();
+
+	SListView::Construct(SListView::FArguments()
+		.ListItemsSource(&EffectItems)
+		.Orientation(Orient_Vertical)
+		.SelectionMode(ESelectionMode::Single)
+		.ClearSelectionOnClick(false)
+		.EnableAnimatedScrolling(true)
+		.ScrollbarVisibility(EVisibility::Visible)
+		.ConsumeMouseWheel(EConsumeMouseWheel::WhenScrollingPossible)
+		.OnGenerateRow(this, &SDMLayerEffectsView::OnGenerateLayerItemWidget)
+		.OnSelectionChanged(this, &SDMLayerEffectsView::OnLayerItemSelectionChanged)
+		.OnContextMenuOpening(this->OnContextMenuOpening)
+	);
+
+	RequestListRefresh();
+
+	if (IsValid(InMaterialEffectStack))
+	{
+		InMaterialEffectStack->GetOnUpdate().AddSP(this, &SDMLayerEffectsView::OnEffectStackUpdate);
+	}
+}
+
+TSharedRef<ITableRow> SDMLayerEffectsView::OnGenerateLayerItemWidget(TSharedPtr<FDMEffectsLayerItem> InItem, const TSharedRef<STableViewBase>& InOwnerTable)
+{
+	return 
+		SNew(SDMLayerEffectsItem, InOwnerTable, InItem)
+		.FlowDirectionPreference(EFlowDirectionPreference::LeftToRight)
+		.OnEffectSelected(this, &SDMLayerEffectsView::OnLayerEffectSelected);
+}
+
+void SDMLayerEffectsView::OnLayerEffectSelected(const bool bInSelected, const TSharedRef<SDMLayerEffectsItem>& InEffectsItemWidget)
+{
+	if (!bInSelected)
+	{
+		SelectedEffectWidgets.Remove(InEffectsItemWidget);
+		InEffectsItemWidget->SetEffectSelected(false);
+
+		if (TSharedPtr<SDMSlot> SlotWidget = SlotWidgetWeak.Pin())
+		{
+			SlotWidget->SetEditedComponent(nullptr);
+			SlotWidget->SetSelectedLayer(nullptr);
+
+			if (TSharedPtr<FDMEffectsLayerItem> SelectedEffectLayerItem = SelectedEffectWidgets.Last()->GetLayerItem())
+			{
+				if (UDMMaterialEffect* Effect = SelectedEffectLayerItem->MaterialEffectWeak.Get())
+				{
+					SlotWidget->SetEditedComponent(Effect);
+
+					if (UDMMaterialEffectStack* EffectStack = Effect->GetEffectStack())
+					{
+						if (UDMMaterialLayerObject* SelectedLayer = EffectStack->GetLayer())
+						{
+							SlotWidget->SetSelectedLayer(nullptr);
+						}
+					}					
+				}
+			}
+		}
+	}
+	else
+	{
+		SelectedEffectWidgets.AddUnique(InEffectsItemWidget);
+		InEffectsItemWidget->SetEffectSelected(true);
+
+		if (TSharedPtr<SDMSlot> SlotWidget = SlotWidgetWeak.Pin())
+		{
+			UDMMaterialLayerObject* Layer = nullptr;
+			UDMMaterialEffect* Effect = nullptr;
+
+			if (TSharedPtr<FDMEffectsLayerItem> EffectViewItem = InEffectsItemWidget->GetLayerItem())
+			{
+				Effect = EffectViewItem->MaterialEffectWeak.Get();
+
+				if (Effect)
+				{
+					if (UDMMaterialEffectStack* EffectStack = Effect->GetEffectStack())
+					{
+						Layer = EffectStack->GetLayer();
+					}
+				}
+			}
+
+			SlotWidget->SetEditedComponent(Effect);
+			SlotWidget->SetSelectedLayer(Layer);
+		}
+	}
+}
+
+void SDMLayerEffectsView::OnEffectStackUpdate(UDMMaterialComponent* InComponent, EDMUpdateType InUpdateType)
+{
+	RebuildList();
+}
+
+TSharedPtr<SDMLayerEffectsItem> SDMLayerEffectsView::WidgetFromLayerItem(const TSharedPtr<FDMEffectsLayerItem>& InItem)
+{
+	return StaticCastSharedPtr<SDMLayerEffectsItem>(WidgetFromItem(InItem));
+}
+
+void SDMLayerEffectsView::OnLayerItemSelectionChanged(TSharedPtr<FDMEffectsLayerItem> InSelectedItem, ESelectInfo::Type InSelectInfo)
+{
+	TSharedPtr<SDMLayerEffectsItem> SlotLayerWidget = WidgetFromLayerItem(InSelectedItem);
+	if (!SlotLayerWidget.IsValid())
+	{
+		return;
+	}
+
+	OnLayerEffectSelected(true, SlotLayerWidget.ToSharedRef());
+
+	SelectedEffectIndex = INDEX_NONE;
+	for (int32 Index = 0; Index < EffectItems.Num(); ++Index)
+	{
+		if (EffectItems[Index] == InSelectedItem)
+		{
+			SelectedEffectIndex = Index;
+		}
+	}
+	if (SelectedEffectIndex == INDEX_NONE)
+	{
+		return;
+	}
+
+	SelectedEffectItems.Empty();
+	SelectedEffectItems.Add(EffectItems[SelectedEffectIndex]);
+
+	OnEffectSelected.ExecuteIfBound(InSelectedItem, SelectedEffectIndex);
+}
+
+void SDMLayerEffectsView::ScrollItemIntoView(const TSharedPtr<FDMEffectsLayerItem>& InItem)
+{
+	bNavigateOnScrollIntoView = true;
+
+	RequestScrollIntoView(InItem, 0);
+}
+
+void SDMLayerEffectsView::FocusOnItem(const TSharedPtr<FDMEffectsLayerItem>& InItem)
+{
+	SelectorItem = InItem;
+	RangeSelectionStart = InItem;
+}
+
+int32 SDMLayerEffectsView::GetLayerItemIndex(const TSharedPtr<FDMEffectsLayerItem>& Item) const
+{
+	return GetItems().Find(TListTypeTraits<TSharedPtr<FDMEffectsLayerItem>>::NullableItemTypeConvertToItemType(Item));
+}
+
+void SDMLayerEffectsView::RebuildList()
+{
+	EffectItems.Empty();
+
+	if (UDMMaterialEffectStack* EffectStack = MaterialEffectStackWeak.Get())
+	{
+		for (UDMMaterialEffect* MaterialEffect : EffectStack->GetEffects())
+		{
+			EffectItems.Add(MakeShared<FDMEffectsLayerItem>(MaterialEffect));
+		}
+	}
+
+	SListView<TSharedPtr<FDMEffectsLayerItem>>::RebuildList();
+}
+
+FCursorReply SDMLayerEffectsView::OnCursorQuery(const FGeometry& MyGeometry, const FPointerEvent& CursorEvent) const
+{
+	if (IsRightClickScrolling() && CursorEvent.IsMouseButtonDown(EKeys::RightMouseButton))
+	{
+		// We hide the native cursor as we'll be drawing the software EMouseCursor::GrabHandClosed cursor
+		return FCursorReply::Cursor(EMouseCursor::None);
+	}
+	return SListView::OnCursorQuery(MyGeometry, CursorEvent);
+}
+
+FReply SDMLayerEffectsView::OnMouseMove(const FGeometry& MyGeometry, const FPointerEvent& MouseEvent)
+{
+	if (MouseEvent.IsMouseButtonDown(EKeys::RightMouseButton) && !MouseEvent.IsTouchEvent())
+	{
+		// We only care about deltas along the scroll axis
+		FTableViewDimensions CursorDeltaDimensions(Orientation, MouseEvent.GetCursorDelta());
+		CursorDeltaDimensions.LineAxis = 0.f;
+
+		const float ScrollByAmount = CursorDeltaDimensions.ScrollAxis / MyGeometry.Scale;
+
+		// If scrolling with the right mouse button, we need to remember how much we scrolled.
+		// If we did not scroll at all, we will bring up the context menu when the mouse is released.
+		AmountScrolledWhileRightMouseDown += FMath::Abs(ScrollByAmount);
+
+		// Has the mouse moved far enough with the right mouse button held down to start capturing
+		// the mouse and dragging the view?
+		if (IsRightClickScrolling())
+		{
+			// Make sure the active timer is registered to update the inertial scroll
+			if (!bIsScrollingActiveTimerRegistered)
+			{
+				bIsScrollingActiveTimerRegistered = true;
+				RegisterActiveTimer(0.f, FWidgetActiveTimerDelegate::CreateSP(this, &SDMLayerEffectsView::UpdateInertialScroll));
+			}
+
+			TickScrollDelta -= ScrollByAmount;
+
+			const float AmountScrolled = this->ScrollBy(MyGeometry, -ScrollByAmount, AllowOverscroll);
+
+			FReply Reply = FReply::Handled();
+
+			// The mouse moved enough that we're now dragging the view. Capture the mouse
+			// so the user does not have to stay within the bounds of the list while dragging.
+			if (this->HasMouseCapture() == false)
+			{
+				Reply.CaptureMouse(AsShared()).UseHighPrecisionMouseMovement(AsShared());
+				SoftwareCursorPosition = MyGeometry.AbsoluteToLocal(MouseEvent.GetScreenSpacePosition());
+				bShowSoftwareCursor    = true;
+			}
+
+			// Check if the mouse has moved.
+			if (AmountScrolled != 0)
+			{
+				SoftwareCursorPosition += CursorDeltaDimensions.ToVector2D();
+			}
+
+			return Reply;
+		}
+	}
+
+	return SListView::OnMouseMove(MyGeometry, MouseEvent);
+}
+
+FReply SDMLayerEffectsView::OnMouseButtonUp(const FGeometry& MyGeometry, const FPointerEvent& MouseEvent)
+{
+	if (MouseEvent.GetEffectingButton() == EKeys::RightMouseButton)
+	{
+		FReply Reply = FReply::Handled().ReleaseMouseCapture();
+		AmountScrolledWhileRightMouseDown = 0;
+		bShowSoftwareCursor = false;
+
+		// If we have mouse capture, snap the mouse back to the closest location that is within the list's bounds
+		if (HasMouseCapture())
+		{
+			const FSlateRect ListScreenSpaceRect = MyGeometry.GetLayoutBoundingRect();
+			const FVector2D CursorPosition = MyGeometry.LocalToAbsolute(SoftwareCursorPosition);
+
+			const FIntPoint BestPositionInList(
+					FMath::RoundToInt(FMath::Clamp(CursorPosition.X, ListScreenSpaceRect.Left, ListScreenSpaceRect.Right)),
+					FMath::RoundToInt(FMath::Clamp(CursorPosition.Y, ListScreenSpaceRect.Top, ListScreenSpaceRect.Bottom))
+				);
+
+			Reply.SetMousePos(BestPositionInList);
+		}
+
+		return Reply;
+	}
+	return SListView::OnMouseButtonUp(MyGeometry, MouseEvent);
+}
+
+bool SDMLayerEffectsView::AddLayerItem(const TSharedPtr<FDMEffectsLayerItem>& InLayerItem)
+{
+	if (!InLayerItem.IsValid())
+	{
+		return false;
+	}
+
+	EffectItems.Add(InLayerItem);
+
+	RequestListRefresh();
+
+	return true;
+}
+
+bool SDMLayerEffectsView::RemoveLayerItem(const TSharedPtr<FDMEffectsLayerItem>& InLayerItem)
+{
+	if (!InLayerItem.IsValid())
+	{
+		return false;
+	}
+
+	EffectItems.Remove(InLayerItem);
+
+	RequestListRefresh();
+
+	return true;
+}
+
+#undef LOCTEXT_NAMESPACE
diff --git a/Engine/Plugins/Experimental/DynamicMaterial/Source/DynamicMaterialEditor/Private/Slate/Layers/SDMLayerEffectsView.h b/Engine/Plugins/Experimental/DynamicMaterial/Source/DynamicMaterialEditor/Private/Slate/Layers/SDMLayerEffectsView.h
new file mode 100644
--- /dev/null
+++ b/Engine/Plugins/Experimental/DynamicMaterial/Source/DynamicMaterialEditor/Private/Slate/Layers/SDMLayerEffectsView.h
@@ -0,0 +1,72 @@
+// Copyright Epic Games, Inc. All Rights Reserved.
+
+#pragma once
+
+#include "Slate/Layers/SDMLayerEffectsItem.h"
+#include "UObject/WeakObjectPtrTemplates.h"
+#include "Widgets/Views/SListView.h"
+
+class SDMLayerEffectsItem;
+class UDMMaterialEffectStack;
+
+DECLARE_DELEGATE_TwoParams(FDMOnEffectsViewItemSelected, TSharedPtr<FDMEffectsLayerItem> /** InLayerItem */, const int32 /** InLayerIndex */)
+
+class SDMLayerEffectsView : public SListView<TSharedPtr<FDMEffectsLayerItem>>
+{
+public:
+	SLATE_BEGIN_ARGS(SDMLayerEffectsView)
+	{}
+		SLATE_EVENT(FDMOnEffectsViewItemSelected, OnEffectSelected)
+	SLATE_END_ARGS()
+
+	virtual ~SDMLayerEffectsView() override {}
+
+	void Construct(const FArguments& InArgs, const TSharedPtr<SDMSlot>& InSlotWidget, UDMMaterialEffectStack* InMaterialEffectStack);
+
+	void ScrollItemIntoView(const TSharedPtr<FDMEffectsLayerItem>& InItem);
+	void FocusOnItem(const TSharedPtr<FDMEffectsLayerItem>& InItem);
+
+	int32 GetLayerItemIndex(const TSharedPtr<FDMEffectsLayerItem>& InItem) const;
+
+	const TArray<TSharedPtr<FDMEffectsLayerItem>>& GetSelectedEffectItems() const { return SelectedEffectItems; }
+	const TArray<TSharedPtr<SDMLayerEffectsItem>>& GetSelectedEffectWidgets() const { return SelectedEffectWidgets; }
+
+	TSharedPtr<SDMLayerEffectsItem> WidgetFromLayerItem(const TSharedPtr<FDMEffectsLayerItem>& InItem);
+
+	bool AddLayerItem(const TSharedPtr<FDMEffectsLayerItem>& InLayerItem);
+	bool RemoveLayerItem(const TSharedPtr<FDMEffectsLayerItem>& InLayerItem);
+
+	int32 GetLayerItemCount() const { return EffectItems.Num(); }
+
+	//~ Begin SListView
+	virtual void RebuildList() override;
+	//~ End SListView
+
+	//~ Begin SWidget
+	virtual FCursorReply OnCursorQuery(const FGeometry& MyGeometry, const FPointerEvent& CursorEvent) const override;
+	virtual FReply OnMouseMove(const FGeometry& MyGeometry, const FPointerEvent& MouseEvent) override;
+	virtual FReply OnMouseButtonUp(const FGeometry& MyGeometry, const FPointerEvent& MouseEvent) override;
+	//~ End SWidget
+
+protected:
+	TWeakPtr<SDMSlot> SlotWidgetWeak;
+	TWeakObjectPtr<UDMMaterialEffectStack> MaterialEffectStackWeak;
+
+	FDMOnEffectsViewItemSelected OnEffectSelected;
+	FOnContextMenuOpening OnContextMenuOpening;
+
+	TArray<TSharedPtr<FDMEffectsLayerItem>> EffectItems;
+
+	int32 SelectedEffectIndex = INDEX_NONE;
+
+	TArray<TSharedPtr<FDMEffectsLayerItem>> SelectedEffectItems;
+	TArray<TSharedPtr<SDMLayerEffectsItem>> SelectedEffectWidgets;
+
+	TSharedRef<ITableRow> OnGenerateLayerItemWidget(TSharedPtr<FDMEffectsLayerItem> InItem, const TSharedRef<STableViewBase>& InOwnerTable);
+
+	void OnLayerItemSelectionChanged(TSharedPtr<FDMEffectsLayerItem> InSelectedItem, ESelectInfo::Type InSelectInfo);
+
+	void OnLayerEffectSelected(const bool bInSelected, const TSharedRef<SDMLayerEffectsItem>& InEffectsItemWidget);
+
+	void OnEffectStackUpdate(UDMMaterialComponent* InComponent, EDMUpdateType InUpdateType);
+};
