[MaterialDesigner]

diff --git a/Engine/Plugins/Experimental/DynamicMaterial/Source/DynamicMaterialEditor/Private/Slate/Layers/SDMSlotLayerItem.cpp b/Engine/Plugins/Experimental/DynamicMaterial/Source/DynamicMaterialEditor/Private/Slate/Layers/SDMSlotLayerItem.cpp
new file mode 100644
--- /dev/null
+++ b/Engine/Plugins/Experimental/DynamicMaterial/Source/DynamicMaterialEditor/Private/Slate/Layers/SDMSlotLayerItem.cpp
@@ -0,0 +1,1266 @@
+// Copyright Epic Games, Inc. All Rights Reserved.
+
+#include "Slate/Layers/SDMSlotLayerItem.h"
+#include "Components/DMMaterialLayer.h"
+#include "Components/DMMaterialProperty.h"
+#include "Components/DMMaterialSlot.h"
+#include "Components/DMMaterialStage.h"
+#include "Components/DMMaterialStageThroughput.h"
+#include "Components/MaterialStageInputs/DMMSITextureUV.h"
+#include "DMPrivate.h"
+#include "DragDrop/DMSlotLayerDragDropOperation.h"
+#include "DynamicMaterialEditorSettings.h"
+#include "DynamicMaterialEditorStyle.h"
+#include "Framework/Application/SlateApplication.h"
+#include "Menus/DMMaterialSlotLayerAddEffectMenus.h"
+#include "Model/DynamicMaterialModelEditorOnlyData.h"
+#include "ScopedTransaction.h"
+#include "Slate/Layers/SDMLayerEffectsView.h"
+#include "Slate/SDMSlot.h"
+#include "Slate/SDMStage.h"
+#include "SlateOptMacros.h"
+#include "Styling/StyleColors.h"
+#include "Widgets/Images/SImage.h"
+#include "Widgets/Input/SButton.h"
+#include "Widgets/Input/SEditableTextBox.h"
+#include "Widgets/Layout/SBox.h"
+#include "Widgets/Layout/SExpandableArea.h"
+#include "Widgets/Text/STextBlock.h"
+
+#define LOCTEXT_NAMESPACE "SDMSlotLayer"
+
+BEGIN_SLATE_FUNCTION_BUILD_OPTIMIZATION
+
+void SDMSlotLayerItem::Construct(const FArguments& InArgs, const TSharedPtr<SDMSlot>& InSlotWidget, const TSharedRef<STableViewBase>& InLayerView, 
+	const TSharedPtr<FDMMaterialLayerReference>& InLayerItem)
+{
+	SlotWidgetWeak = InSlotWidget;
+	LayerViewWeak = InLayerView;
+	LayerItem = InLayerItem;
+
+	StageBaseEnabled = InArgs._StageBaseEnabled;
+	StageBaseSelected = InArgs._StageBaseSelected;
+	StageMaskEnabled = InArgs._StageMaskEnabled;
+	StageMaskSelected = InArgs._StageMaskSelected;
+	PreviewSize = InArgs._PreviewSize;
+	OnStageSelected = InArgs._OnStageSelected;
+	OnLayerLinkToggled = InArgs._OnLayerLinkToggled;
+
+	STableRow<TSharedPtr<FDMMaterialLayerReference>>::Construct(
+		STableRow<TSharedPtr<FDMMaterialLayerReference>>::FArguments()
+		.Padding(2.0f)
+		.ShowSelection(true)
+		.ToolTipText(this, &SDMSlotLayerItem::GetToolTipText)
+		.Style(FDynamicMaterialEditorStyle::Get(), "LayerView.Row")
+		.OnPaintDropIndicator(this, &SDMSlotLayerItem::OnLayerItemPaintDropIndicator)
+		.OnCanAcceptDrop(this, &SDMSlotLayerItem::OnLayerItemCanAcceptDrop)
+		.OnDragDetected(this, &SDMSlotLayerItem::OnLayerItemDragDetected)
+		.OnDragEnter(this, &SDMSlotLayerItem::OnLayerItemDragEnter)
+		.OnDragLeave(this, &SDMSlotLayerItem::OnLayerItemDragLeave)
+		.OnAcceptDrop(this, &SDMSlotLayerItem::OnLayerItemAcceptDrop)
+		, InLayerView
+	);
+
+	SetContent(CreateMainContent());
+}
+
+UDMMaterialLayerObject* SDMSlotLayerItem::GetLayer() const
+{
+	if (LayerItem.IsValid())
+	{
+		return LayerItem->GetLayer();
+	}
+
+	return nullptr;
+}
+
+TSharedRef<SWidget> SDMSlotLayerItem::CreateMainContent()
+{
+	UDMMaterialEffectStack* EffectStack = nullptr;
+
+	if (LayerItem.IsValid())
+	{
+		if (UDMMaterialLayerObject* Layer = LayerItem->GetLayer())
+		{
+			EffectStack = Layer->GetEffectStack();
+		}
+	}
+
+	return 
+		SNew(SVerticalBox)
+
+		+ SVerticalBox::Slot()
+		.AutoHeight()
+		.HAlign(HAlign_Fill)
+		.VAlign(VAlign_Center)
+		[
+			CreateHeaderRowContent()
+		]
+
+		+ SVerticalBox::Slot()
+		.AutoHeight()
+		.HAlign(HAlign_Fill)
+		.VAlign(VAlign_Center)
+		.Padding(10.0f, 0.0f, 0.0f, 0.0f)
+		[
+			SNew(SHorizontalBox)
+			.Visibility(this, &SDMSlotLayerItem::GetEffectsListVisibility)
+			+ SHorizontalBox::Slot()
+			.AutoWidth()
+			[
+				SNew(SBox)
+				.WidthOverride(3.0f)
+				[
+					SNew(SBorder)
+					.BorderBackgroundColor(FLinearColor(1, 1, 1, 0.2f))
+					.BorderImage(FDynamicMaterialEditorStyle::GetBrush("Border.Right"))
+				]
+			]
+			+ SHorizontalBox::Slot()
+			.AutoWidth()
+			[
+				SNew(SBox)
+				.WidthOverride(1.0f)
+				[
+					SNew(SBorder)
+					.BorderBackgroundColor(FLinearColor(1, 1, 1, 0.2f))
+					.BorderImage(FDynamicMaterialEditorStyle::GetBrush("Border.Left"))
+				]
+			]
+			+ SHorizontalBox::Slot()
+			.FillWidth(1.0f)
+			[
+				SAssignNew(EffectsList, SDMLayerEffectsView, SlotWidgetWeak.Pin(), EffectStack)
+			]
+		]
+
+		+ SVerticalBox::Slot()
+		.AutoHeight()
+		.HAlign(HAlign_Fill)
+		.VAlign(VAlign_Center)
+		[
+			SNew(SBox)
+			.HAlign(EHorizontalAlignment::HAlign_Center)
+			.Padding(10.f, 3.f, 5.f, 3.f)
+			.Visibility(this, &SDMSlotLayerItem::GetNewEffectVisibility)
+			[
+				CreateNewEffectButton()
+			]
+		];
+}
+
+TSharedRef<SWidget> SDMSlotLayerItem::CreateHeaderRowContent()
+{
+	return 
+		SNew(SHorizontalBox)
+		+ SHorizontalBox::Slot()
+		.AutoWidth()
+		.HAlign(HAlign_Left)
+		.VAlign(VAlign_Center)
+		.Padding(5.0f, 0.0f, 0.0f, 0.0f)
+		[
+			CreateLayerBypassButton()
+		]
+		+ SHorizontalBox::Slot()
+		.FillWidth(1.0f)
+		.HAlign(HAlign_Center)
+		.VAlign(VAlign_Center)
+		[
+			SNew(SHorizontalBox)
+			+ SHorizontalBox::Slot()
+			.AutoWidth()
+			.HAlign(HAlign_Center)
+			.VAlign(VAlign_Center)
+			.Padding(0.0f, 0.0f, 3.0f, 0.0f)
+			[
+				CreateLayerBaseToggleButton()
+			]
+			+ SHorizontalBox::Slot()
+			.AutoWidth()
+			.HAlign(HAlign_Center)
+			.VAlign(VAlign_Center)
+			.Padding(0.0f, 3.0f)
+			[
+				CreateStageBaseWidget(true, false,
+					TAttribute<FVector2D>::CreateSP(this, &SDMSlotLayerItem::GetStagePreviewSize),
+					FDMOnStageUniformSizeChanged::CreateSP(this, &SDMSlotLayerItem::SaveLayerPreviewSize))
+			]
+		]
+		+ SHorizontalBox::Slot()
+		.AutoWidth()
+		.HAlign(HAlign_Center)
+		.VAlign(VAlign_Center)
+		[
+			SNew(SHorizontalBox)
+			+ SHorizontalBox::Slot()
+			.AutoWidth()
+			.HAlign(HAlign_Center)
+			.VAlign(VAlign_Center)
+			[
+				CreateLayerLinkToggleButton()
+			]
+		]
+		+ SHorizontalBox::Slot()
+		.FillWidth(1.0f)
+		.HAlign(HAlign_Center)
+		.VAlign(VAlign_Center)
+		[
+			SNew(SHorizontalBox)
+			+ SHorizontalBox::Slot()
+			.AutoWidth()
+			.HAlign(HAlign_Center)
+			.VAlign(VAlign_Center)
+			.Padding(0.0f, 0.0f, 3.0f, 0.0f)
+			[
+				CreateLayerMaskToggleButton()
+			]
+			+ SHorizontalBox::Slot()
+			.AutoWidth()
+			.HAlign(HAlign_Center)
+			.VAlign(VAlign_Center)
+			.Padding(0.0f, 3.0f)
+			[
+				CreateStageMaskWidget(true, false,
+					TAttribute<FVector2D>::CreateSP(this, &SDMSlotLayerItem::GetStagePreviewSize),
+					FDMOnStageUniformSizeChanged::CreateSP(this, &SDMSlotLayerItem::SaveLayerPreviewSize))
+			]
+		]
+		+ SHorizontalBox::Slot()
+		.FillWidth(1.0f)
+		.HAlign(HAlign_Center)
+		.VAlign(VAlign_Fill)
+		.Padding(2.0f, 0.0f, 10.0f, 0.0f)
+		[
+			SNew(SVerticalBox)
+			+ SVerticalBox::Slot()
+			.FillHeight(1.0f)
+			.HAlign(HAlign_Left)
+			.VAlign(VAlign_Center)
+			.Padding(0.0f, 2.0f, 0.0f, 0.0f)
+			[
+				SAssignNew(LayerHeaderTextContainer, SBox)
+				[
+					CreateLayerHeaderText()
+				]
+			]
+			+ SVerticalBox::Slot()
+			.FillHeight(1.0f)
+			.HAlign(HAlign_Left)
+			.VAlign(VAlign_Center)
+			.Padding(0.0f, 0.0f, 0.0f, 0.0f)
+			[
+				SNew(STextBlock)
+				.Font(FAppStyle::GetFontStyle("PropertyWindow.NormalFont"))
+				.Text(this, &SDMSlotLayerItem::GetBlendModeText)
+			]
+			+ SVerticalBox::Slot()
+			.FillHeight(1.0f)
+			.HAlign(HAlign_Left)
+			.VAlign(VAlign_Center)
+			.Padding(0.0f, 0.0f, 0.0f, 2.0f)
+			[
+				SNew(STextBlock)
+				.Font(FAppStyle::GetFontStyle("PropertyWindow.NormalFont"))
+				.Text(this, &SDMSlotLayerItem::GetStageDescription)
+			]
+		]
+		+ SHorizontalBox::Slot()
+		.AutoWidth()
+		.HAlign(HAlign_Right)
+		.VAlign(VAlign_Center)
+		.Padding(0.0f, 0.0f, 10.0f, 0.0f)
+		[
+			CreateEffectsToggleButton()
+		];
+}
+
+TSharedRef<SWidget> SDMSlotLayerItem::CreateStageBaseWidget(const bool bInteractable, const bool bInShowTextOverlays, TAttribute<FVector2D> InDesiredSize, FDMOnStageUniformSizeChanged InOnUniformSizeChanged)
+{
+	FProperty* PreviewSizeProperty = UDynamicMaterialEditorSettings::StaticClass()->FindPropertyByName(GET_MEMBER_NAME_CHECKED(UDynamicMaterialEditorSettings, LayerPreviewSize));
+	ensure(PreviewSizeProperty);
+	const int32 ClampMin = PreviewSizeProperty ? PreviewSizeProperty->GetIntMetaData("ClampMin") : 0;
+	const int32 ClampMax = PreviewSizeProperty ? PreviewSizeProperty->GetIntMetaData("ClampMax") : 1;
+
+	const UDMMaterialLayerObject* Layer = GetLayer();
+
+	if (!Layer)
+	{
+		return SNullWidget::NullWidget;
+	}
+
+	return 
+		SAssignNew(BaseStageWidget, SDMStage, Layer->GetStage(EDMMaterialLayerStage::Base))
+		.Interactable(bInteractable)
+		.StageEnabled(StageBaseEnabled)
+		.StageSelected(StageBaseSelected)
+		.StageIsMask(false)
+		.ShowTextOverlays(bInShowTextOverlays)
+		.DesiredSize(InDesiredSize)
+		.OnClick(this, &SDMSlotLayerItem::OnStageClick)
+		.OnCanAcceptDrop(this, &SDMSlotLayerItem::OnStageCanAcceptDrop)
+		.OnAcceptDrop(this, &SDMSlotLayerItem::OnStageAcceptDrop)
+		.MinUniformSize(ClampMin)
+		.MaxUniformSize(ClampMax)
+		.OnUniformSizeChanged(InOnUniformSizeChanged);
+}
+
+TSharedRef<SWidget> SDMSlotLayerItem::CreateStageMaskWidget(const bool bInteractable, const bool bInShowTextOverlays, TAttribute<FVector2D> InDesiredSize, FDMOnStageUniformSizeChanged InOnUniformSizeChanged)
+{
+	FProperty* PreviewSizeProperty = UDynamicMaterialEditorSettings::StaticClass()->FindPropertyByName(GET_MEMBER_NAME_CHECKED(UDynamicMaterialEditorSettings, LayerPreviewSize));
+	ensure(PreviewSizeProperty);
+	const int32 ClampMin = PreviewSizeProperty ? PreviewSizeProperty->GetIntMetaData("ClampMin") : 0;
+	const int32 ClampMax = PreviewSizeProperty ? PreviewSizeProperty->GetIntMetaData("ClampMax") : 1;
+
+	const UDMMaterialLayerObject* Layer = GetLayer();
+
+	if (!Layer)
+	{
+		return SNullWidget::NullWidget;
+	}
+
+	return
+		SAssignNew(MaskStageWidget, SDMStage, Layer->GetStage(EDMMaterialLayerStage::Mask))
+		.Interactable(bInteractable)
+		.StageEnabled(StageMaskEnabled)
+		.StageSelected(StageMaskSelected)
+		.StageIsMask(true)
+		.ShowTextOverlays(bInShowTextOverlays)
+		.DesiredSize(InDesiredSize)
+		.OnClick(this, &SDMSlotLayerItem::OnStageClick)
+		.OnCanAcceptDrop(this, &SDMSlotLayerItem::OnStageCanAcceptDrop)
+		.OnAcceptDrop(this, &SDMSlotLayerItem::OnStageAcceptDrop)
+		.MinUniformSize(ClampMin)
+		.MaxUniformSize(ClampMax)
+		.OnUniformSizeChanged(InOnUniformSizeChanged);
+}
+
+TSharedRef<SWidget> SDMSlotLayerItem::CreateHandleWidget()
+{
+	TSharedRef<SBox> LayerIndexTextHandleWidget = 
+		SNew(SBox)
+		.HAlign(HAlign_Fill)
+		.VAlign(VAlign_Fill)
+		.Cursor(EMouseCursor::GrabHand)
+		.ToolTipText(this, &SDMSlotLayerItem::GetToolTipText)
+		[
+			SNew(SBorder)
+			.HAlign(HAlign_Fill)
+			.VAlign(VAlign_Fill)
+			.BorderImage(this, &SDMSlotLayerItem::GetRowHandleBrush)
+			[
+				SNew(SBox)
+				.HAlign(HAlign_Center)
+				.VAlign(VAlign_Center)
+				[
+					SNew(STextBlock)
+					.TextStyle(&FDynamicMaterialEditorStyle::Get(), "LayerView.Row.HeaderText.Small")
+					.Justification(ETextJustify::Center)
+					.Text(this, &SDMSlotLayerItem::GetLayerIndexText)
+				]
+			]
+		];
+
+	// Make sure all the index numbers align between single and double digit values.
+	const float HandleThickness = 20.0f;
+
+	LayerIndexTextHandleWidget->SetHeightOverride(HandleThickness);
+
+	return LayerIndexTextHandleWidget;
+}
+
+TSharedRef<SWidget> SDMSlotLayerItem::CreateLayerBypassButton()
+{
+	return 
+		SNew(SButton)
+		.ContentPadding(4.0f)
+		.ButtonStyle(FDynamicMaterialEditorStyle::Get(), "HoverHintOnly")
+		.ToolTipText(LOCTEXT("LayerBypassTooltip", "Toggle the bypassing of this layer."))
+		.OnClicked(this, &SDMSlotLayerItem::OnCreateLayerBypassButtonClicked)
+		[
+			SNew(SImage)
+			.DesiredSizeOverride(FVector2D(16.0f))
+			.Image(this, &SDMSlotLayerItem::GetCreateLayerBypassButtonImage)
+		];
+}
+
+TSharedRef<SWidget> SDMSlotLayerItem::CreateTogglesWidget()
+{
+	return
+		SNew(SHorizontalBox)
+		+ SHorizontalBox::Slot()
+		.HAlign(HAlign_Center)
+		.VAlign(VAlign_Fill)
+		[
+			SNew(SBox)
+			.HAlign(HAlign_Center)
+			.VAlign(VAlign_Center)
+			[
+				CreateLayerBaseToggleButton()
+			]
+		]
+		+ SHorizontalBox::Slot()
+		.HAlign(HAlign_Center)
+		.VAlign(VAlign_Fill)
+		[
+			SNew(SBox)
+			.HAlign(HAlign_Center)
+			.VAlign(VAlign_Center)
+			[
+				CreateLayerMaskToggleButton()
+			]
+		];
+}
+
+TSharedRef<SWidget> SDMSlotLayerItem::CreateLayerBaseToggleButton()
+{
+	return
+		SNew(SButton)
+		.ContentPadding(0.0f)
+		.ButtonStyle(FDynamicMaterialEditorStyle::Get(), "HoverHintOnly")
+		.ToolTipText(LOCTEXT("MaterialLayerBaseToggleTooltip", "Toggle the layer base on and off.\n\n"
+			"Warning: Toggling a layer base off may result in inputs being reset where incompatibilities are found.\n\n"
+			"The base of the first layer cannot be toggled off."))
+		.OnClicked(this, &SDMSlotLayerItem::OnBaseToggleButtonClicked)
+		[
+			SNew(SImage)
+			.Image(this, &SDMSlotLayerItem::GetBaseToggleButtonImage)
+		];
+}
+
+TSharedRef<SWidget> SDMSlotLayerItem::CreateLayerMaskToggleButton()
+{
+	return
+		SNew(SButton)
+		.ContentPadding(0.0f)
+		.ButtonStyle(FDynamicMaterialEditorStyle::Get(), "HoverHintOnly")
+		.ToolTipText(LOCTEXT("MaterialLayerMaskToggleTooltip", "Toggle the layer mask on and off."))
+		.OnClicked(this, &SDMSlotLayerItem::OnLayerMaskToggleButtonClicked)
+		[
+			SNew(SImage)
+			.Image(this, &SDMSlotLayerItem::GetLayerMaskToggleButtonImage)
+		];
+}
+
+TSharedRef<SWidget> SDMSlotLayerItem::CreateLayerLinkToggleButton()
+{
+	return 
+		SNew(SButton)
+		.ContentPadding(0.0f)
+		.ButtonStyle(FDynamicMaterialEditorStyle::Get(), "HoverHintOnly")
+		.ToolTipText(LOCTEXT("MaterialStageLinkTooltip", "Click to toggle UV Link on and off."))
+		.OnClicked(this, &SDMSlotLayerItem::OnLayerLinkToggleButton)
+		[
+			SNew(SImage)
+			.Image(this, &SDMSlotLayerItem::GetLayerLinkToggleButtonImage)
+			.Visibility(this, &SDMSlotLayerItem::GetLayerLinkToggleButtonVisibility)
+		];
+}
+
+TSharedRef<SWidget> SDMSlotLayerItem::CreateEffectsToggleButton()
+{
+	return 
+		SNew(SButton)
+		.ContentPadding(0.0f)
+		.ButtonStyle(FDynamicMaterialEditorStyle::Get(), "HoverHintOnly")
+		.ToolTipText(LOCTEXT("MaterialLayerFxTooltip", "Toggle the visibility of this slot layer's effects.\n\nNot implemented."))
+		.OnClicked(this, &SDMSlotLayerItem::OnEffectsToggleButtonClicked)
+		[
+			SNew(SImage)
+			.Image(this, &SDMSlotLayerItem::GetEffectsToggleButtonImage)
+		];
+}
+
+END_SLATE_FUNCTION_BUILD_OPTIMIZATION
+
+UDMMaterialStage* SDMSlotLayerItem::GetBaseStage() const
+{
+	if (const UDMMaterialLayerObject* Layer = GetLayer())
+	{
+		return Layer->GetStage(EDMMaterialLayerStage::Base);
+	}
+
+	return nullptr;
+}
+
+UDMMaterialStage* SDMSlotLayerItem::GetMaskStage() const
+{
+	if (const UDMMaterialLayerObject* Layer = GetLayer())
+	{
+		return Layer->GetStage(EDMMaterialLayerStage::Mask);
+	}
+
+	return nullptr;
+}
+
+UDMMaterialSlot* SDMSlotLayerItem::GetSlot() const
+{
+	if (UDMMaterialStage* const Stage = GetBaseStage())
+	{
+		if (UDMMaterialLayerObject* Layer = Stage->GetLayer())
+		{
+			return Layer->GetSlot();
+		}
+	}
+
+	return nullptr;
+}
+
+bool SDMSlotLayerItem::AreStagesLinked() const
+{
+	if (const UDMMaterialLayerObject* Layer = GetLayer())
+	{
+		return Layer->IsTextureUVLinkEnabled();
+	}
+
+	return false;
+}
+
+bool SDMSlotLayerItem::IsBaseStageEnabled() const
+{
+	if (const UDMMaterialLayerObject* Layer = GetLayer())
+	{
+		return Layer->IsStageEnabled(EDMMaterialLayerStage::Base);
+	}
+
+	return false;
+}
+
+bool SDMSlotLayerItem::IsMaskStageEnabled() const
+{
+	if (const UDMMaterialLayerObject* Layer = GetLayer())
+	{
+		return Layer->IsStageEnabled(EDMMaterialLayerStage::Mask);
+	}
+
+	return false;
+}
+
+FCursorReply SDMSlotLayerItem::OnCursorQuery(const FGeometry& InMyGeometry, const FPointerEvent& InCursorEvent) const
+{
+	if (InCursorEvent.IsMouseButtonDown(EKeys::LeftMouseButton))
+	{
+		return FCursorReply::Cursor(EMouseCursor::GrabHandClosed);
+	}
+	else
+	{
+		return FCursorReply::Cursor(EMouseCursor::GrabHand);
+	}
+}
+
+FReply SDMSlotLayerItem::OnMouseButtonDown(const FGeometry& InMyGeometry, const FPointerEvent& InMouseEvent)
+{
+	const FKey MouseButton = InMouseEvent.GetEffectingButton();
+
+	if (MouseButton == EKeys::LeftMouseButton)
+	{
+
+	}
+
+	return STableRow<TSharedPtr<FDMMaterialLayerReference>>::OnMouseButtonDown(InMyGeometry, InMouseEvent);
+}
+
+int32 SDMSlotLayerItem::OnLayerItemPaintDropIndicator(EItemDropZone InItemDropZone, const FPaintArgs& Args, const FGeometry& AllottedGeometry, const FSlateRect& MyCullingRect, FSlateWindowElementList& OutDrawElements, int32 LayerId, const FWidgetStyle& InWidgetStyle, bool bParentEnabled) const
+{
+	const FSlateBrush* DropIndicatorBrush = GetDropIndicatorBrush(InItemDropZone);
+	static float OffsetX = 10.0f;
+	FVector2D Offset(OffsetX * GetIndentLevel(), 0.0f);
+
+	FSlateDrawElement::MakeBox(
+		OutDrawElements,
+		LayerId++,
+		AllottedGeometry.ToPaintGeometry(FVector2D(AllottedGeometry.GetLocalSize() - Offset), FSlateLayoutTransform(Offset)),
+		DropIndicatorBrush,
+		ESlateDrawEffect::None,
+		DropIndicatorBrush->GetTint(InWidgetStyle) * InWidgetStyle.GetColorAndOpacityTint()
+	);
+
+	return LayerId;
+}
+
+TOptional<EItemDropZone> SDMSlotLayerItem::OnLayerItemCanAcceptDrop(const FDragDropEvent& InDragDropEvent, EItemDropZone InDropZone, 
+	TSharedPtr<FDMMaterialLayerReference> InSlotLayer) const
+{
+	if (!InSlotLayer.IsValid())
+	{
+		return TOptional<EItemDropZone>();
+	}
+
+	const UDMMaterialLayerObject* Layer = GetLayer();
+
+	if (!Layer)
+	{
+		return TOptional<EItemDropZone>();
+	}
+
+	TSharedPtr<FDMSlotLayerDragDropOperation> SlotLayerDragDropOp = InDragDropEvent.GetOperationAs<FDMSlotLayerDragDropOperation>();
+	if (!SlotLayerDragDropOp.IsValid())
+	{
+		return TOptional<EItemDropZone>();
+	}
+
+	UDMMaterialLayerObject* DraggedSlotLayer = SlotLayerDragDropOp->GetLayer();
+	if (!DraggedSlotLayer)
+	{
+		return TOptional<EItemDropZone>();
+	}
+
+	SlotLayerDragDropOp->SetToInvalidDropLocation();
+
+	switch (InDropZone)
+	{
+	case EItemDropZone::AboveItem:
+		if (Layer->CanMoveLayerAbove(DraggedSlotLayer))
+		{
+			SlotLayerDragDropOp->SetToValidDropLocation();
+			return InDropZone;
+		}
+		break;
+
+	case EItemDropZone::OntoItem:
+	case EItemDropZone::BelowItem:
+		if (Layer->CanMoveLayerBelow(DraggedSlotLayer))
+		{
+			SlotLayerDragDropOp->SetToValidDropLocation();
+			return EItemDropZone::BelowItem;
+		}
+		break;
+	}
+
+	return TOptional<EItemDropZone>();
+}
+
+void SDMSlotLayerItem::OnLayerItemDragEnter(const FDragDropEvent& InDragDropEvent)
+{
+
+}
+
+void SDMSlotLayerItem::OnLayerItemDragLeave(const FDragDropEvent& InDragDropEvent)
+{
+
+}
+
+FReply SDMSlotLayerItem::OnLayerItemDragDetected(const FGeometry& InMyGeometry, const FPointerEvent& InMouseEvent)
+{
+	const bool bShouldDuplicate = InMouseEvent.IsAltDown();
+
+	TSharedRef<FDMSlotLayerDragDropOperation> SlotLayerDragDropOp = MakeShared<FDMSlotLayerDragDropOperation>(SharedThis(this), bShouldDuplicate);
+
+	return FReply::Handled().BeginDragDrop(SlotLayerDragDropOp);
+}
+
+FReply SDMSlotLayerItem::OnLayerItemAcceptDrop(const FDragDropEvent& InDragDropEvent, EItemDropZone InDropZone, 
+	TSharedPtr<FDMMaterialLayerReference> InSlotLayer)
+{
+	TSharedPtr<FDMSlotLayerDragDropOperation> SlotLayerDragDropOp = InDragDropEvent.GetOperationAs<FDMSlotLayerDragDropOperation>();
+
+	if (SlotLayerDragDropOp.IsValid())
+	{
+		const UDMMaterialLayerObject* DraggedOverLayer = InSlotLayer.IsValid() ? InSlotLayer->GetLayer() : nullptr;
+		if (!DraggedOverLayer || !DraggedOverLayer->GetStage(EDMMaterialLayerStage::Base))
+		{
+			return FReply::Handled();
+		}
+
+		UDMMaterialLayerObject* DraggedLayer = SlotLayerDragDropOp->GetLayer();
+		if (!DraggedLayer || !DraggedLayer->GetStage(EDMMaterialLayerStage::Base))
+		{
+			return FReply::Handled();
+		}
+
+		const int32 ThisLayerIndex = GetLayerItemIndex();
+		if (ThisLayerIndex == INDEX_NONE)
+		{
+			return FReply::Handled();
+		}
+		
+		UDMMaterialSlot* const Slot = DraggedLayer->GetSlot();
+		if (!Slot)
+		{
+			return FReply::Handled();
+		}
+
+		{
+			FScopedTransaction Transaction(LOCTEXT("MoveLayer", "Move Material Designer Layer"));
+			Slot->Modify();
+			Slot->MoveLayer(DraggedLayer, ThisLayerIndex);
+
+			if (TSharedPtr<STableViewBase> LayerView = LayerViewWeak.Pin())
+			{
+				LayerView->RequestListRefresh();
+
+				if (TSharedPtr<SDMSlot> SlotWidget = StaticCastSharedPtr<SDMSlotLayerView>(LayerView)->GetSlotWidget())
+				{
+					SlotWidget->InvalidateSlotSettingsRowWidget();
+				}
+			}
+		}
+
+		return FReply::Handled();
+	}
+	return FReply::Unhandled();
+}
+
+FText SDMSlotLayerItem::GetToolTipText() const
+{
+	const UDMMaterialLayerObject* Layer = GetLayer();
+	if (!Layer)
+	{
+		return FText();
+	}
+
+	return Layer->GetComponentDescription();
+}
+
+FText SDMSlotLayerItem::GetLayerHeaderText() const
+{
+	const UDMMaterialLayerObject* Layer = GetLayer();
+	if (!Layer)
+	{
+		return FText();
+	}
+
+	return Layer->GetComponentDescription();
+}
+
+FText SDMSlotLayerItem::GetLayerIndexText() const
+{
+	return FText::Format(LOCTEXT("LayerIndexText", "{0}"), GetLayerItemIndex());
+}
+
+FText SDMSlotLayerItem::GetBlendModeText() const
+{
+	if (UDMMaterialLayerObject* Layer = GetLayer())
+	{
+		if (UDMMaterialSlot* Slot = Layer->GetSlot())
+		{
+			if (UDynamicMaterialModelEditorOnlyData* ModelEditorOnlyData = Slot->GetMaterialModelEditorOnlyData())
+			{
+				if (UDMMaterialProperty* Property = ModelEditorOnlyData->GetMaterialProperty(Layer->GetMaterialProperty()))
+				{
+					return Property->GetDescription();
+				}
+			}
+		}
+	}
+
+	return LOCTEXT("Error", "Error");
+}
+
+FText SDMSlotLayerItem::GetStageDescription() const
+{
+	if (UDMMaterialLayerObject* Layer = GetLayer())
+	{
+		if (UDMMaterialStage* Stage = Layer->GetFirstEnabledStage(EDMMaterialLayerStage::All))
+		{
+			return Stage->GetComponentDescription();
+		}
+	}
+
+	return GetDefault<UDMMaterialStage>()->GetComponentDescription();
+}
+
+int32 SDMSlotLayerItem::GetLayerItemIndex() const
+{
+	UDMMaterialLayerObject* Layer = GetLayer();
+
+	if (!Layer)
+	{
+		return INDEX_NONE;
+	}
+
+	return Layer->FindIndex();
+}
+
+bool SDMSlotLayerItem::IsRootLayer() const
+{
+	return GetLayerItemIndex() == 0;
+}
+
+const FSlateBrush* SDMSlotLayerItem::GetRowHandleBrush() const
+{
+	static const FSlateBrush* const Default = FDynamicMaterialEditorStyle::GetBrush("LayerView.Row.Handle.Left");
+	static const FSlateBrush* const Selected = FDynamicMaterialEditorStyle::GetBrush("LayerView.Row.Handle.Left.Select");
+	static const FSlateBrush* const Hovered = FDynamicMaterialEditorStyle::GetBrush("LayerView.Row.Handle.Left.Hover");
+	static const FSlateBrush* const SelectedHovered = FDynamicMaterialEditorStyle::GetBrush("LayerView.Row.Handle.Left.Select.Hover");
+
+	const bool bSelected = IsSelected();
+	const bool bHovered = IsHovered();
+
+	if (bSelected && bHovered)
+	{
+		return SelectedHovered;
+	}
+	else if (bSelected)
+	{
+		return Selected;
+	}
+	else if (bHovered)
+	{
+		return Hovered;
+	}
+	else
+	{
+		return Default;
+	}
+}
+
+void SDMSlotLayerItem::OnStageClick(const FPointerEvent& InMouseEvent, const TSharedRef<SDMStage>& InStageWidget)
+{
+	UDMMaterialStage* const ClickedStage = InStageWidget->GetStage();
+
+	if (!ClickedStage)
+	{
+		return;
+	}
+
+	const FKey MouseButton = InMouseEvent.GetEffectingButton();
+
+	if (MouseButton == EKeys::LeftMouseButton)
+	{
+		if (TSharedPtr<STableViewBase> TableViewBase = LayerViewWeak.Pin())
+		{
+			TSharedPtr<SListView<TSharedPtr<FDMMaterialLayerReference>>> LayerView = StaticCastSharedPtr<SListView<TSharedPtr<FDMMaterialLayerReference>>>(TableViewBase);
+			LayerView->SetSelection(LayerItem);
+		}
+
+		if (SelectedStageWidget.IsValid())
+		{
+			OnStageSelected.ExecuteIfBound(false, SelectedStageWidget.ToSharedRef());
+		}
+
+		SelectedStageWidget = InStageWidget;
+		OnStageSelected.ExecuteIfBound(true, InStageWidget);
+	}
+}
+
+bool SDMSlotLayerItem::OnStageCanAcceptDrop(const FDragDropEvent& InDragDropEvent, const TSharedRef<SDMStage>& InStageWidget) const
+{
+	return true;
+}
+
+FReply SDMSlotLayerItem::OnStageAcceptDrop(const FDragDropEvent& InDragDropEvent, const TSharedRef<SDMStage>& InStageWidget)
+{
+	const TSharedPtr<FDMSlotLayerDragDropOperation> SlotLayerDragDropOp = InDragDropEvent.GetOperationAs<FDMSlotLayerDragDropOperation>();
+	if (!SlotLayerDragDropOp.IsValid())
+	{
+		return FReply::Unhandled();
+	}
+
+	auto ReturnInvalidHandled = [&SlotLayerDragDropOp]()
+	{
+		SlotLayerDragDropOp->SetToInvalidDropLocation();
+		return FReply::Handled();
+	};
+
+	const UDMMaterialLayerObject* Layer = GetLayer();
+
+	if (!Layer)
+	{
+		return ReturnInvalidHandled();
+	}
+
+	UDMMaterialLayerObject* DraggedLayerItem = SlotLayerDragDropOp->GetLayer();
+	if (!DraggedLayerItem || !Layer->GetStage(EDMMaterialLayerStage::Base))
+	{
+		return ReturnInvalidHandled();
+	}
+
+	if (Layer == DraggedLayerItem)
+	{
+		SlotLayerDragDropOp->SetToValidDropLocation();
+		return FReply::Handled();
+	}
+
+	UDMMaterialStage* ThisBaseStage = Layer->GetStage(EDMMaterialLayerStage::Base);
+	UDMMaterialStage* DraggedBaseStage = DraggedLayerItem->GetStage(EDMMaterialLayerStage::Base);
+	if (!ThisBaseStage || !DraggedBaseStage)
+	{
+		return ReturnInvalidHandled();
+	}
+
+	const int32 ThisBaseStageIndex = Layer->FindIndex();
+	const int32 DraggedBaseStageIndex = DraggedLayerItem->FindIndex();
+	if (ThisBaseStageIndex == INDEX_NONE || DraggedBaseStageIndex == INDEX_NONE)
+	{
+		return ReturnInvalidHandled();
+	}
+
+	const UDMMaterialSlot* const ThisBaseSlot = Layer->GetSlot();
+	const UDMMaterialSlot* const DraggedBaseSlot = DraggedLayerItem->GetSlot();
+	if (!ThisBaseSlot || !DraggedBaseSlot)
+	{
+		return ReturnInvalidHandled();
+	}
+
+	if (ThisBaseSlot != DraggedBaseSlot)
+	{
+		return ReturnInvalidHandled();
+	}
+
+	const UDMMaterialLayerObject* PreviousLayer = Layer->GetPreviousLayer(EDMMaterialPropertyType::Any, EDMMaterialLayerStage::None);
+	const UDMMaterialLayerObject* NextLayer = DraggedLayerItem->GetNextLayer(EDMMaterialPropertyType::Any, EDMMaterialLayerStage::None);
+
+	if (PreviousLayer || NextLayer)
+	{
+		if (ThisBaseStageIndex < DraggedBaseStageIndex)
+		{
+			if (PreviousLayer)
+			{
+				const bool bIsCompatibleWithBase = DraggedBaseStage->IsCompatibleWithPreviousStage(PreviousLayer->GetStage(EDMMaterialLayerStage::Base));
+				const bool bIsCompatibleWithNextStage = DraggedBaseStage->IsCompatibleWithNextStage(ThisBaseStage);
+				if (!bIsCompatibleWithBase || !bIsCompatibleWithNextStage)
+				{
+					return ReturnInvalidHandled();
+				}
+			}
+		}
+		else
+		{
+			if (NextLayer)
+			{
+				const bool bIsCompatibleWithPreviousStage = DraggedBaseStage->IsCompatibleWithPreviousStage(ThisBaseStage);
+				const bool bIsCompatibleWithNextStage = DraggedBaseStage->IsCompatibleWithNextStage(NextLayer->GetStage(EDMMaterialLayerStage::Base));
+				if (!bIsCompatibleWithPreviousStage || !bIsCompatibleWithNextStage)
+				{
+					return ReturnInvalidHandled();
+				}
+			}
+		}
+
+		PreviousLayer = DraggedLayerItem->GetPreviousLayer(EDMMaterialPropertyType::Any, EDMMaterialLayerStage::None);
+		NextLayer = DraggedLayerItem->GetNextLayer(EDMMaterialPropertyType::Any, EDMMaterialLayerStage::None);
+		if (PreviousLayer && NextLayer)
+		{
+			if (!PreviousLayer->GetStage(EDMMaterialLayerStage::Base)->IsCompatibleWithNextStage(NextLayer->GetStage(EDMMaterialLayerStage::Base)))
+			{
+				return ReturnInvalidHandled();
+			}
+			if (!NextLayer->GetStage(EDMMaterialLayerStage::Base)->IsCompatibleWithPreviousStage(PreviousLayer->GetStage(EDMMaterialLayerStage::Base)))
+			{
+				return ReturnInvalidHandled();
+			}
+		}
+	}
+
+	UDMMaterialSlot* Slot = DraggedLayerItem->GetSlot();
+	if (!Slot)
+	{
+		return FReply::Unhandled();
+	}
+
+	{
+		FScopedTransaction Transaction(LOCTEXT("MoveLayer", "Move Material Designer Layer"));
+		Slot->Modify();
+		Slot->MoveLayer(DraggedLayerItem, ThisBaseStageIndex);
+	}
+
+	return FReply::Handled();
+}
+
+void SDMSlotLayerItem::SaveLayerPreviewSize(const float InNewSize, const TSharedRef<SDMStage>& InStageWidget)
+{
+	UDynamicMaterialEditorSettings* Settings = UDynamicMaterialEditorSettings::Get();
+	Settings->LayerPreviewSize = InNewSize;
+	Settings->SaveConfig();
+}
+
+EVisibility SDMSlotLayerItem::GetEffectsListVisibility() const
+{
+	return EffectsList->GetLayerItemCount() > 0 && bDisplayEffectsList ? EVisibility::Visible : EVisibility::Collapsed;
+}
+
+EVisibility SDMSlotLayerItem::GetNewEffectVisibility() const
+{
+	return bDisplayEffectsList ? EVisibility::Visible : EVisibility::Collapsed;
+}
+
+TSharedRef<SWidget> SDMSlotLayerItem::CreateNewEffectButton()
+{
+	TSharedRef<SWidget> TextBlock = SNew(STextBlock)
+		.Font(FAppStyle::GetFontStyle("PropertyWindow.NormalFont"))
+		.Text(LOCTEXT("AddNewEffect", "Add New Effect"))
+		.ColorAndOpacity(FSlateColor(EStyleColor::PrimaryHover));
+
+	TextBlock->SetOnMouseButtonDown(FPointerEventHandler::CreateSPLambda(
+		TextBlock,
+		[this, TextBlockWeak = TextBlock->AsShared().ToWeakPtr()](const FGeometry& InGeometry, const FPointerEvent& InPointerEvent)
+		{
+			if (TSharedPtr<SWidget> TextBlock = TextBlockWeak.Pin())
+			{
+				FSlateApplication::Get().PushMenu(
+					TextBlock.ToSharedRef(),
+					FWidgetPath(),
+					CreateNewEffectMenu(),
+					GetTickSpaceGeometry().GetAbsolutePosition(),
+					FPopupTransitionEffect::TopMenu
+				);
+			}
+
+			return FReply::Handled();
+		}
+	));
+
+	return TextBlock;
+}
+
+TSharedRef<SWidget> SDMSlotLayerItem::CreateNewEffectMenu()
+{
+	if (LayerItem.IsValid())
+	{
+		return FDMMaterialSlotLayerAddEffectMenus::OpenAddEffectMenu(LayerItem->GetLayer());
+	}
+
+	return SNullWidget::NullWidget;
+}
+
+FReply SDMSlotLayerItem::OnCreateLayerBypassButtonClicked()
+{
+	if (UDMMaterialLayerObject* Layer = GetLayer())
+	{
+		FScopedTransaction Transaction(LOCTEXT("ToggledLayerVisibility", "Toggle Material Designer Layer Visibility"));
+		Layer->Modify();
+		Layer->SetEnabled(!Layer->IsEnabled());
+
+		return FReply::Handled();
+	}
+
+	return FReply::Unhandled();
+}
+
+const FSlateBrush* SDMSlotLayerItem::GetCreateLayerBypassButtonImage() const
+{
+	static const FSlateBrush* const ExposeBrush = FCoreStyle::Get().GetBrush("Kismet.VariableList.ExposeForInstance");
+	static const FSlateBrush* const HideBrush = FCoreStyle::Get().GetBrush("Kismet.VariableList.HideForInstance");
+
+	if (const UDMMaterialLayerObject* Layer = GetLayer())
+	{
+		if (Layer->IsEnabled())
+		{
+			return ExposeBrush;
+		}
+	}
+
+	return HideBrush;
+}
+
+FReply SDMSlotLayerItem::OnBaseToggleButtonClicked()
+{
+	if (const UDMMaterialLayerObject* Layer = GetLayer())
+	{
+		constexpr EDMMaterialLayerStage StageType = EDMMaterialLayerStage::Base;
+
+		if (UDMMaterialStage* Stage = Layer->GetStage(StageType))
+		{
+			FScopedTransaction Transaction(LOCTEXT("ToggleBaseStageEnabled", "Toggle Material Designer Base Stage Enabled"));
+			Stage->Modify();
+
+			if (!Stage->SetEnabled(!Stage->IsEnabled()))
+			{
+				Transaction.Cancel();
+			}
+		}
+	}
+
+	return FReply::Handled();
+}
+
+const FSlateBrush* SDMSlotLayerItem::GetBaseToggleButtonImage() const
+{
+	static const FSlateBrush* Disabled = FDynamicMaterialEditorStyle::GetBrush("Icons.Stage.Disabled");
+	static const FSlateBrush* Enabled = FDynamicMaterialEditorStyle::GetBrush("Icons.Stage.Enabled");
+
+	if (!IsBaseStageEnabled())
+	{
+		return Disabled;
+	}
+
+	return Enabled;
+}
+
+FReply SDMSlotLayerItem::OnLayerMaskToggleButtonClicked()
+{
+	if (const UDMMaterialLayerObject* Layer = GetLayer())
+	{
+		constexpr EDMMaterialLayerStage StageType = EDMMaterialLayerStage::Mask;
+
+		if (UDMMaterialStage* Stage = Layer->GetStage(StageType))
+		{
+			FScopedTransaction Transaction(LOCTEXT("ToggleMaskStageEnabled", "Toggle Material Designer Mask Stage Enabled"));
+			Stage->Modify();
+
+			if (!Stage->SetEnabled(!Stage->IsEnabled()))
+			{
+				Transaction.Cancel();
+			}
+		}
+	}
+
+	return FReply::Handled();
+}
+
+const FSlateBrush* SDMSlotLayerItem::GetLayerMaskToggleButtonImage() const
+{
+	static const FSlateBrush* Disabled = FDynamicMaterialEditorStyle::GetBrush("Icons.Stage.Disabled");
+	static const FSlateBrush* Enabled = FDynamicMaterialEditorStyle::GetBrush("Icons.Stage.Enabled");
+
+	if (!IsMaskStageEnabled())
+	{
+		return Disabled;
+	}
+
+	return Enabled;
+}
+
+FReply SDMSlotLayerItem::OnLayerLinkToggleButton()
+{
+	if (UDMMaterialLayerObject* Layer = GetLayer())
+	{
+		FScopedTransaction Transaction(LOCTEXT("UVLayerLinkToggle", "Toggle Material Designer Layer UV Link"));
+		Layer->Modify();
+		Layer->ToggleTextureUVLinkEnabled();
+
+		if (SelectedStageWidget.IsValid())
+		{
+			OnLayerLinkToggled.ExecuteIfBound(Layer->IsTextureUVLinkEnabled(), SelectedStageWidget.ToSharedRef());
+		}
+	}
+
+	return FReply::Handled();
+}
+
+const FSlateBrush* SDMSlotLayerItem::GetLayerLinkToggleButtonImage() const
+{
+	static const FSlateBrush* const Unlinked = FDynamicMaterialEditorStyle::GetBrush("Icons.Stage.ChainUnlinked.Vertical");
+	static const FSlateBrush* const Linked = FDynamicMaterialEditorStyle::GetBrush("Icons.Stage.ChainLinked.Vertical");
+
+	return AreStagesLinked() ? Linked : Unlinked;
+}
+
+EVisibility SDMSlotLayerItem::GetLayerLinkToggleButtonVisibility() const
+{
+	if (const UDMMaterialLayerObject* Layer = GetLayer())
+	{
+		if (UDMMaterialStage* BaseStage = Layer->GetStage(EDMMaterialLayerStage::Base))
+		{
+			if (UDMMaterialStageSource* const StageSource = BaseStage->GetSource())
+			{
+				if (StageSource->IsA<UDMMaterialStageInputTextureUV>())
+				{
+					return EVisibility::Visible;
+				}
+
+				if (const UDMMaterialStageThroughput* const Throughput = Cast<UDMMaterialStageThroughput>(StageSource))
+				{
+					if (Throughput->SupportsLayerMaskTextureUVLink())
+					{
+						return EVisibility::Visible;
+					}
+				}
+			}
+		}
+	}
+
+	return EVisibility::Hidden;
+}
+
+FReply SDMSlotLayerItem::OnEffectsToggleButtonClicked()
+{
+	bDisplayEffectsList = !bDisplayEffectsList;
+
+	return FReply::Handled();
+}
+
+const FSlateBrush* SDMSlotLayerItem::GetEffectsToggleButtonImage() const
+{
+	static const FSlateBrush* Displayed = FDynamicMaterialEditorStyle::GetBrush("EffectsView.Row.Fx.Opened");
+	static const FSlateBrush* Hidden = FDynamicMaterialEditorStyle::GetBrush("EffectsView.Row.Fx.Closed");
+
+	if (bDisplayEffectsList)
+	{
+		return Displayed;
+	}
+
+	return Hidden;
+}
+
+FVector2D SDMSlotLayerItem::GetStagePreviewSize() const
+{
+	return FVector2D(PreviewSize.Get());
+}
+
+TSharedRef<SWidget> SDMSlotLayerItem::CreateLayerHeaderText() const
+{
+	TSharedRef<SWidget> TextBlock = SNew(STextBlock)
+		.Font(FAppStyle::GetFontStyle("PropertyWindow.NormalFont"))
+		.Text(this, &SDMSlotLayerItem::GetLayerHeaderText)
+		.ColorAndOpacity(FSlateColor(EStyleColor::PrimaryHover));
+
+	TextBlock->SetOnMouseButtonDown(FPointerEventHandler::CreateSPLambda(
+		this,
+		[this](const FGeometry& InGeometry, const FPointerEvent& InPointerEvent)
+		{
+			if (LayerHeaderTextContainer.IsValid())
+			{
+				TSharedRef<SWidget> EditableContent = CreateLayerHeaderEditableText();
+				LayerHeaderTextContainer->SetContent(EditableContent);
+				FSlateApplication::Get().SetKeyboardFocus(EditableContent);
+			}
+
+			return FReply::Handled();
+		}
+	));
+
+	return TextBlock;
+}
+
+TSharedRef<SWidget> SDMSlotLayerItem::CreateLayerHeaderEditableText() const
+{
+	FText LayerName = FText::GetEmpty();
+
+	if (const UDMMaterialStage* Stage = GetBaseStage())
+	{
+		if (const UDMMaterialLayerObject* Layer = Stage->GetLayer())
+		{
+			if (!Layer->GetLayerName().IsEmpty())
+			{
+				LayerName = Layer->GetLayerName();
+			}
+		}
+	}
+
+	return SNew(SEditableTextBox)
+		.Font(FAppStyle::GetFontStyle("PropertyWindow.NormalFont"))
+		.HintText(LOCTEXT("LayerName", "Layer Name"))
+		.IsEnabled(true)
+		.Text(LayerName)
+		.Style(FDynamicMaterialEditorStyle::Get(), "InlineEditableTextBoxStyle")
+		.OnTextCommitted(FOnTextCommitted::CreateSPLambda(
+			this,
+			[this](const FText& InText, ETextCommit::Type InCommitType)
+			{
+				if (LayerHeaderTextContainer.IsValid())
+				{
+					if (InCommitType != ETextCommit::OnUserMovedFocus && InCommitType != ETextCommit::OnCleared)
+					{
+						if (const UDMMaterialStage* Stage = GetBaseStage())
+						{
+							if (UDMMaterialLayerObject* Layer = Stage->GetLayer())
+							{
+								FScopedTransaction Transaction(LOCTEXT("ChangeLayerName", "Material Designer Change Layer Name"));
+								Layer->Modify();
+								Layer->SetLayerName(InText);
+							}
+						}
+					}
+
+					LayerHeaderTextContainer->SetContent(CreateLayerHeaderText());
+				}
+			}));
+}
+
+void SDMSlotLayerItem::DeselectAllEffects()
+{
+	EffectsList->ClearSelection();
+}
+
+#undef LOCTEXT_NAMESPACE
diff --git a/Engine/Plugins/Experimental/DynamicMaterial/Source/DynamicMaterialEditor/Private/Slate/Layers/SDMSlotLayerItem.h b/Engine/Plugins/Experimental/DynamicMaterial/Source/DynamicMaterialEditor/Private/Slate/Layers/SDMSlotLayerItem.h
new file mode 100644
--- /dev/null
+++ b/Engine/Plugins/Experimental/DynamicMaterial/Source/DynamicMaterialEditor/Private/Slate/Layers/SDMSlotLayerItem.h
@@ -0,0 +1,167 @@
+// Copyright Epic Games, Inc. All Rights Reserved.
+
+#pragma once
+
+#include "Components/DMMaterialSlot.h"
+#include "Slate/SDMStage.h"
+#include "Widgets/Views/SListView.h"
+#include "Widgets/Views/STableRow.h"
+
+class SBox;
+class SDMLayerEffectsView;
+class SDMSlot;
+class SDMStage;
+class UDMMaterialLayerObject;
+class UDMMaterialSlot;
+class UDMMaterialStage;
+class UMaterialFunctionInterface;
+struct FDMMaterialLayerReference;
+struct FSlateBrush;
+
+DECLARE_DELEGATE_TwoParams(FDMOnStageSelected, const bool, const TSharedRef<SDMStage>& /** StageWidget */);
+DECLARE_DELEGATE_TwoParams(FDMOnLayerLinkToggled, const bool, const TSharedRef<SDMStage>& /** StageWidget */);
+
+/**
+ * Material Slot Layer
+ * 
+ * Represents a single item in the layer view that can be dragged and re-arranged.
+ * Displays the base and mask stages, along with enable and link buttons.
+ */
+class SDMSlotLayerItem : public STableRow<TSharedPtr<FDMMaterialLayerReference>>
+{
+public:
+	SLATE_BEGIN_ARGS(SDMSlotLayerItem)
+		: _PreviewSize(40)
+		{}
+		SLATE_ATTRIBUTE(bool, StageBaseEnabled)
+		SLATE_ATTRIBUTE(bool, StageBaseSelected)
+		SLATE_ATTRIBUTE(bool, StageMaskEnabled)
+		SLATE_ATTRIBUTE(bool, StageMaskSelected)
+		SLATE_ATTRIBUTE(int32, PreviewSize)
+		SLATE_EVENT(FDMOnStageSelected, OnStageSelected)
+		SLATE_EVENT(FDMOnLayerLinkToggled, OnLayerLinkToggled)
+	SLATE_END_ARGS()
+
+	virtual ~SDMSlotLayerItem() {}
+
+	void Construct(const FArguments& InArgs, const TSharedPtr<SDMSlot>& InSlotWidget, const TSharedRef<STableViewBase>& InLayerView, 
+		const TSharedPtr<FDMMaterialLayerReference>& InLayerReferenceItem);
+
+	UDMMaterialLayerObject* GetLayer() const;
+
+	int32 GetLayerItemIndex() const;
+	bool IsRootLayer() const;
+
+	FORCEINLINE const TSharedPtr<SDMStage>& GetBaseStageWidget() const { return BaseStageWidget; }
+	FORCEINLINE const TSharedPtr<SDMStage>& GetMaskStageWidget() const { return MaskStageWidget; }
+
+	FORCEINLINE const TSharedPtr<SDMStage>& GetSelectedStageWidget() const { return SelectedStageWidget; }
+
+	FORCEINLINE bool AreEffectsExpanded() const { return bDisplayEffectsList; }
+
+	void DeselectAllEffects();
+
+protected:
+	TWeakPtr<SDMSlot> SlotWidgetWeak;
+	TWeakPtr<STableViewBase> LayerViewWeak;
+	TSharedPtr<FDMMaterialLayerReference> LayerItem;
+
+	TAttribute<bool> StageBaseEnabled;
+	TAttribute<bool> StageBaseSelected;
+	TAttribute<bool> StageMaskEnabled;
+	TAttribute<bool> StageMaskSelected;
+	TAttribute<int32> PreviewSize;
+	FDMOnStageSelected OnStageSelected;
+	FDMOnLayerLinkToggled OnLayerLinkToggled;
+
+	TSharedPtr<SDMStage> BaseStageWidget;
+	TSharedPtr<SDMStage> MaskStageWidget;
+
+	TSharedPtr<SDMStage> SelectedStageWidget;
+
+	TSharedPtr<SDMLayerEffectsView> EffectsList;
+	bool bDisplayEffectsList;
+
+	TSharedPtr<SBox> LayerHeaderTextContainer;
+
+	TSharedRef<SWidget> CreateMainContent();
+	TSharedRef<SWidget> CreateHeaderRowContent();
+
+	/** Create components of the slot layer. */
+	TSharedRef<SWidget> CreateStageBaseWidget(const bool bInteractable, const bool bInShowTextOverlays, TAttribute<FVector2D> InDesiredSize = FVector2D(40.0f), FDMOnStageUniformSizeChanged InOnUniformSizeChanged = nullptr);
+	TSharedRef<SWidget> CreateStageMaskWidget(const bool bInteractable, const bool bInShowTextOverlays, TAttribute<FVector2D> InDesiredSize = FVector2D(40.0f), FDMOnStageUniformSizeChanged InOnUniformSizeChanged = nullptr);
+	TSharedRef<SWidget> CreateHandleWidget();
+
+	TSharedRef<SWidget> CreateLayerBypassButton();
+	TSharedRef<SWidget> CreateTogglesWidget();
+	TSharedRef<SWidget> CreateLayerBaseToggleButton();
+	TSharedRef<SWidget> CreateLayerMaskToggleButton();
+	TSharedRef<SWidget> CreateLayerLinkToggleButton();
+	TSharedRef<SWidget> CreateEffectsToggleButton();
+
+	UDMMaterialStage* GetBaseStage() const;
+	UDMMaterialStage* GetMaskStage() const;
+
+	UDMMaterialSlot* GetSlot() const;
+
+	bool AreStagesLinked() const;
+	bool IsBaseStageEnabled() const;
+	bool IsMaskStageEnabled() const;
+
+	int32 OnLayerItemPaintDropIndicator(EItemDropZone InItemDropZone, const FPaintArgs& Args, const FGeometry& AllottedGeometry, 
+		const FSlateRect& MyCullingRect, FSlateWindowElementList& OutDrawElements, int32 LayerId, const FWidgetStyle& InWidgetStyle, bool bParentEnabled) const;
+	TOptional<EItemDropZone> OnLayerItemCanAcceptDrop(const FDragDropEvent& InDragDropEvent, EItemDropZone InDropZone, 
+		TSharedPtr<FDMMaterialLayerReference> InSlotLayer) const;
+	void OnLayerItemDragEnter(const FDragDropEvent& InDragDropEvent);
+	void OnLayerItemDragLeave(const FDragDropEvent& InDragDropEvent);
+	FReply OnLayerItemDragDetected(const FGeometry& InMyGeometry, const FPointerEvent& InMouseEvent);
+	FReply OnLayerItemAcceptDrop(const FDragDropEvent& InDragDropEvent, EItemDropZone InDropZone, 
+		TSharedPtr<FDMMaterialLayerReference> InSlotLayer);
+
+	FText GetToolTipText() const;
+	FText GetLayerHeaderText() const;
+	FText GetLayerIndexText() const;
+	FText GetBlendModeText() const;
+	FText GetStageDescription() const;
+	const FSlateBrush* GetRowHandleBrush() const;
+
+	//~ Begin SWidget
+	virtual FCursorReply OnCursorQuery(const FGeometry& InMyGeometry, const FPointerEvent& InCursorEvent) const override;
+	virtual FReply OnMouseButtonDown(const FGeometry& InMyGeometry, const FPointerEvent& InMouseEvent) override;
+	//~ End SWidget
+
+	void OnStageClick(const FPointerEvent& InMouseEvent, const TSharedRef<SDMStage>& InStageWidget);
+
+	bool OnStageCanAcceptDrop(const FDragDropEvent& InDragDropEvent, const TSharedRef<SDMStage>& InStageWidget) const;
+	FReply OnStageAcceptDrop(const FDragDropEvent& InDragDropEvent, const TSharedRef<SDMStage>& InStageWidget);
+
+	void SaveLayerPreviewSize(const float InNewSize, const TSharedRef<SDMStage>& InStageWidget);
+
+	EVisibility GetEffectsListVisibility() const;
+
+	EVisibility GetNewEffectVisibility() const;
+	TSharedRef<SWidget> CreateNewEffectButton();
+
+	TSharedRef<SWidget> CreateNewEffectMenu();
+
+	FReply OnCreateLayerBypassButtonClicked();
+	const FSlateBrush* GetCreateLayerBypassButtonImage() const;
+
+	FReply OnBaseToggleButtonClicked();
+	const FSlateBrush* GetBaseToggleButtonImage() const;
+
+	FReply OnLayerMaskToggleButtonClicked();
+	const FSlateBrush* GetLayerMaskToggleButtonImage() const;
+
+	FReply OnLayerLinkToggleButton();
+	const FSlateBrush* GetLayerLinkToggleButtonImage() const;
+	EVisibility GetLayerLinkToggleButtonVisibility() const;
+
+	FReply OnEffectsToggleButtonClicked();
+	const FSlateBrush* GetEffectsToggleButtonImage() const;
+
+	FVector2D GetStagePreviewSize() const;
+
+	TSharedRef<SWidget> CreateLayerHeaderText() const;
+	TSharedRef<SWidget> CreateLayerHeaderEditableText() const;
+};
diff --git a/Engine/Plugins/Experimental/DynamicMaterial/Source/DynamicMaterialEditor/Private/Slate/Layers/SDMSlotLayerView.cpp b/Engine/Plugins/Experimental/DynamicMaterial/Source/DynamicMaterialEditor/Private/Slate/Layers/SDMSlotLayerView.cpp
new file mode 100644
--- /dev/null
+++ b/Engine/Plugins/Experimental/DynamicMaterial/Source/DynamicMaterialEditor/Private/Slate/Layers/SDMSlotLayerView.cpp
@@ -0,0 +1,734 @@
+// Copyright Epic Games, Inc. All Rights Reserved.
+
+#include "Slate/Layers/SDMSlotLayerView.h"
+#include "Components/DMMaterialLayer.h"
+#include "Components/DMMaterialStage.h"
+#include "DMPrivate.h"
+#include "DynamicMaterialEditorCommands.h"
+#include "DynamicMaterialEditorSettings.h"
+#include "Framework/Application/SlateApplication.h"
+#include "Menus/DMMaterialSlotLayerMenus.h"
+#include "ScopedTransaction.h"
+#include "SDMLayerEffectsItem.h"
+#include "Slate/Layers/SDMSlotLayerItem.h"
+#include "Slate/SDMSlot.h"
+#include "ToolMenus.h"
+#include "Components/DMMaterialEffect.h"
+#include "Components/DMMaterialEffectStack.h"
+#include "Widgets/SNullWidget.h"
+
+#define LOCTEXT_NAMESPACE "SDMSlotLayerView"
+
+void SDMSlotLayerView::Construct(const FArguments& InArgs, const TSharedRef<SDMSlot>& InSlotWidget)
+{
+	SlotWidgetWeak = InSlotWidget;
+	MaterialSlotWeak = InSlotWidget->GetSlot();
+	ensure(MaterialSlotWeak.IsValid());
+
+	PreviewSize = InArgs._PreviewSize;
+	LayerSelected = InArgs._OnLayerSelected;
+	LayerStageSelected = InArgs._OnLayerStageSelected;
+
+	bPostRegenSelect = false;
+
+	SListView::Construct(SListView::FArguments()
+		.ListItemsSource(&LayerItems)
+		.SelectionMode(ESelectionMode::Single)
+		.ItemHeight_Lambda([]() { return UDynamicMaterialEditorSettings::Get()->LayerPreviewSize + 32.0f; })
+		.ClearSelectionOnClick(false)
+		.EnableAnimatedScrolling(true)
+		.ScrollbarVisibility(EVisibility::Visible)
+		.ConsumeMouseWheel(EConsumeMouseWheel::WhenScrollingPossible)
+		.OnGenerateRow(this, &SDMSlotLayerView::OnGenerateLayerItemWidget)
+		.OnSelectionChanged(this, &SDMSlotLayerView::OnLayerItemSelectionChanged)
+		.OnContextMenuOpening(this, &SDMSlotLayerView::CreateLayerItemContextMenu)
+	);
+
+	RegenerateItems();
+	RequestListRefresh();
+	BindCommands();	
+}
+
+TSharedRef<ITableRow> SDMSlotLayerView::OnGenerateLayerItemWidget(TSharedPtr<FDMMaterialLayerReference> InItem, const TSharedRef<STableViewBase>& InOwnerTable)
+{
+	TAttribute<bool> StageBaseEnabledAttr;
+	TAttribute<bool> StageBaseSelectedAttr;
+	TAttribute<bool> StageMaskEnabledAttr;
+	TAttribute<bool> StageMaskSelectedAttr;
+
+	if (ensure(InItem.IsValid()))
+	{
+		StageBaseEnabledAttr.Bind(InItem.ToSharedRef(), &FDMMaterialLayerReference::IsBaseEnabled);
+		StageBaseSelectedAttr.Bind(InItem.ToSharedRef(), &FDMMaterialLayerReference::IsBaseBeingEdited);
+		StageMaskEnabledAttr.Bind(InItem.ToSharedRef(), &FDMMaterialLayerReference::IsMaskEnabled);
+		StageMaskSelectedAttr.Bind(InItem.ToSharedRef(), &FDMMaterialLayerReference::IsMaskBeingEdited);
+	}
+
+	TSharedRef<SDMSlotLayerItem> LayerItem = SNew(SDMSlotLayerItem, SlotWidgetWeak.Pin(), InOwnerTable, InItem)
+		.PreviewSize(PreviewSize)
+		.OnStageSelected(this, &SDMSlotLayerView::OnLayerStageSelected)
+		.OnLayerLinkToggled(this, &SDMSlotLayerView::OnLayerLinkToggled)
+		.StageBaseEnabled(StageBaseEnabledAttr)
+		.StageBaseSelected(StageBaseSelectedAttr)
+		.StageMaskEnabled(StageMaskEnabledAttr)
+		.StageMaskSelected(StageMaskSelectedAttr);
+
+	if (UDMMaterialLayerObject* Layer = InItem->GetLayer())
+	{
+		const int32 LayerIndex = Layer->FindIndex();
+
+		if (LayerIndex != INDEX_NONE)
+		{
+			if (!StageWidgets.IsValidIndex(LayerIndex))
+			{
+				StageWidgets.SetNum(LayerItems.Num());
+			}
+
+			if (StageWidgets.IsValidIndex(LayerIndex))
+			{
+				StageWidgets[LayerIndex].Add(EDMMaterialLayerStage::Base, LayerItem->GetBaseStageWidget());
+				StageWidgets[LayerIndex].Add(EDMMaterialLayerStage::Mask, LayerItem->GetMaskStageWidget());
+			}
+		}
+	}
+
+	return LayerItem;
+}
+
+void SDMSlotLayerView::Internal_SelectLayers(const TArray<TSharedPtr<FDMMaterialLayerReference>>& InSelectedLayers, UDMMaterialStage* InStage)
+{
+	bool bHasChangedBeingEdited = false;
+
+	auto EnableSelectedStages = [this, &bHasChangedBeingEdited](const bool bInBeingEdited)
+	{
+		for (const TWeakObjectPtr<UDMMaterialStage>& SelectedStageWeak : SelectedStages)
+		{
+			if (UDMMaterialStage* SelectedStage = SelectedStageWeak.Get())
+			{
+				bHasChangedBeingEdited &= SelectedStage->SetBeingEdited(bInBeingEdited);
+			}
+		}
+	};
+
+	// Disable editing for all previously selected stages.
+	EnableSelectedStages(false);
+
+	// @TODO: Multiple stage selection is currently not fully supported.
+	SelectedStages = { InStage };
+
+	// Enable editing for all new selected stages.
+	EnableSelectedStages(true);
+
+	// Find the selected layer index.
+	SelectedLayerIndex = INDEX_NONE;
+
+	if (InSelectedLayers.Num() == 1)
+	{
+		for (int32 Index = 0; Index < LayerItems.Num(); ++Index)
+		{
+			if (LayerItems[Index] == InSelectedLayers[0])
+			{
+				SelectedLayerIndex = Index;
+			}
+		}
+	}
+
+	if (bHasChangedBeingEdited)
+	{
+		if (TSharedPtr<SDMSlot> SlotWidget = SlotWidgetWeak.Pin())
+		{
+			SlotWidget->InvalidateSlotSettingsRowWidget();
+		}
+	}
+}
+
+void SDMSlotLayerView::RegenerateItems()
+{
+	LayerItems.Empty();
+
+	UDMMaterialSlot* Slot = MaterialSlotWeak.Get();
+
+	if (!Slot)
+	{
+		return;
+	}
+
+	const TArray<TObjectPtr<UDMMaterialLayerObject>>& SlotLayers = Slot->GetLayers();
+	const int LayerCount = SlotLayers.Num();
+
+	LayerItems.Reserve(LayerCount);
+
+	for (int32 LayerIdx = LayerCount - 1; LayerIdx >= 0; --LayerIdx)
+	{
+		LayerItems.Add(MakeShared<FDMMaterialLayerReference>(SlotLayers[LayerIdx]));
+	}
+}
+
+void SDMSlotLayerView::BindCommands()
+{
+	TSharedPtr<SDMSlot> SlotWidget = SlotWidgetWeak.Pin();
+
+	if (!SlotWidget.IsValid())
+	{
+		return;
+	}
+
+	const FDynamicMaterialEditorCommands& Commands = FDynamicMaterialEditorCommands::Get();
+
+	SlotWidget->GetEditorWidget()->GetCommandList()->MapAction(
+		Commands.SelectLayerBaseStage,
+		FExecuteAction::CreateSP(this, &SDMSlotLayerView::ExecuteSelectLayerStage, EDMMaterialLayerStage::Base),
+		FCanExecuteAction::CreateSP(this, &SDMSlotLayerView::CanSelectLayerStage, EDMMaterialLayerStage::Base)
+	);
+
+	SlotWidget->GetEditorWidget()->GetCommandList()->MapAction(
+		Commands.SelectLayerMaskStage,
+		FExecuteAction::CreateSP(this, &SDMSlotLayerView::ExecuteSelectLayerStage, EDMMaterialLayerStage::Mask),
+		FCanExecuteAction::CreateSP(this, &SDMSlotLayerView::CanSelectLayerStage, EDMMaterialLayerStage::Mask)
+	);
+
+	SlotWidget->GetEditorWidget()->GetCommandList()->MapAction(
+		Commands.MoveLayerUp,
+		FExecuteAction::CreateSP(this, &SDMSlotLayerView::ExecuteMoveLayer, -1),
+		FCanExecuteAction::CreateSP(this, &SDMSlotLayerView::CanMoveLayer, -1)
+	);
+
+	SlotWidget->GetEditorWidget()->GetCommandList()->MapAction(
+		Commands.MoveLayerDown,
+		FExecuteAction::CreateSP(this, &SDMSlotLayerView::ExecuteMoveLayer, 1),
+		FCanExecuteAction::CreateSP(this, &SDMSlotLayerView::CanMoveLayer, 1)
+	);
+}
+
+bool SDMSlotLayerView::CanMoveLayer(int32 InOffset) const
+{
+	if (SelectedLayerIndex == INDEX_NONE)
+	{
+		return false;
+	}
+
+	// Already validated by CanMoveLayer
+	UDMMaterialSlot* Slot = MaterialSlotWeak.Get();
+
+	const TArray<TObjectPtr<UDMMaterialLayerObject>>& SlotLayers = Slot->GetLayers();
+
+	const int32 LayerCount = SlotLayers.Num();
+	const int32 SelectedSlotLayerIndex = LayerCount - SelectedLayerIndex - 1;
+
+	if (UDMMaterialLayerObject* Layer = Slot->GetLayer(SelectedSlotLayerIndex))
+	{
+		return !!Layer->GetFirstValidStage(EDMMaterialLayerStage::All);
+	}
+
+	return false;
+}
+
+void SDMSlotLayerView::ExecuteMoveLayer(int32 InOffset)
+{
+	if (!CanMoveLayer(InOffset))
+	{
+		return;
+	}
+
+	const int32 NewIndex = SelectedLayerIndex + InOffset;
+
+	// Already validated by CanMoveLayer
+	UDMMaterialSlot* Slot = MaterialSlotWeak.Get();
+
+	const TArray<TObjectPtr<UDMMaterialLayerObject>>& SlotLayers = Slot->GetLayers();
+
+	const int32 LayerCount = SlotLayers.Num();
+	const int32 SelectedSlotLayerIndex = LayerCount - SelectedLayerIndex - 1;
+	const int32 NewSlotIndex = LayerCount - NewIndex - 1;
+
+	if (UDMMaterialLayerObject* Layer = Slot->GetLayer(SelectedSlotLayerIndex))
+	{
+		FScopedTransaction Transaction(LOCTEXT("MoveLayer", "Material Designer Move Layer"));
+		Slot->Modify();
+
+		if (Slot->MoveLayer(Layer, NewSlotIndex))
+		{
+			PostRegenSelectedStages = SelectedStages;
+			PostRegenSelectedLayerIndex = NewSlotIndex;
+			bPostRegenSelect = true;
+			RequestListRefresh();
+		}
+	}
+}
+
+bool SDMSlotLayerView::CanSelectLayerStage(EDMMaterialLayerStage InStageType) const
+{
+	if (SelectedLayerIndex == INDEX_NONE
+		|| !LayerItems.IsValidIndex(SelectedLayerIndex) 
+		|| !MaterialSlotWeak.IsValid())
+	{
+		return false;
+	}
+
+	UDMMaterialSlot* Slot = MaterialSlotWeak.Get();
+
+	const TArray<TObjectPtr<UDMMaterialLayerObject>>& SlotLayers = Slot->GetLayers();
+
+	const int32 LayerCount = SlotLayers.Num();
+	const int32 SelectedSlotLayerIndex = LayerCount - SelectedLayerIndex - 1;
+
+	if (const TObjectPtr<UDMMaterialLayerObject>& Layer = Slot->GetLayer(SelectedSlotLayerIndex))
+	{
+		if (StageWidgets.IsValidIndex(SelectedLayerIndex))
+		{
+			if (const TWeakPtr<SDMStage>* StageWidgetWeakPtr = StageWidgets[SelectedLayerIndex].Find(InStageType))
+			{
+				if (StageWidgetWeakPtr->IsValid())
+				{
+					return true;
+				}
+			}
+		}
+	}
+
+	return false;
+}
+
+void SDMSlotLayerView::ExecuteSelectLayerStage(EDMMaterialLayerStage InStageType)
+{
+	if (!CanSelectLayerStage(InStageType))
+	{
+		return;
+	}
+
+	// Already validated by CanMoveLayer
+	UDMMaterialSlot* Slot = MaterialSlotWeak.Get();
+	
+	const TArray<TObjectPtr<UDMMaterialLayerObject>>& SlotLayers = Slot->GetLayers();
+
+	const int32 LayerCount = SlotLayers.Num();
+	const int32 SelectedSlotLayerIndex = LayerCount - SelectedLayerIndex - 1;
+
+	if (UDMMaterialLayerObject* Layer = Slot->GetLayer(SelectedSlotLayerIndex))
+	{
+		if (StageWidgets.IsValidIndex(SelectedSlotLayerIndex))
+		{
+			if (const TWeakPtr<SDMStage>* StageWidgetWeakPtr = StageWidgets[SelectedSlotLayerIndex].Find(InStageType))
+			{
+				if (TSharedPtr<SDMStage> StageWidget = StageWidgetWeakPtr->Pin())
+				{
+					OnLayerStageSelected(true, StageWidget.ToSharedRef());
+				}
+			}
+		}
+	}
+}
+
+void SDMSlotLayerView::OnUndo()
+{
+	RequestListRefresh();
+}
+
+void SDMSlotLayerView::OnLayerItemSelectionChanged(TSharedPtr<FDMMaterialLayerReference> InSelectedItem, ESelectInfo::Type InSelectInfo)
+{
+	TSharedPtr<SDMSlotLayerItem> SlotLayerWidget = WidgetFromLayerItem(InSelectedItem);
+	if (SlotLayerWidget.IsValid())
+	{
+		SlotLayerWidget->DeselectAllEffects();
+	}
+
+	Internal_SelectLayers({ InSelectedItem });
+
+	TSharedPtr<SDMSlot> SlotWidget = SlotWidgetWeak.Pin();
+
+	FScopedTransaction Transaction(LOCTEXT("SelectLayer", "Material Designer Select Layer"));
+
+	for (int32 LayerItemIdx = 0; LayerItemIdx < LayerItems.Num(); ++LayerItemIdx)
+	{
+		if (LayerItemIdx == SelectedLayerIndex)
+		{
+			continue;
+		}
+
+		TSharedPtr<ITableRow> RowWidget = WidgetFromItem(LayerItems[LayerItemIdx]);
+
+		if (RowWidget.IsValid())
+		{
+			TSharedPtr<SDMSlotLayerItem> LayerWidget = StaticCastSharedPtr<SDMSlotLayerItem>(RowWidget);
+
+			if (TSharedPtr<SDMStage> BaseWidget = LayerWidget->GetBaseStageWidget())
+			{
+				if (UDMMaterialStage* Stage = BaseWidget->GetStage())
+				{
+					Stage->Modify();
+					Stage->SetBeingEdited(false);
+
+					if (InSelectInfo == ESelectInfo::OnMouseClick && SlotWidget.IsValid())
+					{
+						SlotWidget->InvalidateSlotSettingsRowWidget();
+						SlotWidget->InvalidateComponentEditWidget();
+					}
+				}
+			}
+
+			if (TSharedPtr<SDMStage> MaskWidget = LayerWidget->GetMaskStageWidget())
+			{
+				if (UDMMaterialStage* Stage = MaskWidget->GetStage())
+				{
+					Stage->SetBeingEdited(false);
+
+					if (InSelectInfo == ESelectInfo::OnMouseClick && SlotWidget.IsValid())
+					{
+						SlotWidget->InvalidateSlotSettingsRowWidget();
+						SlotWidget->InvalidateComponentEditWidget();
+					}
+				}
+			}
+		}
+	}
+}
+
+void SDMSlotLayerView::OnLayerStageSelected(const bool bInSelected, const TSharedRef<SDMStage>& InStageWidget)
+{
+	UDMMaterialStage* SelectedStage = InStageWidget->GetStage();
+	if (!SelectedStage)
+	{
+		return;
+	}
+
+	TSharedPtr<FDMMaterialLayerReference> LayerItem = FindLayerItem(SelectedStage);
+	if (!LayerItem.IsValid())
+	{
+		return;
+	}
+
+	Internal_SelectLayers({ LayerItem }, SelectedStage);
+
+	SelectedStageWidgets.Empty();
+	SelectedStageWidgets.Add(InStageWidget);
+	
+	LayerStageSelected.ExecuteIfBound(bInSelected, InStageWidget);
+}
+ 
+void SDMSlotLayerView::OnLayerLinkToggled(const bool bIsLinkActive, const TSharedRef<SDMStage>& InStageWidget)
+{
+	if (UDMMaterialStage* Stage = InStageWidget->GetStage())
+	{
+		if (UDMMaterialLayerObject* Layer = Stage->GetLayer())
+		{
+			if (UDMMaterialSlot* Slot = Layer->GetSlot())
+			{
+				if (TSharedPtr<SDMSlot> SlotWidget = SlotWidgetWeak.Pin())
+				{
+					TArray<int32> SelectedLayers = SlotWidget->GetSelectedLayerIndices();
+
+					if (SelectedLayers.IsEmpty() == false)
+					{
+						const TArray<TObjectPtr<UDMMaterialLayerObject>>& Layers = Slot->GetLayers();
+
+						for (int32 LayerIdx = 0; LayerIdx < Layers.Num(); ++LayerIdx)
+						{
+							if (Layers[LayerIdx]->HasValidStage(Stage))
+							{
+								if (LayerIdx == SelectedLayers[0])
+								{
+									SlotWidget->InvalidateComponentEditWidget();
+								}
+							}
+						}
+					}
+				}
+			}
+		}
+	}
+}
+
+TSharedPtr<SDMSlotLayerItem> SDMSlotLayerView::WidgetFromLayerItem(const TSharedPtr<FDMMaterialLayerReference>& InItem)
+{
+	return StaticCastSharedPtr<SDMSlotLayerItem>(WidgetFromItem(InItem));
+}
+
+TSharedPtr<SWidget> SDMSlotLayerView::CreateLayerItemContextMenu()
+{
+	UDMMaterialLayerObject* LayerObject = nullptr;
+	FVector2f MousePosition = FSlateApplication::Get().GetCursorPos();
+
+	for (const TSharedPtr<FDMMaterialLayerReference>& LayerItem : LayerItems)
+	{
+		TSharedPtr<SDMSlotLayerItem> LayerItemWidget = WidgetFromLayerItem(LayerItem);
+
+		if (LayerItemWidget.IsValid() && LayerItemWidget->GetTickSpaceGeometry().IsUnderLocation(MousePosition))
+		{
+			LayerObject = LayerItem->GetLayer();
+			break;
+		}
+	}
+	
+	if (TSharedPtr<SDMSlot> SlotWidget = SlotWidgetWeak.Pin())
+	{
+		UToolMenu* ContextMenu = FDMMaterialSlotLayerMenus::GenerateSlotLayerMenu(SlotWidget, LayerObject);
+		return UToolMenus::Get()->GenerateWidget(ContextMenu);
+	}
+
+	return SNullWidget::NullWidget;
+}
+
+void SDMSlotLayerView::ScrollItemIntoView(const TSharedPtr<FDMMaterialLayerReference>& InItem)
+{
+	bNavigateOnScrollIntoView = true;
+
+	RequestScrollIntoView(InItem, 0);
+}
+
+void SDMSlotLayerView::FocusOnItem(const TSharedPtr<FDMMaterialLayerReference>& InItem)
+{
+	SelectorItem = InItem;
+	RangeSelectionStart = InItem;
+}
+
+int32 SDMSlotLayerView::GetLayerItemIndex(const TSharedPtr<FDMMaterialLayerReference>& Item) const
+{
+	if (HasValidItemsSource())
+	{
+		return GetItems().Find(TListTypeTraits<TSharedPtr<FDMMaterialLayerReference>>::NullableItemTypeConvertToItemType(Item));
+	}
+
+	return INDEX_NONE;
+}
+
+FCursorReply SDMSlotLayerView::OnCursorQuery(const FGeometry& MyGeometry, const FPointerEvent& CursorEvent) const
+{
+	if (IsRightClickScrolling() && CursorEvent.IsMouseButtonDown(EKeys::RightMouseButton))
+	{
+		// We hide the native cursor as we'll be drawing the software EMouseCursor::GrabHandClosed cursor
+		return FCursorReply::Cursor(EMouseCursor::None);
+	}
+	return SListView::OnCursorQuery(MyGeometry, CursorEvent);
+}
+
+FReply SDMSlotLayerView::OnMouseMove(const FGeometry& MyGeometry, const FPointerEvent& MouseEvent)
+{
+	if (MouseEvent.IsMouseButtonDown(EKeys::RightMouseButton) && !MouseEvent.IsTouchEvent())
+	{
+		// We only care about deltas along the scroll axis
+		FTableViewDimensions CursorDeltaDimensions(Orientation, MouseEvent.GetCursorDelta());
+		CursorDeltaDimensions.LineAxis = 0.f;
+
+		const float ScrollByAmount = CursorDeltaDimensions.ScrollAxis / MyGeometry.Scale;
+
+		// If scrolling with the right mouse button, we need to remember how much we scrolled.
+		// If we did not scroll at all, we will bring up the context menu when the mouse is released.
+		AmountScrolledWhileRightMouseDown += FMath::Abs(ScrollByAmount);
+
+		// Has the mouse moved far enough with the right mouse button held down to start capturing
+		// the mouse and dragging the view?
+		if (IsRightClickScrolling())
+		{
+			// Make sure the active timer is registered to update the inertial scroll
+			if (!bIsScrollingActiveTimerRegistered)
+			{
+				bIsScrollingActiveTimerRegistered = true;
+				RegisterActiveTimer(0.f, FWidgetActiveTimerDelegate::CreateSP(this, &SDMSlotLayerView::UpdateInertialScroll));
+			}
+
+			TickScrollDelta -= ScrollByAmount;
+
+			const float AmountScrolled = this->ScrollBy(MyGeometry, -ScrollByAmount, AllowOverscroll);
+
+			FReply Reply = FReply::Handled();
+
+			// The mouse moved enough that we're now dragging the view. Capture the mouse
+			// so the user does not have to stay within the bounds of the list while dragging.
+			if (this->HasMouseCapture() == false)
+			{
+				Reply.CaptureMouse(AsShared()).UseHighPrecisionMouseMovement(AsShared());
+				SoftwareCursorPosition = MyGeometry.AbsoluteToLocal(MouseEvent.GetScreenSpacePosition());
+				bShowSoftwareCursor    = true;
+			}
+
+			// Check if the mouse has moved.
+			if (AmountScrolled != 0)
+			{
+				SoftwareCursorPosition += CursorDeltaDimensions.ToVector2D();
+			}
+
+			return Reply;
+		}
+	}
+
+	return SListView::OnMouseMove(MyGeometry, MouseEvent);
+}
+
+FReply SDMSlotLayerView::OnMouseButtonUp(const FGeometry& MyGeometry, const FPointerEvent& MouseEvent)
+{
+	if (MouseEvent.GetEffectingButton() == EKeys::RightMouseButton)
+	{
+		FReply Reply = FReply::Handled().ReleaseMouseCapture();
+		AmountScrolledWhileRightMouseDown = 0;
+		bShowSoftwareCursor = false;
+
+		// If we have mouse capture, snap the mouse back to the closest location that is within the list's bounds
+		if (HasMouseCapture())
+		{
+			const FSlateRect ListScreenSpaceRect = MyGeometry.GetLayoutBoundingRect();
+			const FVector2D CursorPosition = MyGeometry.LocalToAbsolute(SoftwareCursorPosition);
+
+			const FIntPoint BestPositionInList(
+				FMath::RoundToInt(FMath::Clamp(CursorPosition.X, ListScreenSpaceRect.Left, ListScreenSpaceRect.Right)),
+				FMath::RoundToInt(FMath::Clamp(CursorPosition.Y, ListScreenSpaceRect.Top, ListScreenSpaceRect.Bottom))
+			);
+
+			Reply.SetMousePos(BestPositionInList);
+		}
+
+		return Reply;
+	}
+
+	if (MouseEvent.GetEffectingButton() == EKeys::LeftMouseButton)
+	{
+		bool bHovering = false;
+		FVector2f MousePosition = FSlateApplication::Get().GetCursorPos();
+
+		for (const TSharedPtr<FDMMaterialLayerReference>& LayerItem : LayerItems)
+		{
+			TSharedPtr<SDMSlotLayerItem> LayerItemWidget = WidgetFromLayerItem(LayerItem);
+
+			if (LayerItemWidget.IsValid() && LayerItemWidget->GetTickSpaceGeometry().IsUnderLocation(MousePosition))
+			{
+				bHovering = true;
+				break;
+			}
+		}
+	}
+
+	return SListView::OnMouseButtonUp(MyGeometry, MouseEvent);
+}
+
+void SDMSlotLayerView::RequestListRefresh()
+{
+	RegenerateItems();
+
+	SListView<TSharedPtr<FDMMaterialLayerReference>>::RequestListRefresh();
+}
+
+void SDMSlotLayerView::PostUndo(bool bSuccess)
+{
+	OnUndo();
+}
+
+void SDMSlotLayerView::PostRedo(bool bSuccess)
+{
+	OnUndo();
+}
+
+TSharedPtr<FDMMaterialLayerReference> SDMSlotLayerView::FindLayerItem(UDMMaterialStage* const InStage) const
+{
+	if (!IsValid(InStage))
+	{
+		return nullptr;
+	}
+
+	for (const TSharedPtr<FDMMaterialLayerReference>& LayerItem : LayerItems)
+	{
+		if (LayerItem.IsValid())
+		{
+			if (UDMMaterialLayerObject* Layer = LayerItem->GetLayer())
+			{
+				if (Layer->HasValidStage(InStage))
+				{
+					return LayerItem;
+				}
+			}
+		}
+	}
+
+	return nullptr;
+}
+
+TSharedPtr<FDMMaterialLayerReference> SDMSlotLayerView::FindLayerItem(UDMMaterialLayerObject* const InLayer) const
+{
+	if (!MaterialSlotWeak.IsValid())
+	{
+		return nullptr;
+	}
+
+	for (const TSharedPtr<FDMMaterialLayerReference>& LayerItem : LayerItems)
+	{
+		if (LayerItem.IsValid())
+		{
+			if (UDMMaterialLayerObject* Layer = LayerItem->GetLayer())
+			{
+				if (Layer == InLayer)
+				{
+					return LayerItem;
+				}
+			}
+		}
+	}
+
+	return nullptr;
+}
+
+void SDMSlotLayerView::SelectLayerItem(UDMMaterialStage* const InStage, const bool bInSelected, const ESelectInfo::Type InSelectInfo)
+{
+	TSharedPtr<FDMMaterialLayerReference> Item = FindLayerItem(InStage);
+	SetItemSelection(Item, bInSelected, InSelectInfo);
+}
+
+void SDMSlotLayerView::SelectLayerItem(UDMMaterialLayerObject* const InLayer, const bool bInMask, const bool bInSelected, const ESelectInfo::Type InSelectInfo)
+{
+	TSharedPtr<FDMMaterialLayerReference> Item = FindLayerItem(InLayer);
+	SetItemSelection(Item, bInSelected, InSelectInfo);
+}
+
+bool SDMSlotLayerView::AddLayerItem(const TSharedPtr<FDMMaterialLayerReference>& InLayerItem)
+{
+	if (!InLayerItem.IsValid())
+	{
+		return false;
+	}
+
+	// The items themselves are only links to indices in an array. Regenerating the list will be quick
+	// and more than enough to get an accurate display.
+
+	RequestListRefresh();
+
+	return true;
+}
+
+bool SDMSlotLayerView::RemoveLayerItem(const TSharedPtr<FDMMaterialLayerReference>& InLayerItem)
+{
+	if (!InLayerItem.IsValid())
+	{
+		return false;
+	}
+
+	// The items themselves are only links to indices in an array. Regenerating the list will be quick
+	// and more than enough to get an accurate display.
+
+	RequestListRefresh();
+
+	return true;
+}
+
+void SDMSlotLayerView::Tick(const FGeometry& AllottedGeometry, const double InCurrentTime, const float InDeltaTime)
+{
+	SListView<TSharedPtr<FDMMaterialLayerReference, ESPMode::ThreadSafe>>::Tick(AllottedGeometry, InCurrentTime, InDeltaTime);
+
+	if (bPostRegenSelect && LayerItems.IsValidIndex(PostRegenSelectedLayerIndex))
+	{
+		FScopedTransaction Transaction(LOCTEXT("SelectLayer", "Material Designer Select Layer"));
+
+		for (const TWeakObjectPtr<UDMMaterialStage>& SelectedStageWeak : PostRegenSelectedStages)
+		{
+			if (UDMMaterialStage* SelectedStage = SelectedStageWeak.Get())
+			{
+				SelectedStage->Modify();
+				SelectedStage->SetBeingEdited(true);
+			}
+		}
+
+		SelectedStages = PostRegenSelectedStages;
+	}
+
+	bPostRegenSelect = false;
+	PostRegenSelectedLayerIndex = INDEX_NONE;
+	PostRegenSelectedStages.Empty();
+}
+
+#undef LOCTEXT_NAMESPACE
diff --git a/Engine/Plugins/Experimental/DynamicMaterial/Source/DynamicMaterialEditor/Private/Slate/Layers/SDMSlotLayerView.h b/Engine/Plugins/Experimental/DynamicMaterial/Source/DynamicMaterialEditor/Private/Slate/Layers/SDMSlotLayerView.h
new file mode 100644
--- /dev/null
+++ b/Engine/Plugins/Experimental/DynamicMaterial/Source/DynamicMaterialEditor/Private/Slate/Layers/SDMSlotLayerView.h
@@ -0,0 +1,117 @@
+// Copyright Epic Games, Inc. All Rights Reserved.
+
+#pragma once
+
+#include "Components/DMMaterialSlot.h"
+#include "Containers/Array.h"
+#include "Containers/Map.h"
+#include "EditorUndoClient.h"
+#include "Slate/Layers/SDMSlotLayerItem.h"
+#include "UObject/WeakObjectPtrTemplates.h"
+#include "Widgets/Views/SListView.h"
+
+class SDMLayerEffectsItem;
+class SDMSlot;
+class UDMMaterialSlot;
+class UDMMaterialStage;
+
+DECLARE_DELEGATE_TwoParams(FDMOnLayerViewItemSelected, TSharedPtr<FDMMaterialLayerReference> /** InLayerItem */, const int32 /** InLayerIndex */)
+
+class SDMSlotLayerView : public SListView<TSharedPtr<FDMMaterialLayerReference>>, public FSelfRegisteringEditorUndoClient
+{
+public:
+	SLATE_BEGIN_ARGS(SDMSlotLayerView)
+	{}
+		SLATE_ATTRIBUTE(int32, PreviewSize)
+		SLATE_EVENT(FDMOnLayerViewItemSelected, OnLayerSelected)
+		SLATE_EVENT(FDMOnStageSelected, OnLayerStageSelected)
+	SLATE_END_ARGS()
+
+	virtual ~SDMSlotLayerView() override {}
+
+	void Construct(const FArguments& InArgs, const TSharedRef<SDMSlot>& InSlotWidget);
+
+	TSharedPtr<SDMSlot> GetSlotWidget() const { return SlotWidgetWeak.Pin(); }
+
+	void ScrollItemIntoView(const TSharedPtr<FDMMaterialLayerReference>& InItem);
+	void FocusOnItem(const TSharedPtr<FDMMaterialLayerReference>& InItem);
+
+	int32 GetLayerItemIndex(const TSharedPtr<FDMMaterialLayerReference>& InItem) const;
+
+	const TArray<TSharedPtr<FDMMaterialLayerReference>>& GetLayerItems() const { return LayerItems; }
+
+	int32 GetSelectedLayerIndex() const { return SelectedLayerIndex; }
+	const TArray<TWeakPtr<SDMStage>>& GetSelectedStageWidgets() const { return SelectedStageWidgets; }
+
+	TSharedPtr<SDMSlotLayerItem> WidgetFromLayerItem(const TSharedPtr<FDMMaterialLayerReference>& InItem);
+
+	TSharedPtr<FDMMaterialLayerReference> FindLayerItem(UDMMaterialStage* const InStage) const;
+	TSharedPtr<FDMMaterialLayerReference> FindLayerItem(UDMMaterialLayerObject* const InLayer) const;
+
+	void SelectLayerItem(UDMMaterialStage* const InStage, const bool bInSelected, const ESelectInfo::Type InSelectInfo);
+	void SelectLayerItem(UDMMaterialLayerObject* const InLayer, const bool bInMask, const bool bInSelected, const ESelectInfo::Type InSelectInfo);
+
+	bool AddLayerItem(const TSharedPtr<FDMMaterialLayerReference>& InLayerItem);
+	bool RemoveLayerItem(const TSharedPtr<FDMMaterialLayerReference>& InLayerItem);
+
+	//~ Begin SWidget
+	virtual void Tick(const FGeometry& AllottedGeometry, const double InCurrentTime, const float InDeltaTime) override;
+	virtual FCursorReply OnCursorQuery(const FGeometry& MyGeometry, const FPointerEvent& CursorEvent) const override;
+	virtual FReply OnMouseMove(const FGeometry& MyGeometry, const FPointerEvent& MouseEvent) override;
+	virtual FReply OnMouseButtonUp(const FGeometry& MyGeometry, const FPointerEvent& MouseEvent) override;
+	//~ End SWidget
+
+	virtual void RequestListRefresh() override;
+
+	//~ Begin FUndoClient
+	virtual void PostUndo(bool bSuccess) override;
+	virtual void PostRedo(bool bSuccess) override;
+	//~ End FUndoClient
+
+protected:
+	TWeakPtr<SDMSlot> SlotWidgetWeak;
+	TWeakObjectPtr<UDMMaterialSlot> MaterialSlotWeak;
+
+	TAttribute<int32> PreviewSize;
+	FDMOnLayerViewItemSelected LayerSelected;
+	FDMOnStageSelected LayerStageSelected;
+
+	TArray<TSharedPtr<FDMMaterialLayerReference>> LayerItems;
+	TArray<TMap<EDMMaterialLayerStage, TWeakPtr<SDMStage>>> StageWidgets;
+
+	int32 SelectedLayerIndex = INDEX_NONE;
+
+	TArray<TWeakPtr<SDMStage>> SelectedStageWidgets;
+
+	TArray<TWeakObjectPtr<UDMMaterialStage>> SelectedStages;
+
+	bool bPostRegenSelect = false;
+	int32 PostRegenSelectedLayerIndex = INDEX_NONE;
+	TArray<TWeakObjectPtr<UDMMaterialStage>> PostRegenSelectedStages;
+
+	TSharedRef<ITableRow> OnGenerateLayerItemWidget(TSharedPtr<FDMMaterialLayerReference> InItem, const TSharedRef<STableViewBase>& InOwnerTable);
+
+	void OnLayerItemSelectionChanged(TSharedPtr<FDMMaterialLayerReference> InSelectedItem, ESelectInfo::Type InSelectInfo);
+
+	TSharedPtr<SWidget> CreateLayerItemContextMenu();
+
+	void OnLayerStageSelected(const bool bInSelected, const TSharedRef<SDMStage>& InStageWidget);
+
+	void OnLayerLinkToggled(const bool bIsLinkActive, const TSharedRef<SDMStage>& InStageWidget);
+
+	void Internal_SelectLayers(const TArray<TSharedPtr<FDMMaterialLayerReference>>& InSelectedLayers, UDMMaterialStage* InStage = nullptr);
+
+	void RegenerateItems();
+
+	void BindCommands();
+
+	bool CanMoveLayer(int32 InOffset) const;
+
+	void ExecuteMoveLayer(int32 InOffset);
+
+	bool CanSelectLayerStage(EDMMaterialLayerStage InStageType) const;
+
+	void ExecuteSelectLayerStage(EDMMaterialLayerStage InStageType);
+
+	void OnUndo();
+};
diff --git a/Engine/Plugins/Experimental/DynamicMaterial/Source/DynamicMaterialEditor/Private/Slate/Previews/SDMMaterialValuePreview.cpp b/Engine/Plugins/Experimental/DynamicMaterial/Source/DynamicMaterialEditor/Private/Slate/Previews/SDMMaterialValuePreview.cpp
new file mode 100644
--- /dev/null
+++ b/Engine/Plugins/Experimental/DynamicMaterial/Source/DynamicMaterialEditor/Private/Slate/Previews/SDMMaterialValuePreview.cpp
@@ -0,0 +1,84 @@
+// Copyright Epic Games, Inc. All Rights Reserved.
+
+#include "SDMMaterialValuePreview.h"
+#include "Components/DMMaterialValue.h"
+#include "DynamicMaterialModule.h"
+#include "Materials/Material.h"
+#include "Widgets/Images/SImage.h"
+
+SDMMaterialValuePreview::SDMMaterialValuePreview()
+	: Brush(FSlateMaterialBrush(FVector2D(1.f, 1.f)))
+{
+
+}
+
+SDMMaterialValuePreview::~SDMMaterialValuePreview()
+{
+	if (!FDynamicMaterialModule::AreUObjectsSafe())
+	{
+		return;
+	}
+
+	UDMMaterialValue* Value = ValueWeak.Get();
+
+	if (Value)
+	{
+		Value->GetOnUpdate().RemoveAll(this);
+	}
+}
+
+void SDMMaterialValuePreview::Construct(const FArguments& InArgs, UDMMaterialValue* InValue)
+{
+	ValueWeak = InValue;
+
+	SetCanTick(true);
+
+	Brush.SetImageSize(InArgs._DesiredSize);
+
+	if (ensure(IsValid(InValue)))
+	{
+		InValue->GetOnUpdate().AddSP(this, &SDMMaterialValuePreview::OnValueUpdated);
+		OnValueUpdated(InValue, EDMUpdateType::Value);
+	}
+
+	ChildSlot
+	[
+		SNew(SImage)
+		.Image(&Brush)
+		.DesiredSizeOverride(InArgs._DesiredSize)
+	];
+}
+
+void SDMMaterialValuePreview::Tick(const FGeometry& AllottedGeometry, const double InCurrentTime, const float InDeltaTime)
+{
+	SCompoundWidget::Tick(AllottedGeometry, InCurrentTime, InDeltaTime);
+
+	if (UDMMaterialComponent::CanClean())
+	{
+		if (UDMMaterialValue* Value = ValueWeak.Get())
+		{
+			if (Value->IsComponentValid() && Value->NeedsClean())
+			{
+				Value->DoClean();
+				OnValueUpdated(Value, EDMUpdateType::Value);
+			}
+		}
+	}
+}
+
+void SDMMaterialValuePreview::OnValueUpdated(UDMMaterialComponent* InComponent, EDMUpdateType InUpdateType)
+{
+	if (UDMMaterialValue* Value = Cast<UDMMaterialValue>(InComponent))
+	{
+		if (Value == ValueWeak.Get() && IsValid(Value) && Value->IsComponentValid())
+		{
+			UMaterial* PreviewMaterial = Value->GetPreviewMaterial();
+
+			if (Brush.GetResourceObject() != PreviewMaterial)
+			{
+				Brush.SetMaterial(PreviewMaterial);
+			}
+		}
+	}
+}
+
diff --git a/Engine/Plugins/Experimental/DynamicMaterial/Source/DynamicMaterialEditor/Private/Slate/Previews/SDMMaterialValuePreview.h b/Engine/Plugins/Experimental/DynamicMaterial/Source/DynamicMaterialEditor/Private/Slate/Previews/SDMMaterialValuePreview.h
new file mode 100644
--- /dev/null
+++ b/Engine/Plugins/Experimental/DynamicMaterial/Source/DynamicMaterialEditor/Private/Slate/Previews/SDMMaterialValuePreview.h
@@ -0,0 +1,33 @@
+// Copyright Epic Games, Inc. All Rights Reserved.
+
+#pragma once
+
+#include "DMEDefs.h"
+#include "SlateMaterialBrush.h"
+#include "Widgets/SCompoundWidget.h"
+
+class UDMMaterialComponent;
+class UDMMaterialValue;
+
+class SDMMaterialValuePreview : public SCompoundWidget
+{
+	SLATE_BEGIN_ARGS(SDMMaterialValuePreview)
+		: _DesiredSize(FVector2D(48.f, 48.f))
+		{}
+		SLATE_ARGUMENT(FVector2D, DesiredSize)
+	SLATE_END_ARGS()
+
+public:
+	SDMMaterialValuePreview();
+	virtual ~SDMMaterialValuePreview();
+
+	void Construct(const FArguments& InArgs, UDMMaterialValue* InValue);
+
+	virtual void Tick(const FGeometry& AllottedGeometry, const double InCurrentTime, const float InDeltaTime) override;
+
+protected:
+	TWeakObjectPtr<UDMMaterialValue> ValueWeak;
+	FSlateMaterialBrush Brush;
+
+	void OnValueUpdated(UDMMaterialComponent* InComponent, EDMUpdateType InUpdateType);
+};
diff --git a/Engine/Plugins/Experimental/DynamicMaterial/Source/DynamicMaterialEditor/Private/Slate/Previews/SDMSlotPreview.cpp b/Engine/Plugins/Experimental/DynamicMaterial/Source/DynamicMaterialEditor/Private/Slate/Previews/SDMSlotPreview.cpp
new file mode 100644
--- /dev/null
+++ b/Engine/Plugins/Experimental/DynamicMaterial/Source/DynamicMaterialEditor/Private/Slate/Previews/SDMSlotPreview.cpp
@@ -0,0 +1,189 @@
+// Copyright Epic Games, Inc. All Rights Reserved.
+
+#include "SDMSlotPreview.h"
+#include "Components/DMMaterialSlot.h"
+#include "DynamicMaterialEditorSettings.h"
+#include "DynamicMaterialEditorStyle.h"
+#include "DynamicMaterialModule.h"
+#include "Materials/Material.h"
+#include "Slate/SMaterialToolTip.h"
+#include "Widgets/Images/SImage.h"
+#include "Widgets/SOverlay.h"
+#include "Widgets/SToolTip.h"
+#include "Widgets/Views/STableViewBase.h"
+
+SDMSlotPreview::SDMSlotPreview()
+	: Brush(FSlateMaterialBrush(FVector2D(1.f, 1.f)))
+{
+}
+
+SDMSlotPreview::~SDMSlotPreview()
+{
+	if (!FDynamicMaterialModule::AreUObjectsSafe())
+	{
+		return;
+	}
+
+	if (UDMMaterialSlot* Slot = SlotWeak.Get())
+	{
+		Slot->GetOnUpdate().RemoveAll(this);
+	}
+}
+
+void SDMSlotPreview::Construct(const FArguments& InArgs, UDMMaterialSlot* InSlot, EDMMaterialLayerStage InLayerStage)
+{
+	SlotWeak = InSlot;
+	LayerStage = InLayerStage;
+
+	PreviewSize = InArgs._PreviewSize;
+
+	SetCanTick(true);
+
+	Brush.SetImageSize(PreviewSize.Get());
+
+	if (ensure(IsValid(InSlot)))
+	{
+		InSlot->GetOnUpdate().AddSP(this, &SDMSlotPreview::OnSlotUpdated);
+		OnSlotUpdated(InSlot, EDMUpdateType::Structure);
+
+		SetToolTipText(FText());
+		SetToolTip(
+			SNew(SToolTip)
+			.IsInteractive(false)
+			.BorderImage(FCoreStyle::Get().GetBrush("ToolTip.Background"))
+			[
+				SNew(SMaterialToolTip)
+				.Material(InSlot->GetPreviewMaterial(InLayerStage))
+				.Text(InArgs._ToolTipText)
+				.MaterialSize_Lambda([]() { return FVector2D(UDynamicMaterialEditorSettings::Get()->TooltipTextureSize); })
+				.ShowMaterial_Lambda([]() { return UDynamicMaterialEditorSettings::Get()->bShowTooltipPreview; })
+			]
+		);
+	}
+	
+	ChildSlot
+	[
+		SNew(SOverlay)
+		+ SOverlay::Slot()
+		.HAlign(HAlign_Fill)
+		.VAlign(VAlign_Center)
+		.Padding(2.f)
+		[
+			SNew(SImage)
+			.Image(&Brush)
+			.DesiredSizeOverride(this, &SDMSlotPreview::GetPreviewSizeInner)
+		]
+		+ SOverlay::Slot()
+		.HAlign(HAlign_Fill)
+		.VAlign(VAlign_Center)
+		[
+			SNew(SImage)
+			.DesiredSizeOverride(this, &SDMSlotPreview::GetPreviewSizeInner)
+			.Image(FDynamicMaterialEditorStyle::GetBrush("ImageBorder"))
+			.Cursor(EMouseCursor::Hand)
+			.OnMouseButtonDown(InArgs._OnMouseButtonDown)
+		]
+	];
+}
+
+void SDMSlotPreview::OnSlotUpdated(UDMMaterialComponent* InComponent, EDMUpdateType InUpdateType)
+{
+	if (UDMMaterialSlot* Slot = Cast<UDMMaterialSlot>(InComponent))
+	{
+		if (Slot == SlotWeak.Get() && IsValid(Slot) && Slot->IsComponentValid())
+		{
+			UMaterial* PreviewMaterial = Slot->GetPreviewMaterial(LayerStage);
+
+			if (Brush.GetResourceObject() != PreviewMaterial)
+			{
+				Brush.SetMaterial(PreviewMaterial);
+			}
+		}
+	}
+}
+
+TOptional<FVector2D> SDMSlotPreview::GetPreviewSizeInner() const
+{
+	return PreviewSize.Get();
+}
+
+TOptional<FVector2D> SDMSlotPreview::GetPreviewSizeOuter() const
+{
+	return PreviewSize.Get() + FVector2D(2.0f);
+}
+
+void SDMSlotPreview::Tick(const FGeometry& InAllottedGeometry, const double InCurrentTime, const float InDeltaTime)
+{
+	SCompoundWidget::Tick(InAllottedGeometry, InCurrentTime, InDeltaTime);
+
+	if (UDMMaterialComponent::CanClean())
+	{
+		if (UDMMaterialSlot* Slot = SlotWeak.Get())
+		{
+			if (Slot->IsComponentValid() && Slot->NeedsClean())
+			{
+				Slot->DoClean();
+				OnSlotUpdated(Slot, EDMUpdateType::Structure);
+			}
+		}
+	}
+}
+
+FCursorReply SDMSlotPreview::OnCursorQuery(const FGeometry& InMyGeometry, const FPointerEvent& InCursorEvent) const
+{
+	if (InCursorEvent.IsMouseButtonDown(EKeys::MiddleMouseButton))
+	{
+		return FCursorReply::Cursor(EMouseCursor::None);
+	}
+	return SCompoundWidget::OnCursorQuery(InMyGeometry, InCursorEvent);
+}
+
+FReply SDMSlotPreview::OnMouseMove(const FGeometry& InMyGeometry, const FPointerEvent& InMouseEvent)
+{
+	if (InMouseEvent.IsMouseButtonDown(EKeys::MiddleMouseButton) && !InMouseEvent.IsTouchEvent())
+	{
+		FReply Reply = FReply::Handled();
+
+		if (!HasMouseCapture())
+		{
+			CachedCursorLocation = InMouseEvent.GetScreenSpacePosition();
+			CachedPreviewSize = UDynamicMaterialEditorSettings::Get()->SlotPreviewSize;
+
+			Reply.CaptureMouse(AsShared()).UseHighPrecisionMouseMovement(AsShared());
+		}
+
+		FTableViewDimensions CursorDeltaDimensions(Orient_Vertical, InMouseEvent.GetCursorDelta());
+		CursorDeltaDimensions.LineAxis = 0.0f;
+		const float SizeChange = CursorDeltaDimensions.ScrollAxis / InMyGeometry.Scale;
+
+		FProperty* PreviewSizeProperty = UDynamicMaterialEditorSettings::StaticClass()->FindPropertyByName(GET_MEMBER_NAME_CHECKED(UDynamicMaterialEditorSettings, SlotPreviewSize));
+		ensure(PreviewSizeProperty);
+		const int32 ClampMin = PreviewSizeProperty ? PreviewSizeProperty->GetIntMetaData("ClampMin") : 0;
+		const int32 ClampMax = PreviewSizeProperty ? PreviewSizeProperty->GetIntMetaData("ClampMax") : 1;
+
+		CachedPreviewSize += SizeChange;
+		CachedPreviewSize = FMath::Clamp(CachedPreviewSize, ClampMin, ClampMax);
+
+		UDynamicMaterialEditorSettings* Settings = UDynamicMaterialEditorSettings::Get();
+		Settings->SlotPreviewSize = CachedPreviewSize;
+		Settings->SaveConfig();
+
+		return Reply;
+	}
+
+	return SCompoundWidget::OnMouseMove(InMyGeometry, InMouseEvent);
+}
+
+FReply SDMSlotPreview::OnMouseButtonUp(const FGeometry& InMyGeometry, const FPointerEvent& InMouseEvent)
+{
+	if (InMouseEvent.GetEffectingButton() == EKeys::MiddleMouseButton)
+	{
+		FReply Reply = FReply::Handled().ReleaseMouseCapture();
+		if (HasMouseCapture())
+		{
+			Reply.SetMousePos(FIntPoint(CachedCursorLocation.X, CachedCursorLocation.Y));
+		}
+		return Reply;
+	}
+	return SCompoundWidget::OnMouseButtonUp(InMyGeometry, InMouseEvent);
+}
diff --git a/Engine/Plugins/Experimental/DynamicMaterial/Source/DynamicMaterialEditor/Private/Slate/Previews/SDMSlotPreview.h b/Engine/Plugins/Experimental/DynamicMaterial/Source/DynamicMaterialEditor/Private/Slate/Previews/SDMSlotPreview.h
new file mode 100644
--- /dev/null
+++ b/Engine/Plugins/Experimental/DynamicMaterial/Source/DynamicMaterialEditor/Private/Slate/Previews/SDMSlotPreview.h
@@ -0,0 +1,50 @@
+// Copyright Epic Games, Inc. All Rights Reserved.
+
+#pragma once
+
+#include "DMDefs.h"
+#include "SlateMaterialBrush.h"
+#include "Types/WidgetMouseEventsDelegate.h"
+#include "Widgets/SCompoundWidget.h"
+
+class UDMMaterialComponent;
+class UDMMaterialSlot;
+enum class EDMMaterialLayerStage : uint8;
+
+class SDMSlotPreview : public SCompoundWidget
+{
+	SLATE_BEGIN_ARGS(SDMSlotPreview)
+		: _PreviewSize(FVector2D(48.0f))
+		{}
+		SLATE_ATTRIBUTE(FVector2D, PreviewSize)
+		SLATE_EVENT(FPointerEventHandler, OnMouseButtonDown)
+	SLATE_END_ARGS()
+
+public:
+	SDMSlotPreview();
+	virtual ~SDMSlotPreview();
+
+	void Construct(const FArguments& InArgs, UDMMaterialSlot* InSlot, EDMMaterialLayerStage InLayerStage);
+
+	//~ Begin SWidget
+	virtual void Tick(const FGeometry& InAllottedGeometry, const double InCurrentTime, const float InDeltaTime) override;
+	virtual FCursorReply OnCursorQuery(const FGeometry& InMyGeometry, const FPointerEvent& InCursorEvent) const override;
+	virtual FReply OnMouseMove(const FGeometry& InMyGeometry, const FPointerEvent& InMouseEvent) override;
+	virtual FReply OnMouseButtonUp(const FGeometry& InMyGeometry, const FPointerEvent& InMouseEvent) override;
+	//~ End SWidget
+
+protected:
+	TWeakObjectPtr<UDMMaterialSlot> SlotWeak;
+	EDMMaterialLayerStage LayerStage;
+	FSlateMaterialBrush Brush;
+
+	TAttribute<FVector2D> PreviewSize;
+
+	FVector2D CachedCursorLocation;
+	float CachedPreviewSize;
+
+	void OnSlotUpdated(UDMMaterialComponent* InComponent, EDMUpdateType InUpdateType);
+
+	TOptional<FVector2D> GetPreviewSizeInner() const;
+	TOptional<FVector2D> GetPreviewSizeOuter() const;
+};
diff --git a/Engine/Plugins/Experimental/DynamicMaterial/Source/DynamicMaterialEditor/Private/Slate/Previews/SDMStagePreview.cpp b/Engine/Plugins/Experimental/DynamicMaterial/Source/DynamicMaterialEditor/Private/Slate/Previews/SDMStagePreview.cpp
new file mode 100644
--- /dev/null
+++ b/Engine/Plugins/Experimental/DynamicMaterial/Source/DynamicMaterialEditor/Private/Slate/Previews/SDMStagePreview.cpp
@@ -0,0 +1,87 @@
+// Copyright Epic Games, Inc. All Rights Reserved.
+
+#include "SDMStagePreview.h"
+#include "Components/DMMaterialStage.h"
+#include "DynamicMaterialModule.h"
+#include "Materials/Material.h"
+#include "Widgets/Images/SImage.h"
+
+SDMStagePreview::SDMStagePreview()
+	: Brush(FSlateMaterialBrush(FVector2D(1.f, 1.f)))
+{
+}
+
+SDMStagePreview::~SDMStagePreview()
+{
+	if (!FDynamicMaterialModule::AreUObjectsSafe())
+	{
+		return;
+	}
+
+	if (UDMMaterialStage* Stage = StageWeak.Get())
+	{
+		Stage->GetOnUpdate().RemoveAll(this);
+	}
+}
+
+void SDMStagePreview::Construct(const FArguments& InArgs, UDMMaterialStage* InStage)
+{
+	StageWeak = InStage;
+
+	PreviewSize = InArgs._PreviewSize;
+
+	SetCanTick(true);
+
+	Brush.SetImageSize(PreviewSize.Get());
+
+	if (ensure(IsValid(InStage)))
+	{
+		InStage->GetOnUpdate().AddSP(this, &SDMStagePreview::OnStageUpdated);
+		OnStageUpdated(InStage, EDMUpdateType::Structure);
+	}
+
+	ChildSlot
+	[
+		SNew(SImage)
+		.Image(&Brush)
+		.DesiredSizeOverride(this, &SDMStagePreview::GetPreviewSize)
+	];
+}
+
+void SDMStagePreview::Tick(const FGeometry& AllottedGeometry, const double InCurrentTime, const float InDeltaTime)
+{
+	SCompoundWidget::Tick(AllottedGeometry, InCurrentTime, InDeltaTime);
+
+	if (UDMMaterialComponent::CanClean())
+	{
+		if (UDMMaterialStage* Stage = StageWeak.Get())
+		{
+			if (Stage->IsComponentValid() && Stage->NeedsClean())
+			{
+				Stage->DoClean();
+				OnStageUpdated(StageWeak.Get(), EDMUpdateType::Structure);
+			}
+		}
+	}	
+}
+
+void SDMStagePreview::OnStageUpdated(UDMMaterialComponent* InComponent, EDMUpdateType InUpdateType)
+{
+	if (UDMMaterialStage* Stage = Cast<UDMMaterialStage>(InComponent))
+	{
+		if (Stage == StageWeak.Get() && IsValid(Stage) && Stage->IsComponentValid())
+		{
+			UMaterialInterface* PreviewMaterial = Stage->GetPreviewMaterial();
+
+			if (Brush.GetResourceObject() != PreviewMaterial)
+			{
+				Brush.SetMaterial(PreviewMaterial);
+			}
+		}
+	}
+}
+
+TOptional<FVector2D> SDMStagePreview::GetPreviewSize() const
+{
+	return PreviewSize.Get();
+}
diff --git a/Engine/Plugins/Experimental/DynamicMaterial/Source/DynamicMaterialEditor/Private/Slate/Previews/SDMStagePreview.h b/Engine/Plugins/Experimental/DynamicMaterial/Source/DynamicMaterialEditor/Private/Slate/Previews/SDMStagePreview.h
new file mode 100644
--- /dev/null
+++ b/Engine/Plugins/Experimental/DynamicMaterial/Source/DynamicMaterialEditor/Private/Slate/Previews/SDMStagePreview.h
@@ -0,0 +1,38 @@
+// Copyright Epic Games, Inc. All Rights Reserved.
+
+#pragma once
+
+#include "DMEDefs.h"
+#include "SlateMaterialBrush.h"
+#include "Widgets/SCompoundWidget.h"
+
+class UDMMaterialComponent;
+class UDMMaterialStage;
+
+class SDMStagePreview : public SCompoundWidget
+{
+public:
+	SLATE_BEGIN_ARGS(SDMStagePreview)
+		: _PreviewSize(FVector2D(48.f, 48.f))
+	{}
+		SLATE_ATTRIBUTE(FVector2D, PreviewSize)
+	SLATE_END_ARGS()
+
+	SDMStagePreview();
+	virtual ~SDMStagePreview() override;
+
+	void Construct(const FArguments& InArgs, UDMMaterialStage* InStage);
+
+	//~ Begin SWidget
+	virtual void Tick(const FGeometry& InAllottedGeometry, const double InCurrentTime, const float InDeltaTime) override;
+	//~ End SWidget
+
+protected:
+	TWeakObjectPtr<UDMMaterialStage> StageWeak;
+	FSlateMaterialBrush Brush;
+	TAttribute<FVector2D> PreviewSize;
+
+	void OnStageUpdated(UDMMaterialComponent* InComponent, EDMUpdateType InUpdateType);
+
+	TOptional<FVector2D> GetPreviewSize() const;
+};
diff --git a/Engine/Plugins/Experimental/DynamicMaterial/Source/DynamicMaterialEditor/Private/Slate/Previews/SDMStageSourcePreview.cpp b/Engine/Plugins/Experimental/DynamicMaterial/Source/DynamicMaterialEditor/Private/Slate/Previews/SDMStageSourcePreview.cpp
new file mode 100644
--- /dev/null
+++ b/Engine/Plugins/Experimental/DynamicMaterial/Source/DynamicMaterialEditor/Private/Slate/Previews/SDMStageSourcePreview.cpp
@@ -0,0 +1,84 @@
+// Copyright Epic Games, Inc. All Rights Reserved.
+
+#include "SDMStageSourcePreview.h"
+#include "Components/DMMaterialStageSource.h"
+#include "DynamicMaterialModule.h"
+#include "Materials/Material.h"
+#include "Widgets/Images/SImage.h"
+
+SDMStageSourcePreview::SDMStageSourcePreview()
+	: Brush(FSlateMaterialBrush(FVector2D(1.f, 1.f)))
+{
+
+}
+
+SDMStageSourcePreview::~SDMStageSourcePreview()
+{
+	if (!FDynamicMaterialModule::AreUObjectsSafe())
+	{
+		return;
+	}
+
+	UDMMaterialStageSource* StageSource = StageSourceWeak.Get();
+
+	if (StageSource)
+	{
+		StageSource->GetOnUpdate().RemoveAll(this);
+	}
+}
+
+void SDMStageSourcePreview::Construct(const FArguments& InArgs, UDMMaterialStageSource* InStageSource)
+{
+	StageSourceWeak = InStageSource;
+
+	SetCanTick(true);
+
+	Brush.SetImageSize(InArgs._DesiredSize);
+
+	if (ensure(IsValid(InStageSource)))
+	{
+		InStageSource->GetOnUpdate().AddSP(this, &SDMStageSourcePreview::OnStageSourceUpdated);
+		OnStageSourceUpdated(InStageSource, EDMUpdateType::Structure);
+	}
+
+	ChildSlot
+	[
+		SNew(SImage)
+		.Image(&Brush)
+		.DesiredSizeOverride(InArgs._DesiredSize)
+	];
+}
+
+void SDMStageSourcePreview::Tick(const FGeometry& AllottedGeometry, const double InCurrentTime, const float InDeltaTime)
+{
+	SCompoundWidget::Tick(AllottedGeometry, InCurrentTime, InDeltaTime);
+
+	if (UDMMaterialComponent::CanClean())
+	{
+		if (UDMMaterialStageSource* Source = StageSourceWeak.Get())
+		{
+			if (Source->IsComponentValid() && Source->NeedsClean())
+			{
+				Source->DoClean();
+				OnStageSourceUpdated(Source, EDMUpdateType::Structure);
+			}
+		}
+	}
+}
+
+void SDMStageSourcePreview::OnStageSourceUpdated(UDMMaterialComponent* InComponent, EDMUpdateType InUpdateType)
+{
+	if (UDMMaterialStageSource* Source = Cast<UDMMaterialStageSource>(InComponent))
+	{
+		if (Source == StageSourceWeak.Get() && Source->IsComponentValid())
+		{
+			UMaterial* PreviewMaterial = Source->GetPreviewMaterial();
+
+			if (Brush.GetResourceObject() != PreviewMaterial)
+			{
+				Brush.SetMaterial(PreviewMaterial);
+			}
+		}
+	}
+}
+
diff --git a/Engine/Plugins/Experimental/DynamicMaterial/Source/DynamicMaterialEditor/Private/Slate/Previews/SDMStageSourcePreview.h b/Engine/Plugins/Experimental/DynamicMaterial/Source/DynamicMaterialEditor/Private/Slate/Previews/SDMStageSourcePreview.h
new file mode 100644
--- /dev/null
+++ b/Engine/Plugins/Experimental/DynamicMaterial/Source/DynamicMaterialEditor/Private/Slate/Previews/SDMStageSourcePreview.h
@@ -0,0 +1,34 @@
+// Copyright Epic Games, Inc. All Rights Reserved.
+
+#pragma once
+
+#include "DMEDefs.h"
+#include "SlateMaterialBrush.h"
+#include "Widgets/SCompoundWidget.h"
+
+class UDMMaterialComponent;
+class UDMMaterialStageSource;
+
+class SDMStageSourcePreview : public SCompoundWidget
+{
+	SLATE_BEGIN_ARGS(SDMStageSourcePreview)
+		: _DesiredSize(FVector2D(48.f, 48.f))
+		{}
+		SLATE_ARGUMENT(FVector2D, DesiredSize)
+	SLATE_END_ARGS()
+
+public:
+
+	SDMStageSourcePreview();
+	virtual ~SDMStageSourcePreview();
+
+	void Construct(const FArguments& InArgs, UDMMaterialStageSource* InStageSource);
+
+	virtual void Tick(const FGeometry& AllottedGeometry, const double InCurrentTime, const float InDeltaTime) override;
+
+protected:
+	TWeakObjectPtr<UDMMaterialStageSource> StageSourceWeak;
+	FSlateMaterialBrush Brush;
+
+	void OnStageSourceUpdated(UDMMaterialComponent* InComponent, EDMUpdateType InUpdateType);
+};
diff --git a/Engine/Plugins/Experimental/DynamicMaterial/Source/DynamicMaterialEditor/Private/Slate/Properties/Editors/SDMPropertyEditBool.cpp b/Engine/Plugins/Experimental/DynamicMaterial/Source/DynamicMaterialEditor/Private/Slate/Properties/Editors/SDMPropertyEditBool.cpp
new file mode 100644
--- /dev/null
+++ b/Engine/Plugins/Experimental/DynamicMaterial/Source/DynamicMaterialEditor/Private/Slate/Properties/Editors/SDMPropertyEditBool.cpp
@@ -0,0 +1,77 @@
+// Copyright Epic Games, Inc. All Rights Reserved.
+
+#include "Slate/Properties/Editors/SDMPropertyEditBool.h"
+#include "PropertyHandle.h"
+
+#define LOCTEXT_NAMESPACE "SDMPropertyEditBool"
+
+void SDMPropertyEditBool::Construct(const FArguments& InArgs, const TSharedPtr<IPropertyHandle>& InPropertyHandle)
+{
+	ensure(InPropertyHandle.IsValid());
+
+	SDMPropertyEdit::Construct(
+		SDMPropertyEdit::FArguments()
+			.InputCount(1)
+			.ComponentEditWidget(InArgs._ComponentEditWidget)
+			.PropertyHandle(InPropertyHandle)
+	);
+}
+
+TSharedRef<SWidget> SDMPropertyEditBool::GetComponentWidget(int32 InIndex)
+{
+	ensure(InIndex == 0);
+
+	return CreateCheckbox(
+		TAttribute<ECheckBoxState>::CreateSP(this, &SDMPropertyEditBool::IsChecked),
+		FOnCheckStateChanged::CreateSP(this, &SDMPropertyEditBool::OnValueToggled)
+	);
+}
+
+ECheckBoxState SDMPropertyEditBool::IsChecked() const
+{
+	if (PropertyHandle.IsValid() && PropertyHandle->IsValidHandle())
+	{
+		bool bValue;
+		PropertyHandle->GetValue(bValue);
+
+		return bValue ? ECheckBoxState::Checked : ECheckBoxState::Unchecked;
+	}
+
+	return ECheckBoxState::Unchecked;
+}
+
+void SDMPropertyEditBool::OnValueToggled(ECheckBoxState InNewState)
+{
+	if (PropertyHandle.IsValid() && PropertyHandle->IsValidHandle())
+	{
+		TArray<UObject*> Outers;
+		PropertyHandle->GetOuterObjects(Outers);
+
+		bool bCurrentValue = false;
+		const bool bNewValue = InNewState == ECheckBoxState::Checked;
+
+		if (PropertyHandle->GetValue(bCurrentValue) == FPropertyAccess::Fail)
+		{
+			return;
+		}
+
+		if (bCurrentValue == bNewValue)
+		{
+			return;
+		}
+
+		if (Outers.IsEmpty() == false && IsValid(Outers[0]))
+		{
+			StartTransaction(LOCTEXT("TransactionDescription", "Material Designer Value Set (Checkbox)"));
+		}
+
+		PropertyHandle->SetValue(bNewValue, EPropertyValueSetFlags::NotTransactable | EPropertyValueSetFlags::InteractiveChange);
+
+		if (Outers.IsEmpty() == false && IsValid(Outers[0]))
+		{
+			EndTransaction();
+		}
+	}
+}
+
+#undef LOCTEXT_NAMESPACE
diff --git a/Engine/Plugins/Experimental/DynamicMaterial/Source/DynamicMaterialEditor/Private/Slate/Properties/Editors/SDMPropertyEditBool.h b/Engine/Plugins/Experimental/DynamicMaterial/Source/DynamicMaterialEditor/Private/Slate/Properties/Editors/SDMPropertyEditBool.h
new file mode 100644
--- /dev/null
+++ b/Engine/Plugins/Experimental/DynamicMaterial/Source/DynamicMaterialEditor/Private/Slate/Properties/Editors/SDMPropertyEditBool.h
@@ -0,0 +1,30 @@
+// Copyright Epic Games, Inc. All Rights Reserved.
+
+#pragma once
+
+#include "Slate/Properties/SDMPropertyEdit.h"
+#include "Templates/SharedPointer.h"
+
+class IPropertyHandle;
+class SWidget;
+enum class ECheckBoxState : uint8;
+
+class SDMPropertyEditBool : public SDMPropertyEdit
+{
+	SLATE_BEGIN_ARGS(SDMPropertyEditBool)
+		{}
+		SLATE_ARGUMENT(TSharedPtr<SDMComponentEdit>, ComponentEditWidget)
+	SLATE_END_ARGS()
+
+public:
+	SDMPropertyEditBool() = default;
+	virtual ~SDMPropertyEditBool() override = default;
+
+	void Construct(const FArguments& InArgs, const TSharedPtr<IPropertyHandle>& InPropertyHandle);
+
+protected:
+	virtual TSharedRef<SWidget> GetComponentWidget(int32 InIndex) override;
+
+	ECheckBoxState IsChecked() const;
+	void OnValueToggled(ECheckBoxState InNewState);
+};
diff --git a/Engine/Plugins/Experimental/DynamicMaterial/Source/DynamicMaterialEditor/Private/Slate/Properties/Editors/SDMPropertyEditBoolValue.cpp b/Engine/Plugins/Experimental/DynamicMaterial/Source/DynamicMaterialEditor/Private/Slate/Properties/Editors/SDMPropertyEditBoolValue.cpp
new file mode 100644
--- /dev/null
+++ b/Engine/Plugins/Experimental/DynamicMaterial/Source/DynamicMaterialEditor/Private/Slate/Properties/Editors/SDMPropertyEditBoolValue.cpp
@@ -0,0 +1,75 @@
+// Copyright Epic Games, Inc. All Rights Reserved.
+ 
+#include "SDMPropertyEditBoolValue.h"
+#include "Components/MaterialValues/DMMaterialValueBool.h"
+#include "Widgets/Input/SCheckBox.h"
+#include "Widgets/SNullWidget.h"
+
+#define LOCTEXT_NAMESPACE "SDMPropertyEditBoolValue"
+ 
+TSharedPtr<SWidget> SDMPropertyEditBoolValue::CreateEditWidget(const TSharedPtr<SDMComponentEdit>& InComponentEditWidget, UDMMaterialValue* InBoolValue)
+{
+	return SNew(SDMPropertyEditBoolValue, Cast<UDMMaterialValueBool>(InBoolValue))
+		.ComponentEditWidget(InComponentEditWidget);
+}
+ 
+void SDMPropertyEditBoolValue::Construct(const FArguments& InArgs, UDMMaterialValueBool* InBoolValue)
+{
+	ensure(IsValid(InBoolValue));
+
+	SDMPropertyEdit::Construct(
+		SDMPropertyEdit::FArguments()
+			.InputCount(1)
+			.ComponentEditWidget(InArgs._ComponentEditWidget)
+			.PropertyMaterialValue(InBoolValue)
+	);
+}
+ 
+UDMMaterialValueBool* SDMPropertyEditBoolValue::GetBoolValue() const
+{
+	return Cast<UDMMaterialValueBool>(GetValue());
+}
+ 
+TSharedRef<SWidget> SDMPropertyEditBoolValue::GetComponentWidget(int32 InIndex)
+{
+	ensure(InIndex == 0);
+ 
+	UDMMaterialValueBool* BoolValue = GetBoolValue();
+
+	if (!IsValid(BoolValue) || !BoolValue->IsComponentValid())
+	{
+		return SNullWidget::NullWidget;
+	}
+ 
+	return CreateCheckbox(
+		TAttribute<ECheckBoxState>::CreateSP(this, &SDMPropertyEditBoolValue::IsChecked),
+		FOnCheckStateChanged::CreateSP(this, &SDMPropertyEditBoolValue::OnValueToggled)
+	);
+}
+ 
+ECheckBoxState SDMPropertyEditBoolValue::IsChecked() const
+{
+	UDMMaterialValueBool* BoolValue = GetBoolValue();
+
+	if (!IsValid(BoolValue) || !BoolValue->IsComponentValid())
+	{
+		return ECheckBoxState::Unchecked;
+	}
+ 
+	return BoolValue->GetValue() ? ECheckBoxState::Checked : ECheckBoxState::Unchecked;
+}
+ 
+void SDMPropertyEditBoolValue::OnValueToggled(ECheckBoxState InNewState)
+{
+	UDMMaterialValueBool* BoolValue = GetBoolValue();
+	const bool bNewValue = InNewState == ECheckBoxState::Checked;
+	
+	if (IsValid(BoolValue) && BoolValue->IsComponentValid() && BoolValue->GetValue() != bNewValue)
+	{
+		StartTransaction(LOCTEXT("TransactionDescription", "Material Designer Value Toggled (Checkbox)"));
+		BoolValue->SetValue(bNewValue);
+		EndTransaction();
+	}
+}
+
+#undef LOCTEXT_NAMESPACE
diff --git a/Engine/Plugins/Experimental/DynamicMaterial/Source/DynamicMaterialEditor/Private/Slate/Properties/Editors/SDMPropertyEditBoolValue.h b/Engine/Plugins/Experimental/DynamicMaterial/Source/DynamicMaterialEditor/Private/Slate/Properties/Editors/SDMPropertyEditBoolValue.h
new file mode 100644
--- /dev/null
+++ b/Engine/Plugins/Experimental/DynamicMaterial/Source/DynamicMaterialEditor/Private/Slate/Properties/Editors/SDMPropertyEditBoolValue.h
@@ -0,0 +1,35 @@
+// Copyright Epic Games, Inc. All Rights Reserved.
+
+#pragma once
+
+#include "Slate/Properties/SDMPropertyEdit.h"
+#include "Templates/SharedPointer.h"
+
+class FAssetThumbnailPool;
+class UDMMaterialValue;
+class UDMMaterialValueBool;
+enum class ECheckBoxState : uint8;
+
+class SDMPropertyEditBoolValue : public SDMPropertyEdit
+{
+	SLATE_BEGIN_ARGS(SDMPropertyEditBoolValue)
+		{}
+		SLATE_ARGUMENT(TSharedPtr<SDMComponentEdit>, ComponentEditWidget)
+	SLATE_END_ARGS()
+
+public:
+	static TSharedPtr<SWidget> CreateEditWidget(const TSharedPtr<SDMComponentEdit>& InComponentEditWidget, UDMMaterialValue* InBoolValue);
+
+	SDMPropertyEditBoolValue() = default;
+	virtual ~SDMPropertyEditBoolValue() override = default;
+
+	void Construct(const FArguments& InArgs, UDMMaterialValueBool* InBoolValue);
+
+	UDMMaterialValueBool* GetBoolValue() const;
+
+protected:
+	virtual TSharedRef<SWidget> GetComponentWidget(int32 InIndex)override;
+
+	ECheckBoxState IsChecked() const;
+	void OnValueToggled(ECheckBoxState InNewState);
+};
