[MaterialDesigner]

diff --git a/Engine/Plugins/Experimental/DynamicMaterial/Source/DynamicMaterialEditor/Private/Slate/Properties/Editors/SDMPropertyEditColor.cpp b/Engine/Plugins/Experimental/DynamicMaterial/Source/DynamicMaterialEditor/Private/Slate/Properties/Editors/SDMPropertyEditColor.cpp
new file mode 100644
--- /dev/null
+++ b/Engine/Plugins/Experimental/DynamicMaterial/Source/DynamicMaterialEditor/Private/Slate/Properties/Editors/SDMPropertyEditColor.cpp
@@ -0,0 +1,93 @@
+// Copyright Epic Games, Inc. All Rights Reserved.
+
+#include "Slate/Properties/Editors/SDMPropertyEditColor.h"
+#include "PropertyHandle.h"
+
+#define LOCTEXT_NAMESPACE "SDMPropertyEditColor"
+
+void SDMPropertyEditColor::Construct(const FArguments& InArgs, const TSharedPtr<IPropertyHandle>& InPropertyHandle)
+{
+	ensure(InPropertyHandle.IsValid());
+
+	SDMPropertyEdit::Construct(
+		SDMPropertyEdit::FArguments()
+			.InputCount(1)
+			.ComponentEditWidget(InArgs._ComponentEditWidget)
+			.PropertyHandle(InPropertyHandle)
+	);
+}
+
+TSharedRef<SWidget> SDMPropertyEditColor::GetComponentWidget(int32 InIndex)
+{
+	ensure(InIndex == 0);
+
+	return CreateColorPicker(
+		true, // Use alpha
+		TAttribute<FLinearColor>::CreateSP(this, &SDMPropertyEditColor::GetColorValue),
+		FOnLinearColorValueChanged::CreateSP(this, &SDMPropertyEditColor::OnColorChanged),
+		LOCTEXT("TransactionDescription", "Material Designer Value Set (Color)")
+	);
+}
+
+FLinearColor SDMPropertyEditColor::GetColorValue() const
+{
+	if (PropertyHandle.IsValid() && PropertyHandle->IsValidHandle() && Property)
+	{
+		TArray<UObject*> Outers;
+		PropertyHandle->GetOuterObjects(Outers);
+
+		if (Outers.IsEmpty() == false && IsValid(Outers[0]))
+		{
+			FLinearColor Value = FLinearColor::Black;
+			Property->GetValue_InContainer(Outers[0], &Value);
+			return Value;
+		}
+	}
+
+	return FLinearColor::Black;
+}
+
+void SDMPropertyEditColor::OnColorChanged(FLinearColor InNewValue)
+{
+	if (PropertyHandle.IsValid() && PropertyHandle->IsValidHandle() && Property)
+	{
+		TArray<UObject*> Outers;
+		PropertyHandle->GetOuterObjects(Outers);
+
+		if (Outers.IsEmpty() == false && IsValid(Outers[0]))
+		{
+			FLinearColor CurrentValue = FLinearColor::Black;
+			Property->GetValue_InContainer(Outers[0], &CurrentValue);
+
+			const bool bSameValue = FMath::IsNearlyEqual(InNewValue.R, CurrentValue.R)
+				&& FMath::IsNearlyEqual(InNewValue.G, CurrentValue.G)
+				&& FMath::IsNearlyEqual(InNewValue.B, CurrentValue.B)
+				&& FMath::IsNearlyEqual(InNewValue.A, CurrentValue.A);
+
+			if (bSameValue)
+			{
+				return;
+			}
+
+			StartTransaction(LOCTEXT("TransactionDescription", "Material Designer Value Set (Color)"));
+
+			if (Property->HasSetter())
+			{
+				Property->CallSetter(Outers[0], &InNewValue);
+			}
+			else
+			{
+				Outers[0]->PreEditChange(Property);
+
+				Property->SetValue_InContainer(Outers[0], &InNewValue);
+
+				FPropertyChangedEvent PCE = FPropertyChangedEvent(Property, EPropertyChangeType::ValueSet);
+				Outers[0]->PostEditChangeProperty(PCE);
+			}
+
+			EndTransaction();
+		}		
+	}
+}
+
+#undef LOCTEXT_NAMESPACE
diff --git a/Engine/Plugins/Experimental/DynamicMaterial/Source/DynamicMaterialEditor/Private/Slate/Properties/Editors/SDMPropertyEditColor.h b/Engine/Plugins/Experimental/DynamicMaterial/Source/DynamicMaterialEditor/Private/Slate/Properties/Editors/SDMPropertyEditColor.h
new file mode 100644
--- /dev/null
+++ b/Engine/Plugins/Experimental/DynamicMaterial/Source/DynamicMaterialEditor/Private/Slate/Properties/Editors/SDMPropertyEditColor.h
@@ -0,0 +1,30 @@
+// Copyright Epic Games, Inc. All Rights Reserved.
+
+#pragma once
+
+#include "Slate/Properties/SDMPropertyEdit.h"
+#include "Templates/SharedPointer.h"
+
+class IPropertyHandle;
+class SWidget;
+struct FLinearColor;
+
+class SDMPropertyEditColor : public SDMPropertyEdit
+{
+	SLATE_BEGIN_ARGS(SDMPropertyEditColor)
+		{}
+		SLATE_ARGUMENT(TSharedPtr<SDMComponentEdit>, ComponentEditWidget)
+	SLATE_END_ARGS()
+
+public:
+	SDMPropertyEditColor() = default;
+	virtual ~SDMPropertyEditColor() override = default;
+
+	void Construct(const FArguments& InArgs, const TSharedPtr<IPropertyHandle>& InPropertyHandle);
+
+protected:
+	virtual TSharedRef<SWidget> GetComponentWidget(int32 InIndex) override;
+
+	FLinearColor GetColorValue() const;
+	void OnColorChanged(FLinearColor InNewValue);
+};
diff --git a/Engine/Plugins/Experimental/DynamicMaterial/Source/DynamicMaterialEditor/Private/Slate/Properties/Editors/SDMPropertyEditEnum.cpp b/Engine/Plugins/Experimental/DynamicMaterial/Source/DynamicMaterialEditor/Private/Slate/Properties/Editors/SDMPropertyEditEnum.cpp
new file mode 100644
--- /dev/null
+++ b/Engine/Plugins/Experimental/DynamicMaterial/Source/DynamicMaterialEditor/Private/Slate/Properties/Editors/SDMPropertyEditEnum.cpp
@@ -0,0 +1,181 @@
+// Copyright Epic Games, Inc. All Rights Reserved.
+
+#include "Slate/Properties/Editors/SDMPropertyEditEnum.h"
+#include "PropertyHandle.h"
+#include "Widgets/SToolTip.h"
+
+#define LOCTEXT_NAMESPACE "SDMPropertyEditEnum"
+
+void SDMPropertyEditEnum::Construct(const FArguments& InArgs, const TSharedPtr<IPropertyHandle>& InPropertyHandle)
+{
+	ensure(InPropertyHandle.IsValid());
+
+	SDMPropertyEdit::Construct(
+		SDMPropertyEdit::FArguments()
+			.InputCount(1)
+			.ComponentEditWidget(InArgs._ComponentEditWidget)
+			.PropertyHandle(InPropertyHandle)
+	);
+}
+
+TSharedRef<SWidget> SDMPropertyEditEnum::GetComponentWidget(int32 InIndex)
+{
+	ensure(InIndex == 0);
+
+	return CreateEnum(
+		FOnGetPropertyComboBoxValue::CreateSP(this, &SDMPropertyEditEnum::GetEnumString),
+		FOnGetPropertyComboBoxStrings::CreateSP(this, &SDMPropertyEditEnum::GetEnumStrings),
+		FOnPropertyComboBoxValueSelected::CreateSP(this, &SDMPropertyEditEnum::OnValueChanged)
+	);
+}
+
+int64 SDMPropertyEditEnum::GetEnumValue() const
+{
+	if (PropertyHandle.IsValid() && PropertyHandle->IsValidHandle())
+	{
+		if (FEnumProperty* EnumProperty = CastField<FEnumProperty>(Property))
+		{
+			TArray<UObject*> Outers;
+			PropertyHandle->GetOuterObjects(Outers);
+
+			if (Outers.IsEmpty() == false && IsValid(Outers[0]))
+			{
+				if (EnumProperty->HasGetter())
+				{
+					int64 Value = -1;
+					EnumProperty->CallGetter(Outers[0], &Value);
+					return Value;
+				}
+				else
+				{
+					return *EnumProperty->ContainerPtrToValuePtr<uint8>(Outers[0]);
+				}
+			}
+		}
+	}
+
+	return -1;
+}
+
+FString SDMPropertyEditEnum::GetEnumString() const
+{
+	if (PropertyHandle.IsValid() && PropertyHandle->IsValidHandle())
+	{
+		if (FEnumProperty* EnumProperty = CastField<FEnumProperty>(Property))
+		{
+			if (UEnum* Enum = EnumProperty->GetEnum())
+			{
+				int64 Value = GetEnumValue();
+
+				if (Enum->IsValidEnumValue(Value))
+				{
+					return Enum->GetDisplayNameTextByValue(Value).ToString();
+				}
+			}
+		}
+	}
+
+	return "";
+}
+
+void SDMPropertyEditEnum::GetEnumStrings(TArray<TSharedPtr<FString>>& OutComboBoxStrings, TArray<TSharedPtr<SToolTip>>& OutToolTips, 
+	TArray<bool>& OutRestrictedItems) const
+{
+	OutComboBoxStrings.Empty();
+	OutToolTips.Empty();
+	OutRestrictedItems.Empty();
+
+	if (PropertyHandle.IsValid() && PropertyHandle->IsValidHandle())
+	{
+		if (FEnumProperty* EnumProperty = CastField<FEnumProperty>(Property))
+		{
+			if (UEnum* Enum = EnumProperty->GetEnum())
+			{
+				int32 EnumCount = Enum->NumEnums();
+
+				OutComboBoxStrings.Reserve(EnumCount);
+				OutToolTips.Reserve(EnumCount);
+				OutRestrictedItems.Reserve(EnumCount);
+
+				for (int32 EnumIdx = 0; EnumIdx < EnumCount; ++EnumIdx)
+				{
+					int64 Value = Enum->GetValueByIndex(EnumIdx);
+
+					if (Value == Enum->GetMaxEnumValue())
+					{
+						break;
+					}
+
+					OutComboBoxStrings.Add(MakeShared<FString>(Enum->GetDisplayNameTextByValue(Value).ToString()));
+
+					FText ToolTip = Enum->GetToolTipTextByIndex(EnumIdx);
+
+					if (ToolTip.IsEmpty())
+					{
+						OutToolTips.Add(SNew(SToolTip).Text(Enum->GetDisplayNameTextByValue(Value)));
+					}
+					else
+					{
+						OutToolTips.Add(SNew(SToolTip).Text(ToolTip));
+					}
+
+					OutRestrictedItems.Add(false);
+				}
+			}
+		}
+	}
+}
+
+void SDMPropertyEditEnum::OnValueChanged(const FString& InNewValueString)
+{
+	if (PropertyHandle.IsValid() && PropertyHandle->IsValidHandle())
+	{
+		if (FEnumProperty* EnumProperty = CastField<FEnumProperty>(Property))
+		{
+			if (UEnum* Enum = EnumProperty->GetEnum())
+			{
+				TArray<UObject*> Outers;
+				PropertyHandle->GetOuterObjects(Outers);
+
+				if (Outers.IsEmpty() == false && IsValid(Outers[0]))
+				{
+					int32 EnumCount = Enum->NumEnums();
+
+					for (int32 EnumIdx = 0; EnumIdx < EnumCount; ++EnumIdx)
+					{
+						const FString EnumString = Enum->GetDisplayNameTextByIndex(EnumIdx).ToString();
+
+						if (EnumString == InNewValueString)
+						{
+							int64 Value = Enum->GetValueByIndex(EnumIdx);
+							uint8* Data = EnumProperty->ContainerPtrToValuePtr<uint8>(Outers[0]);
+
+							if (Value != *Data)
+							{
+								StartTransaction(LOCTEXT("TransactionDescription", "Material Designer Value Set (Enum)"));
+
+								if (Property->HasSetter())
+								{
+									Property->CallSetter(Outers[0], &Value);
+								}
+								else
+								{
+									Outers[0]->PreEditChange(Property);
+
+									*Data = Value;
+
+									FPropertyChangedEvent PCE = FPropertyChangedEvent(Property, EPropertyChangeType::ValueSet);
+									Outers[0]->PostEditChangeProperty(PCE);
+								}
+
+								EndTransaction();
+							}
+						}
+					}
+				}
+			}
+		}
+	}
+}
+
+#undef LOCTEXT_NAMESPACE
diff --git a/Engine/Plugins/Experimental/DynamicMaterial/Source/DynamicMaterialEditor/Private/Slate/Properties/Editors/SDMPropertyEditEnum.h b/Engine/Plugins/Experimental/DynamicMaterial/Source/DynamicMaterialEditor/Private/Slate/Properties/Editors/SDMPropertyEditEnum.h
new file mode 100644
--- /dev/null
+++ b/Engine/Plugins/Experimental/DynamicMaterial/Source/DynamicMaterialEditor/Private/Slate/Properties/Editors/SDMPropertyEditEnum.h
@@ -0,0 +1,32 @@
+// Copyright Epic Games, Inc. All Rights Reserved.
+
+#pragma once
+
+#include "Slate/Properties/SDMPropertyEdit.h"
+#include "Templates/SharedPointer.h"
+
+class FString;
+class IPropertyHandle;
+class SWidget;
+
+class SDMPropertyEditEnum : public SDMPropertyEdit
+{
+	SLATE_BEGIN_ARGS(SDMPropertyEditEnum)
+		{}
+		SLATE_ARGUMENT(TSharedPtr<SDMComponentEdit>, ComponentEditWidget)
+	SLATE_END_ARGS()
+
+public:
+	SDMPropertyEditEnum() = default;
+	virtual ~SDMPropertyEditEnum() override = default;
+
+	void Construct(const FArguments& InArgs, const TSharedPtr<IPropertyHandle>& InPropertyHandle);
+
+protected:
+	virtual TSharedRef<SWidget> GetComponentWidget(int32 InIndex)override;
+
+	int64 GetEnumValue() const;
+	FString GetEnumString() const;
+	void GetEnumStrings(TArray<TSharedPtr<FString>>& OutComboBoxStrings, TArray<TSharedPtr<SToolTip>>& OutToolTips, TArray<bool>& OutRestrictedItems) const;
+	void OnValueChanged(const FString& InNewValueString);
+};
diff --git a/Engine/Plugins/Experimental/DynamicMaterial/Source/DynamicMaterialEditor/Private/Slate/Properties/Editors/SDMPropertyEditFloat.cpp b/Engine/Plugins/Experimental/DynamicMaterial/Source/DynamicMaterialEditor/Private/Slate/Properties/Editors/SDMPropertyEditFloat.cpp
new file mode 100644
--- /dev/null
+++ b/Engine/Plugins/Experimental/DynamicMaterial/Source/DynamicMaterialEditor/Private/Slate/Properties/Editors/SDMPropertyEditFloat.cpp
@@ -0,0 +1,74 @@
+// Copyright Epic Games, Inc. All Rights Reserved.
+
+#include "Slate/Properties/Editors/SDMPropertyEditFloat.h"
+#include "DynamicMaterialEditorSettings.h"
+#include "PropertyHandle.h"
+
+#define LOCTEXT_NAMESPACE "SDMPropertyEditFloat"
+
+void SDMPropertyEditFloat::Construct(const FArguments& InArgs, const TSharedPtr<IPropertyHandle>& InPropertyHandle)
+{
+	ensure(InPropertyHandle.IsValid());
+
+	FloatInterval = InArgs._FloatInterval;
+
+	SDMPropertyEdit::Construct(
+		SDMPropertyEdit::FArguments()
+			.InputCount(1)
+			.ComponentEditWidget(InArgs._ComponentEditWidget)
+			.PropertyHandle(InPropertyHandle)
+	);
+}
+
+TSharedRef<SWidget> SDMPropertyEditFloat::GetComponentWidget(int32 InIndex)
+{
+	ensure(InIndex == 0);
+
+	return CreateSpinBox(
+		TAttribute<float>::CreateSP(this, &SDMPropertyEditFloat::GetFloatValue),
+		FOnFloatValueChanged::CreateSP(this, &SDMPropertyEditFloat::OnValueChanged),
+		LOCTEXT("TransactionDescription", "Material Designer Value Scrubbing (Float)"),
+		FloatInterval
+	);
+}
+
+float SDMPropertyEditFloat::GetMaxWidthForWidget(int32 InIndex) const
+{
+	ensure(InIndex == 0);
+
+	return UDynamicMaterialEditorSettings::Get()->MaxFloatSliderWidth;
+}
+
+float SDMPropertyEditFloat::GetFloatValue() const
+{
+	if (PropertyHandle.IsValid() && PropertyHandle->IsValidHandle())
+	{
+		float Value;
+		PropertyHandle->GetValue(Value);
+		return Value;
+	}
+
+	return 0.f;
+}
+
+void SDMPropertyEditFloat::OnValueChanged(float InNewValue) const
+{
+	if (PropertyHandle.IsValid() && PropertyHandle->IsValidHandle())
+	{
+		float CurrentValue = 0.f;
+
+		if (PropertyHandle->GetValue(CurrentValue) == FPropertyAccess::Fail)
+		{
+			return;
+		}
+
+		if (FMath::IsNearlyEqual(CurrentValue, InNewValue))
+		{
+			return;
+		}
+
+		PropertyHandle->SetValue(InNewValue, EPropertyValueSetFlags::NotTransactable | EPropertyValueSetFlags::InteractiveChange);
+	}
+}
+
+#undef LOCTEXT_NAMESPACE
diff --git a/Engine/Plugins/Experimental/DynamicMaterial/Source/DynamicMaterialEditor/Private/Slate/Properties/Editors/SDMPropertyEditFloat.h b/Engine/Plugins/Experimental/DynamicMaterial/Source/DynamicMaterialEditor/Private/Slate/Properties/Editors/SDMPropertyEditFloat.h
new file mode 100644
--- /dev/null
+++ b/Engine/Plugins/Experimental/DynamicMaterial/Source/DynamicMaterialEditor/Private/Slate/Properties/Editors/SDMPropertyEditFloat.h
@@ -0,0 +1,34 @@
+// Copyright Epic Games, Inc. All Rights Reserved.
+
+#pragma once
+
+#include "Slate/Properties/SDMPropertyEdit.h"
+#include "Templates/SharedPointer.h"
+
+class IPropertyHandle;
+class SWidget;
+
+class SDMPropertyEditFloat : public SDMPropertyEdit
+{
+	SLATE_BEGIN_ARGS(SDMPropertyEditFloat)
+		: _FloatInterval(nullptr)
+		{}
+		SLATE_ARGUMENT(TSharedPtr<SDMComponentEdit>, ComponentEditWidget)
+		SLATE_ARGUMENT(const FFloatInterval*, FloatInterval)
+	SLATE_END_ARGS()
+
+public:
+	SDMPropertyEditFloat() = default;
+	virtual ~SDMPropertyEditFloat() override = default;
+
+	void Construct(const FArguments& InArgs, const TSharedPtr<IPropertyHandle>& InPropertyHandle);
+
+protected:
+	const FFloatInterval* FloatInterval;
+
+	virtual TSharedRef<SWidget> GetComponentWidget(int32 InIndex) override;
+	virtual float GetMaxWidthForWidget(int32 InIndex) const override;
+
+	virtual float GetFloatValue() const;
+	void OnValueChanged(float InNewValue) const;
+};
diff --git a/Engine/Plugins/Experimental/DynamicMaterial/Source/DynamicMaterialEditor/Private/Slate/Properties/Editors/SDMPropertyEditFloat1Value.cpp b/Engine/Plugins/Experimental/DynamicMaterial/Source/DynamicMaterialEditor/Private/Slate/Properties/Editors/SDMPropertyEditFloat1Value.cpp
new file mode 100644
--- /dev/null
+++ b/Engine/Plugins/Experimental/DynamicMaterial/Source/DynamicMaterialEditor/Private/Slate/Properties/Editors/SDMPropertyEditFloat1Value.cpp
@@ -0,0 +1,86 @@
+// Copyright Epic Games, Inc. All Rights Reserved.
+ 
+#include "SDMPropertyEditFloat1Value.h"
+#include "Components/MaterialValues/DMMaterialValueFloat1.h"
+#include "DMPrivate.h"
+#include "DynamicMaterialEditorSettings.h"
+#include "Widgets/Input/SSpinBox.h"
+#include "Widgets/SNullWidget.h"
+ 
+#define LOCTEXT_NAMESPACE "SDMPropertyEditFloat1Value"
+
+TSharedPtr<SWidget> SDMPropertyEditFloat1Value::CreateEditWidget(const TSharedPtr<SDMComponentEdit>& InComponentEditWidget, UDMMaterialValue* InFloat1Value)
+{
+	return SNew(SDMPropertyEditFloat1Value, Cast<UDMMaterialValueFloat1>(InFloat1Value))
+		.ComponentEditWidget(InComponentEditWidget);
+}
+ 
+void SDMPropertyEditFloat1Value::Construct(const FArguments& InArgs, UDMMaterialValueFloat1* InFloat1Value)
+{
+	ensure(IsValid(InFloat1Value));
+
+	SDMPropertyEdit::Construct(
+		SDMPropertyEdit::FArguments()
+			.InputCount(1)
+			.ComponentEditWidget(InArgs._ComponentEditWidget)
+			.PropertyMaterialValue(InFloat1Value)
+	);
+}
+ 
+UDMMaterialValueFloat1* SDMPropertyEditFloat1Value::GetFloat1Value() const
+{
+	return Cast<UDMMaterialValueFloat1>(GetValue());
+}
+ 
+TSharedRef<SWidget> SDMPropertyEditFloat1Value::GetComponentWidget(int32 InIndex)
+{
+	ensure(InIndex == 0);
+ 
+	UDMMaterialValueFloat1* Float1Value = GetFloat1Value();
+
+	if (!IsValid(Float1Value) || !Float1Value->IsComponentValid())
+	{
+		return SNullWidget::NullWidget;
+	}
+ 
+	const FFloatInterval& ValueRange = Float1Value->GetValueRange();
+ 
+	return CreateSpinBox(
+		TAttribute<float>::CreateSP(this, &SDMPropertyEditFloat1Value::GetSpinBoxValue),
+		SSpinBox<float>::FOnValueChanged::CreateSP(this, &SDMPropertyEditFloat1Value::OnSpinBoxValueChanged),
+		LOCTEXT("TransactionDescription", "Material Designer Value Scrubbing (Float)"),
+		Float1Value->HasValueRange() ? &ValueRange : nullptr
+	);
+}
+
+float SDMPropertyEditFloat1Value::GetMaxWidthForWidget(int32 InIndex) const
+{
+	ensure(InIndex == 0);
+
+	return UDynamicMaterialEditorSettings::Get()->MaxFloatSliderWidth;
+}
+
+float SDMPropertyEditFloat1Value::GetSpinBoxValue() const
+{
+	UDMMaterialValueFloat1* Float1Value = GetFloat1Value();
+
+	if (!IsValid(Float1Value) || !Float1Value->IsComponentValid())
+	{
+		return 0;
+	}
+ 
+	return Float1Value->GetValue();
+}
+ 
+void SDMPropertyEditFloat1Value::OnSpinBoxValueChanged(float InNewValue) const
+{
+	UDMMaterialValueFloat1* Float1Value = GetFloat1Value();
+	
+	if (IsValid(Float1Value) && Float1Value->IsComponentValid() 
+		&& FMath::IsNearlyEqual(Float1Value->GetValue(), InNewValue) == false)
+	{
+		Float1Value->SetValue(InNewValue);
+	}
+}
+ 
+#undef LOCTEXT_NAMESPACE
diff --git a/Engine/Plugins/Experimental/DynamicMaterial/Source/DynamicMaterialEditor/Private/Slate/Properties/Editors/SDMPropertyEditFloat1Value.h b/Engine/Plugins/Experimental/DynamicMaterial/Source/DynamicMaterialEditor/Private/Slate/Properties/Editors/SDMPropertyEditFloat1Value.h
new file mode 100644
--- /dev/null
+++ b/Engine/Plugins/Experimental/DynamicMaterial/Source/DynamicMaterialEditor/Private/Slate/Properties/Editors/SDMPropertyEditFloat1Value.h
@@ -0,0 +1,33 @@
+// Copyright Epic Games, Inc. All Rights Reserved.
+
+#pragma once
+
+#include "Slate/Properties/SDMPropertyEdit.h"
+
+class UDMMaterialValue;
+class UDMMaterialValueFloat1;
+
+class SDMPropertyEditFloat1Value : public SDMPropertyEdit
+{
+	SLATE_BEGIN_ARGS(SDMPropertyEditFloat1Value)
+		{}
+		SLATE_ARGUMENT(TSharedPtr<SDMComponentEdit>, ComponentEditWidget)
+	SLATE_END_ARGS()
+
+public:
+	static TSharedPtr<SWidget> CreateEditWidget(const TSharedPtr<SDMComponentEdit>& InComponentEditWidget, UDMMaterialValue* InFloat1Value);
+
+	SDMPropertyEditFloat1Value() = default;
+	virtual ~SDMPropertyEditFloat1Value() override = default;
+
+	void Construct(const FArguments& InArgs, UDMMaterialValueFloat1* InFloat1Value);
+
+	UDMMaterialValueFloat1* GetFloat1Value() const;
+
+protected:
+	virtual TSharedRef<SWidget> GetComponentWidget(int32 InIndex) override;
+	virtual float GetMaxWidthForWidget(int32 InIndex) const override;
+
+	float GetSpinBoxValue() const;
+	void OnSpinBoxValueChanged(float InNewValue) const;
+};
diff --git a/Engine/Plugins/Experimental/DynamicMaterial/Source/DynamicMaterialEditor/Private/Slate/Properties/Editors/SDMPropertyEditFloat2Value.cpp b/Engine/Plugins/Experimental/DynamicMaterial/Source/DynamicMaterialEditor/Private/Slate/Properties/Editors/SDMPropertyEditFloat2Value.cpp
new file mode 100644
--- /dev/null
+++ b/Engine/Plugins/Experimental/DynamicMaterial/Source/DynamicMaterialEditor/Private/Slate/Properties/Editors/SDMPropertyEditFloat2Value.cpp
@@ -0,0 +1,99 @@
+// Copyright Epic Games, Inc. All Rights Reserved.
+ 
+#include "SDMPropertyEditFloat2Value.h"
+#include "Components/MaterialValues/DMMaterialValueFloat2.h"
+#include "DMPrivate.h"
+#include "DynamicMaterialEditorSettings.h"
+#include "Widgets/Input/SSpinBox.h"
+#include "Widgets/SNullWidget.h"
+ 
+#define LOCTEXT_NAMESPACE "SDMPropertyEditFloat2Value"
+
+TSharedPtr<SWidget> SDMPropertyEditFloat2Value::CreateEditWidget(const TSharedPtr<SDMComponentEdit>& InComponentEditWidget, UDMMaterialValue* InFloat2Value)
+{
+	return SNew(SDMPropertyEditFloat2Value, Cast<UDMMaterialValueFloat2>(InFloat2Value))
+		.ComponentEditWidget(InComponentEditWidget);
+}
+ 
+void SDMPropertyEditFloat2Value::Construct(const FArguments& InArgs, UDMMaterialValueFloat2* InFloat2Value)
+{
+	ensure(IsValid(InFloat2Value));
+
+	SDMPropertyEdit::Construct(
+		SDMPropertyEdit::FArguments()
+			.InputCount(2)
+			.ComponentEditWidget(InArgs._ComponentEditWidget)
+			.PropertyMaterialValue(InFloat2Value)
+	);
+}
+ 
+UDMMaterialValueFloat2* SDMPropertyEditFloat2Value::GetFloat2Value() const
+{
+	return Cast<UDMMaterialValueFloat2>(GetValue());
+}
+ 
+TSharedRef<SWidget> SDMPropertyEditFloat2Value::GetComponentWidget(int32 InIndex)
+{
+	ensure(InIndex == 0 || InIndex == 1);
+ 
+	UDMMaterialValueFloat2* Float2Value = GetFloat2Value();
+
+	if (!IsValid(Float2Value) || !Float2Value->IsComponentValid())
+	{
+		return SNullWidget::NullWidget;
+	}
+ 
+	const FFloatInterval& ValueRange = Float2Value->GetValueRange();
+	static const TArray<FText> Labels = {
+		LOCTEXT("X", "X:"),
+		LOCTEXT("Y", "Y:")
+	};
+ 
+	return AddWidgetLabel(
+		Labels[InIndex],
+		CreateSpinBox(
+			TAttribute<float>::CreateSP(this, &SDMPropertyEditFloat2Value::GetSpinBoxValue, InIndex),
+			SSpinBox<float>::FOnValueChanged::CreateSP(this, &SDMPropertyEditFloat2Value::OnSpinBoxValueChanged, InIndex),
+			LOCTEXT("TransactionDescription", "Material Designer Value Scrubbing (UV)"),
+			Float2Value->HasValueRange() ? &ValueRange : nullptr
+		)
+	);
+}
+
+float SDMPropertyEditFloat2Value::GetMaxWidthForWidget(int32 InIndex) const
+{
+	ensure(InIndex == 0 || InIndex == 1);
+
+	return UDynamicMaterialEditorSettings::Get()->MaxFloatSliderWidth;
+}
+
+float SDMPropertyEditFloat2Value::GetSpinBoxValue(int32 InComponent) const
+{
+	ensure(InComponent == 0 || InComponent == 1);
+ 
+	UDMMaterialValueFloat2* Float2Value = GetFloat2Value();
+
+	if (!IsValid(Float2Value) || !Float2Value->IsComponentValid())
+	{
+		return 0;
+	}
+ 
+	return Float2Value->GetValue()[InComponent];
+}
+ 
+void SDMPropertyEditFloat2Value::OnSpinBoxValueChanged(float InNewValue, int32 InComponent) const
+{
+	ensure(InComponent == 0 || InComponent == 1);
+ 
+	UDMMaterialValueFloat2* Float2Value = GetFloat2Value();
+
+	if (IsValid(Float2Value) && Float2Value->IsComponentValid() 
+		&& FMath::IsNearlyEqual(Float2Value->GetValue()[InComponent], InNewValue) == false)
+	{
+		FVector2D NewTotalValue = Float2Value->GetValue();
+		NewTotalValue[InComponent] = InNewValue;
+		Float2Value->SetValue(NewTotalValue);
+	}
+}
+ 
+#undef LOCTEXT_NAMESPACE
diff --git a/Engine/Plugins/Experimental/DynamicMaterial/Source/DynamicMaterialEditor/Private/Slate/Properties/Editors/SDMPropertyEditFloat2Value.h b/Engine/Plugins/Experimental/DynamicMaterial/Source/DynamicMaterialEditor/Private/Slate/Properties/Editors/SDMPropertyEditFloat2Value.h
new file mode 100644
--- /dev/null
+++ b/Engine/Plugins/Experimental/DynamicMaterial/Source/DynamicMaterialEditor/Private/Slate/Properties/Editors/SDMPropertyEditFloat2Value.h
@@ -0,0 +1,33 @@
+// Copyright Epic Games, Inc. All Rights Reserved.
+
+#pragma once
+
+#include "Slate/Properties/SDMPropertyEdit.h"
+
+class UDMMaterialValue;
+class UDMMaterialValueFloat2;
+
+class SDMPropertyEditFloat2Value : public SDMPropertyEdit
+{
+	SLATE_BEGIN_ARGS(SDMPropertyEditFloat2Value)
+		{}
+		SLATE_ARGUMENT(TSharedPtr<SDMComponentEdit>, ComponentEditWidget)
+	SLATE_END_ARGS()
+
+public:
+	static TSharedPtr<SWidget> CreateEditWidget(const TSharedPtr<SDMComponentEdit>& InComponentEditWidget, UDMMaterialValue* InFloat2Value);
+
+	SDMPropertyEditFloat2Value() = default;
+	virtual ~SDMPropertyEditFloat2Value() override = default;
+
+	void Construct(const FArguments& InArgs, UDMMaterialValueFloat2* InFloat2Value);
+
+	UDMMaterialValueFloat2* GetFloat2Value() const;
+
+protected:
+	virtual TSharedRef<SWidget> GetComponentWidget(int32 InIndex) override;
+	virtual float GetMaxWidthForWidget(int32 InIndex) const override;
+
+	float GetSpinBoxValue(int32 InComponent) const;
+	void OnSpinBoxValueChanged(float InNewValue, int32 InComponent) const;
+};
diff --git a/Engine/Plugins/Experimental/DynamicMaterial/Source/DynamicMaterialEditor/Private/Slate/Properties/Editors/SDMPropertyEditFloat3RGBValue.cpp b/Engine/Plugins/Experimental/DynamicMaterial/Source/DynamicMaterialEditor/Private/Slate/Properties/Editors/SDMPropertyEditFloat3RGBValue.cpp
new file mode 100644
--- /dev/null
+++ b/Engine/Plugins/Experimental/DynamicMaterial/Source/DynamicMaterialEditor/Private/Slate/Properties/Editors/SDMPropertyEditFloat3RGBValue.cpp
@@ -0,0 +1,87 @@
+// Copyright Epic Games, Inc. All Rights Reserved.
+ 
+#include "SDMPropertyEditFloat3RGBValue.h"
+#include "Components/MaterialValues/DMMaterialValueFloat3RGB.h"
+#include "DMPrivate.h"
+#include "Widgets/SNullWidget.h"
+
+#define LOCTEXT_NAMESPACE "SDMPropertyEditFloat3RGBValue"
+
+TSharedPtr<SWidget> SDMPropertyEditFloat3RGBValue::CreateEditWidget(const TSharedPtr<SDMComponentEdit>& InComponentEditWidget, UDMMaterialValue* InFloat3RGBValue)
+{
+	return SNew(SDMPropertyEditFloat3RGBValue, Cast<UDMMaterialValueFloat3RGB>(InFloat3RGBValue))
+		.ComponentEditWidget(InComponentEditWidget);
+}
+ 
+void SDMPropertyEditFloat3RGBValue::Construct(const FArguments& InArgs, UDMMaterialValueFloat3RGB* InFloat3RGBValue)
+{
+	ensure(IsValid(InFloat3RGBValue));
+
+	SDMPropertyEdit::Construct(
+		SDMPropertyEdit::FArguments()
+			.InputCount(1)
+			.ComponentEditWidget(InArgs._ComponentEditWidget)
+			.PropertyMaterialValue(InFloat3RGBValue)
+	);
+}
+ 
+UDMMaterialValueFloat3RGB* SDMPropertyEditFloat3RGBValue::GetFloat3RGBValue() const
+{
+	return Cast<UDMMaterialValueFloat3RGB>(GetValue());
+}
+ 
+TSharedRef<SWidget> SDMPropertyEditFloat3RGBValue::GetComponentWidget(int32 InIndex)
+{
+	ensure(InIndex == 0);
+ 
+	UDMMaterialValueFloat3RGB* Float3RGBValue = GetFloat3RGBValue();
+
+	if (!IsValid(Float3RGBValue) || !Float3RGBValue->IsComponentValid())
+	{
+		return SNullWidget::NullWidget;
+	}
+ 
+	const FFloatInterval& ValueRange = Float3RGBValue->GetValueRange();
+ 
+	return CreateColorPicker(
+		/* bUseAlpha */ false,
+		TAttribute<FLinearColor>::CreateSP(this, &SDMPropertyEditFloat3RGBValue::GetColor),
+		FOnLinearColorValueChanged::CreateSP(this, &SDMPropertyEditFloat3RGBValue::OnColorValueChanged),
+		LOCTEXT("TransactionDescription", "Material Designer Value Scrubbing (RGB)"),
+		Float3RGBValue->HasValueRange() ? &ValueRange : nullptr
+	);
+}
+ 
+FLinearColor SDMPropertyEditFloat3RGBValue::GetColor() const
+{
+	UDMMaterialValueFloat3RGB* Float3RGBValue = GetFloat3RGBValue();
+
+	if (!IsValid(Float3RGBValue) || !Float3RGBValue->IsComponentValid())
+	{
+		return FLinearColor::Black;
+	}
+ 
+	return Float3RGBValue->GetValue();
+}
+ 
+void SDMPropertyEditFloat3RGBValue::OnColorValueChanged(FLinearColor InNewColor) const
+{
+	UDMMaterialValueFloat3RGB* Float3RGBValue = GetFloat3RGBValue();
+
+	if (IsValid(Float3RGBValue) && Float3RGBValue->IsComponentValid())
+	{
+		InNewColor.A = 1.f;
+		const FLinearColor& CurrentColor = Float3RGBValue->GetValue();
+
+		const bool bSameValue = FMath::IsNearlyEqual(InNewColor.R, CurrentColor.R)
+			&& FMath::IsNearlyEqual(InNewColor.G, CurrentColor.G)
+			&& FMath::IsNearlyEqual(InNewColor.B, CurrentColor.B);
+
+		if (!bSameValue)
+		{
+			Float3RGBValue->SetValue(InNewColor);
+		}
+	}
+}
+ 
+#undef LOCTEXT_NAMESPACE
diff --git a/Engine/Plugins/Experimental/DynamicMaterial/Source/DynamicMaterialEditor/Private/Slate/Properties/Editors/SDMPropertyEditFloat3RGBValue.h b/Engine/Plugins/Experimental/DynamicMaterial/Source/DynamicMaterialEditor/Private/Slate/Properties/Editors/SDMPropertyEditFloat3RGBValue.h
new file mode 100644
--- /dev/null
+++ b/Engine/Plugins/Experimental/DynamicMaterial/Source/DynamicMaterialEditor/Private/Slate/Properties/Editors/SDMPropertyEditFloat3RGBValue.h
@@ -0,0 +1,32 @@
+// Copyright Epic Games, Inc. All Rights Reserved.
+
+#pragma once
+
+#include "Slate/Properties/SDMPropertyEdit.h"
+
+class UDMMaterialValue;
+class UDMMaterialValueFloat3RGB;
+
+class SDMPropertyEditFloat3RGBValue : public SDMPropertyEdit
+{
+	SLATE_BEGIN_ARGS(SDMPropertyEditFloat3RGBValue)
+		{}
+		SLATE_ARGUMENT(TSharedPtr<SDMComponentEdit>, ComponentEditWidget)
+	SLATE_END_ARGS()
+
+public:
+	static TSharedPtr<SWidget> CreateEditWidget(const TSharedPtr<SDMComponentEdit>& InComponentEditWidget, UDMMaterialValue* InFloat3RGBValue);
+
+	SDMPropertyEditFloat3RGBValue() = default;
+	virtual ~SDMPropertyEditFloat3RGBValue() override = default;
+
+	void Construct(const FArguments& InArgs, UDMMaterialValueFloat3RGB* InFloat3RGBValue);
+
+	UDMMaterialValueFloat3RGB* GetFloat3RGBValue() const;
+
+protected:
+	virtual TSharedRef<SWidget> GetComponentWidget(int32 InIndex) override;
+
+	FLinearColor GetColor() const;
+	void OnColorValueChanged(FLinearColor InNewColor) const;
+};
diff --git a/Engine/Plugins/Experimental/DynamicMaterial/Source/DynamicMaterialEditor/Private/Slate/Properties/Editors/SDMPropertyEditFloat3RPYValue.cpp b/Engine/Plugins/Experimental/DynamicMaterial/Source/DynamicMaterialEditor/Private/Slate/Properties/Editors/SDMPropertyEditFloat3RPYValue.cpp
new file mode 100644
--- /dev/null
+++ b/Engine/Plugins/Experimental/DynamicMaterial/Source/DynamicMaterialEditor/Private/Slate/Properties/Editors/SDMPropertyEditFloat3RPYValue.cpp
@@ -0,0 +1,121 @@
+// Copyright Epic Games, Inc. All Rights Reserved.
+ 
+#include "SDMPropertyEditFloat3RPYValue.h"
+#include "Components/MaterialValues/DMMaterialValueFloat3RPY.h"
+#include "DMPrivate.h"
+#include "DynamicMaterialEditorSettings.h"
+#include "Widgets/Input/SSpinBox.h"
+#include "Widgets/SNullWidget.h"
+ 
+#define LOCTEXT_NAMESPACE "SDMPropertyEditFloat3RPYValue"
+
+TSharedPtr<SWidget> SDMPropertyEditFloat3RPYValue::CreateEditWidget(const TSharedPtr<SDMComponentEdit>& InComponentEditWidget, UDMMaterialValue* InFloat3RPYValue)
+{
+	return SNew(SDMPropertyEditFloat3RPYValue, Cast<UDMMaterialValueFloat3RPY>(InFloat3RPYValue))
+		.ComponentEditWidget(InComponentEditWidget);
+}
+ 
+void SDMPropertyEditFloat3RPYValue::Construct(const FArguments& InArgs, UDMMaterialValueFloat3RPY* InFloat3RPYValue)
+{
+	ensure(IsValid(InFloat3RPYValue));
+
+	SDMPropertyEdit::Construct(
+		SDMPropertyEdit::FArguments()
+			.InputCount(3)
+			.ComponentEditWidget(InArgs._ComponentEditWidget)
+			.PropertyMaterialValue(InFloat3RPYValue)
+	);
+}
+ 
+UDMMaterialValueFloat3RPY* SDMPropertyEditFloat3RPYValue::GetFloat3RPYValue() const
+{
+	return Cast<UDMMaterialValueFloat3RPY>(GetValue());
+}
+ 
+TSharedRef<SWidget> SDMPropertyEditFloat3RPYValue::GetComponentWidget(int32 InIndex)
+{
+	ensure(InIndex >= 0 && InIndex <= 2);
+ 
+	UDMMaterialValueFloat3RPY* Float3RPYValue = GetFloat3RPYValue();
+
+	if (!IsValid(Float3RPYValue) || !Float3RPYValue->IsComponentValid())
+	{
+		return SNullWidget::NullWidget;
+	}
+ 
+	EAxis::Type ComponentAxis = EAxis::None;
+ 
+	switch (InIndex)
+	{
+		case 0:
+			ComponentAxis = EAxis::X;
+			break;
+ 
+		case 1:
+			ComponentAxis = EAxis::Y;
+			break;
+ 
+		case 2:
+			ComponentAxis = EAxis::Z;
+			break;
+ 
+		default:
+			checkNoEntry();
+			break;
+	}
+ 
+	const FFloatInterval& ValueRange = Float3RPYValue->GetValueRange();
+	static const TArray<FText> Labels = {
+		LOCTEXT("R", "R:"),
+		LOCTEXT("P", "P:"),
+		LOCTEXT("Y", "Y:")
+	};
+ 
+	return AddWidgetLabel(
+		Labels[InIndex],
+		CreateSpinBox(
+			TAttribute<float>::CreateSP(this, &SDMPropertyEditFloat3RPYValue::GetSpinBoxValue, ComponentAxis),
+			SSpinBox<float>::FOnValueChanged::CreateSP(this, &SDMPropertyEditFloat3RPYValue::OnSpinBoxValueChanged, ComponentAxis),
+			LOCTEXT("TransactionDescription", "Material Designer Value Scrubbing (Rotator)"),
+			Float3RPYValue->HasValueRange() ? &ValueRange : nullptr
+		)
+	);
+}
+
+float SDMPropertyEditFloat3RPYValue::GetMaxWidthForWidget(int32 InIndex) const
+{
+	ensure(InIndex >= 0 && InIndex <= 2);
+
+	return UDynamicMaterialEditorSettings::Get()->MaxFloatSliderWidth;
+}
+
+float SDMPropertyEditFloat3RPYValue::GetSpinBoxValue(EAxis::Type InAxis) const
+{
+	ensure(InAxis == EAxis::X || InAxis == EAxis::Y || InAxis == EAxis::Z);
+ 
+	UDMMaterialValueFloat3RPY* Float3RPYValue = GetFloat3RPYValue();
+
+	if (!IsValid(Float3RPYValue) || !Float3RPYValue->IsComponentValid())
+	{
+		return 0;
+	}
+ 
+	return Float3RPYValue->GetValue().GetComponentForAxis(InAxis);
+}
+ 
+void SDMPropertyEditFloat3RPYValue::OnSpinBoxValueChanged(float InNewValue, EAxis::Type InAxis) const
+{
+	ensure(InAxis == EAxis::X || InAxis == EAxis::Y || InAxis == EAxis::Z);
+ 
+	UDMMaterialValueFloat3RPY* Float3RPYValue = GetFloat3RPYValue();
+	
+	if (IsValid(Float3RPYValue) && Float3RPYValue->IsComponentValid() 
+		&& FMath::IsNearlyEqual(Float3RPYValue->GetValue().GetComponentForAxis(InAxis), InNewValue) == false)
+	{
+		FRotator NewTotalValue = Float3RPYValue->GetValue();
+		NewTotalValue.SetComponentForAxis(InAxis, InNewValue);
+		Float3RPYValue->SetValue(NewTotalValue);
+	}
+}
+ 
+#undef LOCTEXT_NAMESPACE
diff --git a/Engine/Plugins/Experimental/DynamicMaterial/Source/DynamicMaterialEditor/Private/Slate/Properties/Editors/SDMPropertyEditFloat3RPYValue.h b/Engine/Plugins/Experimental/DynamicMaterial/Source/DynamicMaterialEditor/Private/Slate/Properties/Editors/SDMPropertyEditFloat3RPYValue.h
new file mode 100644
--- /dev/null
+++ b/Engine/Plugins/Experimental/DynamicMaterial/Source/DynamicMaterialEditor/Private/Slate/Properties/Editors/SDMPropertyEditFloat3RPYValue.h
@@ -0,0 +1,33 @@
+// Copyright Epic Games, Inc. All Rights Reserved.
+
+#pragma once
+
+#include "Slate/Properties/SDMPropertyEdit.h"
+
+class UDMMaterialValue;
+class UDMMaterialValueFloat3RPY;
+
+class SDMPropertyEditFloat3RPYValue : public SDMPropertyEdit
+{
+	SLATE_BEGIN_ARGS(SDMPropertyEditFloat3RPYValue)
+		{}
+		SLATE_ARGUMENT(TSharedPtr<SDMComponentEdit>, ComponentEditWidget)
+	SLATE_END_ARGS()
+
+public:
+	static TSharedPtr<SWidget> CreateEditWidget(const TSharedPtr<SDMComponentEdit>& InComponentEditWidget, UDMMaterialValue* InFloat3RPYValue);
+
+	SDMPropertyEditFloat3RPYValue() = default;
+	virtual ~SDMPropertyEditFloat3RPYValue() override = default;
+
+	void Construct(const FArguments& InArgs, UDMMaterialValueFloat3RPY* InFloat3RPYValue);
+
+	UDMMaterialValueFloat3RPY* GetFloat3RPYValue() const;
+
+protected:
+	virtual TSharedRef<SWidget> GetComponentWidget(int32 InIndex) override;
+	virtual float GetMaxWidthForWidget(int32 InIndex) const override;
+
+	float GetSpinBoxValue(EAxis::Type InAxis) const;
+	void OnSpinBoxValueChanged(float InNewValue, EAxis::Type InAxis) const;
+};
diff --git a/Engine/Plugins/Experimental/DynamicMaterial/Source/DynamicMaterialEditor/Private/Slate/Properties/Editors/SDMPropertyEditFloat3XYZValue.cpp b/Engine/Plugins/Experimental/DynamicMaterial/Source/DynamicMaterialEditor/Private/Slate/Properties/Editors/SDMPropertyEditFloat3XYZValue.cpp
new file mode 100644
--- /dev/null
+++ b/Engine/Plugins/Experimental/DynamicMaterial/Source/DynamicMaterialEditor/Private/Slate/Properties/Editors/SDMPropertyEditFloat3XYZValue.cpp
@@ -0,0 +1,100 @@
+// Copyright Epic Games, Inc. All Rights Reserved.
+ 
+#include "SDMPropertyEditFloat3XYZValue.h"
+#include "Components/MaterialValues/DMMaterialValueFloat3XYZ.h"
+#include "DMPrivate.h"
+#include "DynamicMaterialEditorSettings.h"
+#include "Widgets/Input/SSpinBox.h"
+#include "Widgets/SNullWidget.h"
+ 
+#define LOCTEXT_NAMESPACE "SDMPropertyEditFloat3XYZValue"
+
+TSharedPtr<SWidget> SDMPropertyEditFloat3XYZValue::CreateEditWidget(const TSharedPtr<SDMComponentEdit>& InComponentEditWidget, UDMMaterialValue* InFloat3XYZValue)
+{
+	return SNew(SDMPropertyEditFloat3XYZValue, Cast<UDMMaterialValueFloat3XYZ>(InFloat3XYZValue))
+		.ComponentEditWidget(InComponentEditWidget);
+}
+ 
+void SDMPropertyEditFloat3XYZValue::Construct(const FArguments& InArgs, UDMMaterialValueFloat3XYZ* InFloat3XYZValue)
+{
+	ensure(IsValid(InFloat3XYZValue));
+
+	SDMPropertyEdit::Construct(
+		SDMPropertyEdit::FArguments()
+			.InputCount(3)
+			.ComponentEditWidget(InArgs._ComponentEditWidget)
+			.PropertyMaterialValue(InFloat3XYZValue)
+	);
+}
+ 
+UDMMaterialValueFloat3XYZ* SDMPropertyEditFloat3XYZValue::GetFloat3XYZValue() const
+{
+	return Cast<UDMMaterialValueFloat3XYZ>(GetValue());
+}
+ 
+TSharedRef<SWidget> SDMPropertyEditFloat3XYZValue::GetComponentWidget(int32 InIndex)
+{
+	ensure(InIndex >= 0 && InIndex <= 2);
+ 
+	UDMMaterialValueFloat3XYZ* Float3XYZValue = GetFloat3XYZValue();
+
+	if (!IsValid(Float3XYZValue) || !Float3XYZValue->IsComponentValid())
+	{
+		return SNullWidget::NullWidget;
+	}
+ 
+	const FFloatInterval& ValueRange = Float3XYZValue->GetValueRange();
+	static const TArray<FText> Labels = {
+		LOCTEXT("X", "X:"),
+		LOCTEXT("Y", "Y:"),
+		LOCTEXT("Z", "Z:")
+	};
+ 
+	return AddWidgetLabel(
+		Labels[InIndex],
+		CreateSpinBox(
+			TAttribute<float>::CreateSP(this, &SDMPropertyEditFloat3XYZValue::GetSpinBoxValue, InIndex),
+			SSpinBox<float>::FOnValueChanged::CreateSP(this, &SDMPropertyEditFloat3XYZValue::OnSpinBoxValueChanged, InIndex),
+			LOCTEXT("TransactionDescription", "Material Designer Value Scrubbing (Vector)"),
+			Float3XYZValue->HasValueRange() ? &ValueRange : nullptr
+		)
+	);
+}
+
+float SDMPropertyEditFloat3XYZValue::GetMaxWidthForWidget(int32 InIndex) const
+{
+	ensure(InIndex >= 0 && InIndex <= 2);
+
+	return UDynamicMaterialEditorSettings::Get()->MaxFloatSliderWidth;
+}
+
+float SDMPropertyEditFloat3XYZValue::GetSpinBoxValue(int32 InComponent) const
+{
+	ensure(InComponent >= 0 && InComponent <= 2);
+ 
+	UDMMaterialValueFloat3XYZ* Float3XYZValue = GetFloat3XYZValue();
+
+	if (!IsValid(Float3XYZValue) || !Float3XYZValue->IsComponentValid())
+	{
+		return 0;
+	}
+ 
+	return Float3XYZValue->GetValue()[InComponent];
+}
+ 
+void SDMPropertyEditFloat3XYZValue::OnSpinBoxValueChanged(float InNewValue, int32 InComponent) const
+{
+	ensure(InComponent >= 0 && InComponent <= 2);
+ 
+	UDMMaterialValueFloat3XYZ* Float3XYZValue = GetFloat3XYZValue();
+	
+	if (IsValid(Float3XYZValue) && Float3XYZValue->IsComponentValid() 
+		&& FMath::IsNearlyEqual(Float3XYZValue->GetValue()[InComponent], InNewValue) == false)
+	{
+		FVector NewTotalValue = Float3XYZValue->GetValue();
+		NewTotalValue[InComponent] = InNewValue;
+		Float3XYZValue->SetValue(NewTotalValue);
+	}
+}
+ 
+#undef LOCTEXT_NAMESPACE
diff --git a/Engine/Plugins/Experimental/DynamicMaterial/Source/DynamicMaterialEditor/Private/Slate/Properties/Editors/SDMPropertyEditFloat3XYZValue.h b/Engine/Plugins/Experimental/DynamicMaterial/Source/DynamicMaterialEditor/Private/Slate/Properties/Editors/SDMPropertyEditFloat3XYZValue.h
new file mode 100644
--- /dev/null
+++ b/Engine/Plugins/Experimental/DynamicMaterial/Source/DynamicMaterialEditor/Private/Slate/Properties/Editors/SDMPropertyEditFloat3XYZValue.h
@@ -0,0 +1,33 @@
+// Copyright Epic Games, Inc. All Rights Reserved.
+
+#pragma once
+
+#include "Slate/Properties/SDMPropertyEdit.h"
+
+class UDMMaterialValue;
+class UDMMaterialValueFloat3XYZ;
+
+class SDMPropertyEditFloat3XYZValue : public SDMPropertyEdit
+{
+	SLATE_BEGIN_ARGS(SDMPropertyEditFloat3XYZValue)
+		{}
+		SLATE_ARGUMENT(TSharedPtr<SDMComponentEdit>, ComponentEditWidget)
+	SLATE_END_ARGS()
+
+public:
+	static TSharedPtr<SWidget> CreateEditWidget(const TSharedPtr<SDMComponentEdit>& InComponentEditWidget, UDMMaterialValue* InFloat3XYZValue);
+
+	SDMPropertyEditFloat3XYZValue() = default;
+	virtual ~SDMPropertyEditFloat3XYZValue() override = default;
+
+	void Construct(const FArguments& InArgs, UDMMaterialValueFloat3XYZ* InFloat3XYZValue);
+
+	UDMMaterialValueFloat3XYZ* GetFloat3XYZValue() const;
+
+protected:
+	virtual TSharedRef<SWidget> GetComponentWidget(int32 InIndex) override;
+	virtual float GetMaxWidthForWidget(int32 InIndex) const override;
+
+	float GetSpinBoxValue(int32 InComponent) const;
+	void OnSpinBoxValueChanged(float InNewValue, int32 InComponent) const;
+};
diff --git a/Engine/Plugins/Experimental/DynamicMaterial/Source/DynamicMaterialEditor/Private/Slate/Properties/Editors/SDMPropertyEditFloat4Value.cpp b/Engine/Plugins/Experimental/DynamicMaterial/Source/DynamicMaterialEditor/Private/Slate/Properties/Editors/SDMPropertyEditFloat4Value.cpp
new file mode 100644
--- /dev/null
+++ b/Engine/Plugins/Experimental/DynamicMaterial/Source/DynamicMaterialEditor/Private/Slate/Properties/Editors/SDMPropertyEditFloat4Value.cpp
@@ -0,0 +1,87 @@
+// Copyright Epic Games, Inc. All Rights Reserved.
+ 
+#include "SDMPropertyEditFloat4Value.h"
+#include "Components/MaterialValues/DMMaterialValueFloat4.h"
+#include "DMPrivate.h"
+#include "Widgets/SNullWidget.h"
+
+#define LOCTEXT_NAMESPACE "SDMPropertyEditFloat4Value"
+
+TSharedPtr<SWidget> SDMPropertyEditFloat4Value::CreateEditWidget(const TSharedPtr<SDMComponentEdit>& InComponentEditWidget, UDMMaterialValue* InFloat4Value)
+{
+	return SNew(SDMPropertyEditFloat4Value, Cast<UDMMaterialValueFloat4>(InFloat4Value))
+		.ComponentEditWidget(InComponentEditWidget);
+}
+ 
+void SDMPropertyEditFloat4Value::Construct(const FArguments& InArgs, UDMMaterialValueFloat4* InFloat4Value)
+{
+	ensure(IsValid(InFloat4Value));
+
+	SDMPropertyEdit::Construct(
+		SDMPropertyEdit::FArguments()
+			.InputCount(1)
+			.ComponentEditWidget(InArgs._ComponentEditWidget)
+			.PropertyMaterialValue(InFloat4Value)
+	);
+}
+ 
+UDMMaterialValueFloat4* SDMPropertyEditFloat4Value::GetFloat4Value() const
+{
+	return Cast<UDMMaterialValueFloat4>(GetValue());
+}
+ 
+TSharedRef<SWidget> SDMPropertyEditFloat4Value::GetComponentWidget(int32 InIndex)
+{
+	ensure(InIndex == 0);
+ 
+	UDMMaterialValueFloat4* Float4Value = GetFloat4Value();
+
+	if (!IsValid(Float4Value) || !Float4Value->IsComponentValid())
+	{
+		return SNullWidget::NullWidget;
+	}
+ 
+	const FFloatInterval& ValueRange = Float4Value->GetValueRange();
+ 
+	return CreateColorPicker(
+		/* bUseAlpha */ true,
+		TAttribute<FLinearColor>::CreateSP(this, &SDMPropertyEditFloat4Value::GetColor),
+		FOnLinearColorValueChanged::CreateSP(this, &SDMPropertyEditFloat4Value::OnColorValueChanged),
+		LOCTEXT("TransactionDescription", "Material Designer Value Scrubbing (RGBA)"),
+		Float4Value->HasValueRange() ? &ValueRange : nullptr
+	);
+}
+ 
+FLinearColor SDMPropertyEditFloat4Value::GetColor() const
+{
+	UDMMaterialValueFloat4* Float4Value = GetFloat4Value();
+
+	if (!IsValid(Float4Value) || !Float4Value->IsComponentValid())
+	{
+		return FLinearColor::Black;
+	}
+ 
+	return Float4Value->GetValue();
+}
+ 
+void SDMPropertyEditFloat4Value::OnColorValueChanged(FLinearColor InNewColor) const
+{
+	UDMMaterialValueFloat4* Float4Value = GetFloat4Value();
+	
+	if (IsValid(Float4Value) && Float4Value->IsComponentValid())
+	{
+		const FLinearColor& CurrentColor = Float4Value->GetValue();
+
+		const bool bSameValue = FMath::IsNearlyEqual(InNewColor.R, CurrentColor.R)
+			&& FMath::IsNearlyEqual(InNewColor.G, CurrentColor.G)
+			&& FMath::IsNearlyEqual(InNewColor.B, CurrentColor.B)
+			&& FMath::IsNearlyEqual(InNewColor.A, CurrentColor.A);
+
+		if (!bSameValue)
+		{
+			Float4Value->SetValue(InNewColor);
+		}
+	}
+}
+ 
+#undef LOCTEXT_NAMESPACE
diff --git a/Engine/Plugins/Experimental/DynamicMaterial/Source/DynamicMaterialEditor/Private/Slate/Properties/Editors/SDMPropertyEditFloat4Value.h b/Engine/Plugins/Experimental/DynamicMaterial/Source/DynamicMaterialEditor/Private/Slate/Properties/Editors/SDMPropertyEditFloat4Value.h
new file mode 100644
--- /dev/null
+++ b/Engine/Plugins/Experimental/DynamicMaterial/Source/DynamicMaterialEditor/Private/Slate/Properties/Editors/SDMPropertyEditFloat4Value.h
@@ -0,0 +1,32 @@
+// Copyright Epic Games, Inc. All Rights Reserved.
+
+#pragma once
+
+#include "Slate/Properties/SDMPropertyEdit.h"
+
+class UDMMaterialValue;
+class UDMMaterialValueFloat4;
+
+class SDMPropertyEditFloat4Value : public SDMPropertyEdit
+{
+	SLATE_BEGIN_ARGS(SDMPropertyEditFloat4Value)
+		{}
+		SLATE_ARGUMENT(TSharedPtr<SDMComponentEdit>, ComponentEditWidget)
+	SLATE_END_ARGS()
+
+public:
+	static TSharedPtr<SWidget> CreateEditWidget(const TSharedPtr<SDMComponentEdit>& InComponentEditWidget, UDMMaterialValue* InFloat4Value);
+
+	SDMPropertyEditFloat4Value() = default;
+	virtual ~SDMPropertyEditFloat4Value() override = default;
+
+	void Construct(const FArguments& InArgs, UDMMaterialValueFloat4* InFloat4Value);
+
+	UDMMaterialValueFloat4* GetFloat4Value() const;
+
+protected:
+	virtual TSharedRef<SWidget> GetComponentWidget(int32 InIndex) override;
+
+	FLinearColor GetColor() const;
+	void OnColorValueChanged(FLinearColor InNewColor) const;
+};
diff --git a/Engine/Plugins/Experimental/DynamicMaterial/Source/DynamicMaterialEditor/Private/Slate/Properties/Editors/SDMPropertyEditObject.cpp b/Engine/Plugins/Experimental/DynamicMaterial/Source/DynamicMaterialEditor/Private/Slate/Properties/Editors/SDMPropertyEditObject.cpp
new file mode 100644
--- /dev/null
+++ b/Engine/Plugins/Experimental/DynamicMaterial/Source/DynamicMaterialEditor/Private/Slate/Properties/Editors/SDMPropertyEditObject.cpp
@@ -0,0 +1,87 @@
+// Copyright Epic Games, Inc. All Rights Reserved.
+
+#include "Slate/Properties/Editors/SDMPropertyEditObject.h"
+#include "PropertyHandle.h"
+
+#define LOCTEXT_NAMESPACE "SDMPropertyEditObject"
+
+void SDMPropertyEditObject::Construct(const FArguments& InArgs, const TSharedPtr<IPropertyHandle>& InPropertyHandle, UClass* InAllowedClass)
+{
+	ensure(InPropertyHandle.IsValid());
+
+	SDMPropertyEdit::Construct(
+		SDMPropertyEdit::FArguments()
+			.InputCount(1)
+			.ComponentEditWidget(InArgs._ComponentEditWidget)
+			.PropertyHandle(InPropertyHandle)
+	);
+
+	AllowedClass = InAllowedClass;
+}
+
+TSharedRef<SWidget> SDMPropertyEditObject::GetComponentWidget(int32 InIndex)
+{
+	ensure(InIndex == 0);
+
+	return CreateAssetPicker(
+		AllowedClass.Get(),
+		TAttribute<FString>::CreateSP(this, &SDMPropertyEditObject::GetObjectPath),
+		FOnSetObject::CreateSP(this, &SDMPropertyEditObject::OnValueChanged)
+	);
+}
+
+FString SDMPropertyEditObject::GetObjectPath() const
+{
+	if (PropertyHandle.IsValid() && PropertyHandle->IsValidHandle())
+	{
+		UObject* Object = nullptr;
+		PropertyHandle->GetValue(Object);
+
+		if (!IsValid(Object))
+		{
+			return "";
+		}
+
+		return Object->GetPathName();
+	}
+
+	return "";
+}
+
+void SDMPropertyEditObject::OnValueChanged(const FAssetData& InAssetData)
+{
+	if (PropertyHandle.IsValid() && PropertyHandle->IsValidHandle())
+	{
+		UObject* CurrentValue = nullptr;
+
+		if (PropertyHandle->GetValue(CurrentValue) == FPropertyAccess::Fail)
+		{
+			return;
+		}
+
+		UObject* NewValue = InAssetData.GetAsset();
+
+		if (CurrentValue == NewValue)
+		{
+			return;
+		}
+
+		TArray<UObject*> Outers;
+		PropertyHandle->GetOuterObjects(Outers);
+		const bool bHasValidOuter = Outers.IsEmpty() == false && IsValid(Outers[0]);
+
+		if (bHasValidOuter)
+		{
+			StartTransaction(LOCTEXT("TransactionDescription", "Material Designer Value Set (Asset)"));
+		}
+
+		PropertyHandle->SetValue(NewValue, EPropertyValueSetFlags::NotTransactable | EPropertyValueSetFlags::InteractiveChange);
+
+		if (bHasValidOuter)
+		{
+			EndTransaction();
+		}
+	}
+}
+
+#undef LOCTEXT_NAMESPACE
diff --git a/Engine/Plugins/Experimental/DynamicMaterial/Source/DynamicMaterialEditor/Private/Slate/Properties/Editors/SDMPropertyEditObject.h b/Engine/Plugins/Experimental/DynamicMaterial/Source/DynamicMaterialEditor/Private/Slate/Properties/Editors/SDMPropertyEditObject.h
new file mode 100644
--- /dev/null
+++ b/Engine/Plugins/Experimental/DynamicMaterial/Source/DynamicMaterialEditor/Private/Slate/Properties/Editors/SDMPropertyEditObject.h
@@ -0,0 +1,36 @@
+// Copyright Epic Games, Inc. All Rights Reserved.
+
+#pragma once
+
+#include "Slate/Properties/SDMPropertyEdit.h"
+#include "Templates/SharedPointer.h"
+#include "UObject/WeakObjectPtr.h"
+#include "UObject/WeakObjectPtrTemplates.h"
+
+class FString;
+class IPropertyHandle;
+class SWidget;
+class UClass;
+struct FAssetData;
+
+class SDMPropertyEditObject : public SDMPropertyEdit
+{
+	SLATE_BEGIN_ARGS(SDMPropertyEditObject)
+		{}
+		SLATE_ARGUMENT(TSharedPtr<SDMComponentEdit>, ComponentEditWidget)
+	SLATE_END_ARGS()
+
+public:
+	SDMPropertyEditObject() = default;
+	virtual ~SDMPropertyEditObject() override = default;
+
+	void Construct(const FArguments& InArgs, const TSharedPtr<IPropertyHandle>& InPropertyHandle, UClass* InAllowedClass);
+
+protected:
+	TWeakObjectPtr<UClass> AllowedClass;
+
+	virtual TSharedRef<SWidget> GetComponentWidget(int32 InIndex) override;
+
+	FString GetObjectPath() const;
+	void OnValueChanged(const FAssetData& InAssetData);
+};
diff --git a/Engine/Plugins/Experimental/DynamicMaterial/Source/DynamicMaterialEditor/Private/Slate/Properties/Editors/SDMPropertyEditOpacity.cpp b/Engine/Plugins/Experimental/DynamicMaterial/Source/DynamicMaterialEditor/Private/Slate/Properties/Editors/SDMPropertyEditOpacity.cpp
new file mode 100644
--- /dev/null
+++ b/Engine/Plugins/Experimental/DynamicMaterial/Source/DynamicMaterialEditor/Private/Slate/Properties/Editors/SDMPropertyEditOpacity.cpp
@@ -0,0 +1,44 @@
+// Copyright Epic Games, Inc.All Rights Reserved.
+
+#include "SDMPropertyEditOpacity.h"
+#include "Components/MaterialValues/DMMaterialValueFloat1.h"
+#include "DynamicMaterialEditorSettings.h"
+#include "Slate/SDMEditor.h"
+#include "Templates/SharedPointer.h"
+#include "Widgets/Input/SSpinBox.h"
+#include "Widgets/SWidget.h"
+
+#define LOCTEXT_NAMESPACE "SDMPropertyEditOpacity"
+
+void SDMPropertyEditOpacity::Construct(const FArguments& InArgs, const TSharedPtr<SWidget>& InWidget, UDMMaterialValueFloat1* InOpacityValue)
+{
+	if (ensure(IsValid(InOpacityValue)))
+	{
+		Super::Construct(
+			Super::FArguments(),
+			SDMEditor::GetPropertyHandle(InWidget.Get(), InOpacityValue, UDMMaterialValue::ValueName).PropertyHandle
+		);
+	}
+}
+
+TSharedRef<SWidget> SDMPropertyEditOpacity::GetComponentWidget(int32 InIndex)
+{
+	ensure(InIndex == 0);
+
+	static const FFloatInterval OpacityInterval(0.f, 1.f);
+
+	return CreateSpinBox(
+		TAttribute<float>::CreateSP(this, &SDMPropertyEditOpacity::GetFloatValue),
+		FOnFloatValueChanged::CreateSP(this, &SDMPropertyEditOpacity::OnValueChanged),
+		LOCTEXT("TransactionDescription", "Material Designer Value Scrubbing (Opacity)"),
+		&OpacityInterval
+	);
+}
+
+float SDMPropertyEditOpacity::GetMaxWidthForWidget(int32 InIndex) const
+{
+	// Disable
+	return TNumericLimits<float>::Max();
+}
+
+#undef LOCTEXT_NAMESPACE
diff --git a/Engine/Plugins/Experimental/DynamicMaterial/Source/DynamicMaterialEditor/Private/Slate/Properties/Editors/SDMPropertyEditOpacity.h b/Engine/Plugins/Experimental/DynamicMaterial/Source/DynamicMaterialEditor/Private/Slate/Properties/Editors/SDMPropertyEditOpacity.h
new file mode 100644
--- /dev/null
+++ b/Engine/Plugins/Experimental/DynamicMaterial/Source/DynamicMaterialEditor/Private/Slate/Properties/Editors/SDMPropertyEditOpacity.h
@@ -0,0 +1,24 @@
+// Copyright Epic Games, Inc.All Rights Reserved.
+
+#pragma once
+
+#include "SDMPropertyEditFloat.h"
+
+class UDMMaterialValueFloat1;
+class SDMSlot;
+
+class SDMPropertyEditOpacity : public SDMPropertyEditFloat
+{
+	SLATE_BEGIN_ARGS(SDMPropertyEditFloat)
+		{}
+	SLATE_END_ARGS()
+
+public:
+	using Super = SDMPropertyEditFloat;
+
+	void Construct(const FArguments& InArgs, const TSharedPtr<SWidget>& InWidget, UDMMaterialValueFloat1* InOpacityValue);
+
+protected:
+	virtual TSharedRef<SWidget> GetComponentWidget(int32 InIndex) override;
+	virtual float GetMaxWidthForWidget(int32 InIndex) const override;
+};
diff --git a/Engine/Plugins/Experimental/DynamicMaterial/Source/DynamicMaterialEditor/Private/Slate/Properties/Editors/SDMPropertyEditTextureValue.cpp b/Engine/Plugins/Experimental/DynamicMaterial/Source/DynamicMaterialEditor/Private/Slate/Properties/Editors/SDMPropertyEditTextureValue.cpp
new file mode 100644
--- /dev/null
+++ b/Engine/Plugins/Experimental/DynamicMaterial/Source/DynamicMaterialEditor/Private/Slate/Properties/Editors/SDMPropertyEditTextureValue.cpp
@@ -0,0 +1,95 @@
+// Copyright Epic Games, Inc. All Rights Reserved.
+ 
+#include "SDMPropertyEditTextureValue.h"
+#include "Components/MaterialValues/DMMaterialValueTexture.h"
+#include "DMPrivate.h"
+#include "Engine/Texture.h"
+#include "Widgets/SNullWidget.h"
+
+#define LOCTEXT_NAMESPACE "SDMPropertyEditTextureValue"
+
+TSharedPtr<SWidget> SDMPropertyEditTextureValue::CreateEditWidget(const TSharedPtr<SDMComponentEdit>& InComponentEditWidget, UDMMaterialValue* InTextureValue)
+{
+	return SNew(SDMPropertyEditTextureValue, Cast<UDMMaterialValueTexture>(InTextureValue))
+		.ComponentEditWidget(InComponentEditWidget);
+}
+ 
+void SDMPropertyEditTextureValue::Construct(const FArguments& InArgs, UDMMaterialValueTexture* InTextureValue)
+{
+	ensure(IsValid(InTextureValue));
+
+	SDMPropertyEdit::Construct(
+		SDMPropertyEdit::FArguments()
+			.InputCount(1)
+			.ComponentEditWidget(InArgs._ComponentEditWidget)
+			.PropertyMaterialValue(InTextureValue)
+	);
+}
+ 
+UDMMaterialValueTexture* SDMPropertyEditTextureValue::GetTextureValue() const
+{
+	return Cast<UDMMaterialValueTexture>(GetValue());
+}
+ 
+FString SDMPropertyEditTextureValue::GetTexturePath() const
+{
+	UDMMaterialValueTexture* TextureValue = GetTextureValue();
+
+	if (!IsValid(TextureValue) || !TextureValue->IsComponentValid())
+	{
+		return "";
+	}
+ 
+	const UTexture* Texture = TextureValue->GetValue();
+ 
+	if (Texture)
+	{
+		return Texture->GetPathName();
+	}
+ 
+	return "";
+}
+ 
+TSharedRef<SWidget> SDMPropertyEditTextureValue::GetComponentWidget(int32 InIndex)
+{
+	ensure(InIndex == 0);
+ 
+	UDMMaterialValueTexture* TextureValue = GetTextureValue();
+
+	if (!IsValid(TextureValue) || !TextureValue->IsComponentValid())
+	{
+		return SNullWidget::NullWidget;
+	}
+ 
+	return CreateAssetPicker(
+		UTexture::StaticClass(),
+		TAttribute<FString>::CreateSP(this, &SDMPropertyEditTextureValue::GetTexturePath),
+		FOnSetObject::CreateSP(this, &SDMPropertyEditTextureValue::OnTextureSelected)
+	);
+}
+ 
+void SDMPropertyEditTextureValue::OnTextureSelected(const FAssetData& InAssetData)
+{
+	UDMMaterialValueTexture* TextureValue = GetTextureValue();
+
+	if (IsValid(TextureValue) && TextureValue->IsComponentValid())
+	{
+		UObject* NewValue = InAssetData.GetAsset();
+		UTexture* NewTexture = Cast<UTexture>(NewValue);
+
+		// Not a UTexture
+		if (NewValue && !NewTexture)
+		{
+			return;
+		}
+
+		if (TextureValue->GetValue() != NewTexture)
+		{
+			StartTransaction(LOCTEXT("SetValue", "Material Designer Value Set (Texture)"));
+			TextureValue->SetValue(NewTexture);
+			EndTransaction();
+		}
+	}
+}
+
+#undef LOCTEXT_NAMESPACE
diff --git a/Engine/Plugins/Experimental/DynamicMaterial/Source/DynamicMaterialEditor/Private/Slate/Properties/Editors/SDMPropertyEditTextureValue.h b/Engine/Plugins/Experimental/DynamicMaterial/Source/DynamicMaterialEditor/Private/Slate/Properties/Editors/SDMPropertyEditTextureValue.h
new file mode 100644
--- /dev/null
+++ b/Engine/Plugins/Experimental/DynamicMaterial/Source/DynamicMaterialEditor/Private/Slate/Properties/Editors/SDMPropertyEditTextureValue.h
@@ -0,0 +1,34 @@
+// Copyright Epic Games, Inc. All Rights Reserved.
+
+#pragma once
+
+#include "Slate/Properties/SDMPropertyEdit.h"
+
+class FAssetThumbnailPool;
+class UDMMaterialValue;
+class UDMMaterialValueTexture;
+struct FAssetData;
+
+class SDMPropertyEditTextureValue : public SDMPropertyEdit
+{
+	SLATE_BEGIN_ARGS(SDMPropertyEditTextureValue)
+		{}
+		SLATE_ARGUMENT(TSharedPtr<SDMComponentEdit>, ComponentEditWidget)
+	SLATE_END_ARGS()
+
+public:
+	static TSharedPtr<SWidget> CreateEditWidget(const TSharedPtr<SDMComponentEdit>& InComponentEditWidget, UDMMaterialValue* InTextureValue);
+
+	SDMPropertyEditTextureValue() = default;
+	virtual ~SDMPropertyEditTextureValue() override = default;
+
+	void Construct(const FArguments& InArgs, UDMMaterialValueTexture* InTextureValue);
+
+	UDMMaterialValueTexture* GetTextureValue() const;
+
+protected:
+	virtual TSharedRef<SWidget> GetComponentWidget(int32 InIndex) override;
+
+	FString GetTexturePath() const;
+	void OnTextureSelected(const FAssetData& InAssetData);
+};
diff --git a/Engine/Plugins/Experimental/DynamicMaterial/Source/DynamicMaterialEditor/Private/Slate/Properties/Editors/SDMPropertyEditVector.cpp b/Engine/Plugins/Experimental/DynamicMaterial/Source/DynamicMaterialEditor/Private/Slate/Properties/Editors/SDMPropertyEditVector.cpp
new file mode 100644
--- /dev/null
+++ b/Engine/Plugins/Experimental/DynamicMaterial/Source/DynamicMaterialEditor/Private/Slate/Properties/Editors/SDMPropertyEditVector.cpp
@@ -0,0 +1,192 @@
+// Copyright Epic Games, Inc. All Rights Reserved.
+
+#include "Slate/Properties/Editors/SDMPropertyEditVector.h"
+#include "DynamicMaterialEditorSettings.h"
+#include "PropertyHandle.h"
+
+#define LOCTEXT_NAMESPACE "SDMPropertyEditVector"
+
+void SDMPropertyEditVector::Construct(const FArguments& InArgs, const TSharedPtr<IPropertyHandle>& InPropertyHandle, int InComponentCount)
+{
+	ensure(InPropertyHandle.IsValid());
+	ComponentCount = InComponentCount;
+
+	SDMPropertyEdit::Construct(
+		SDMPropertyEdit::FArguments()
+			.InputCount(InComponentCount)
+			.ComponentEditWidget(InArgs._ComponentEditWidget)
+			.PropertyHandle(InPropertyHandle)
+	);
+}
+
+TSharedRef<SWidget> SDMPropertyEditVector::GetComponentWidget(int32 InIndex)
+{
+	ensure(InIndex >= 0 && InIndex < ComponentCount);
+	static const TArray<FText> Labels = {
+		LOCTEXT("X", "X:"),
+		LOCTEXT("Y", "Y:"),
+		LOCTEXT("Z", "Z:"),
+		LOCTEXT("W", "W:"),
+	};
+
+	if (InIndex < 4)
+	{
+		return AddWidgetLabel(
+			Labels[InIndex],
+			CreateSpinBox(
+				TAttribute<float>::CreateSP(this, &SDMPropertyEditVector::GetVectorValue, InIndex),
+				FOnFloatValueChanged::CreateSP(this, &SDMPropertyEditVector::OnValueChanged, InIndex),
+				LOCTEXT("TransactionDescription", "Material Designer Value Scrubbing (Vector)")
+			)
+		);
+
+	}
+
+	return CreateSpinBox(
+		TAttribute<float>::CreateSP(this, &SDMPropertyEditVector::GetVectorValue, InIndex),
+		FOnFloatValueChanged::CreateSP(this, &SDMPropertyEditVector::OnValueChanged, InIndex),
+		LOCTEXT("TransactionDescription", "Material Designer Value Scrubbing (Vector)")
+	);
+}
+
+float SDMPropertyEditVector::GetMaxWidthForWidget(int32 InIndex) const
+{
+	ensure(InIndex >= 0 && InIndex < ComponentCount);
+
+	return UDynamicMaterialEditorSettings::Get()->MaxFloatSliderWidth;
+}
+
+float SDMPropertyEditVector::GetVectorValue(int32 InComponent) const
+{
+	ensure(InComponent >= 0 && InComponent < ComponentCount);
+
+	if (PropertyHandle.IsValid() && PropertyHandle->IsValidHandle())
+	{
+		if (FStructProperty* StructProperty = CastField<FStructProperty>(Property))
+		{
+			if (StructProperty->Struct == TBaseStructure<FVector2D>::Get())
+			{
+				FVector2D Value = FVector2D::ZeroVector;
+				PropertyHandle->GetValue(Value);
+				return Value[InComponent];
+			}
+
+			if (StructProperty->Struct == TBaseStructure<FVector>::Get())
+			{
+				FVector Value = FVector::ZeroVector;
+				PropertyHandle->GetValue(Value);
+				return Value[InComponent];
+			}
+
+			if (StructProperty->Struct == TBaseStructure<FRotator>::Get())
+			{
+				FRotator Value = FRotator::ZeroRotator;
+				PropertyHandle->GetValue(Value);
+				EAxis::Type Axis;
+
+				switch (InComponent)
+				{
+					case 0:
+						Axis = EAxis::X;
+						break;
+
+					case 1:
+						Axis = EAxis::Y;
+						break;
+
+					case 2:
+						Axis = EAxis::Z;
+						break;
+
+					default:
+						return 0.f;
+				}
+
+				return Value.GetComponentForAxis(Axis);
+			}
+		}
+	}
+
+	return 0.f;	
+}
+
+void SDMPropertyEditVector::OnValueChanged(float InNewValue, int32 InComponent) const
+{
+	ensure(InComponent >= 0 && InComponent < ComponentCount);
+
+	if (PropertyHandle.IsValid() && PropertyHandle->IsValidHandle())
+	{
+		if (FStructProperty* StructProperty = CastField<FStructProperty>(Property))
+		{
+			if (StructProperty->Struct == TBaseStructure<FVector2D>::Get())
+			{
+				FVector2D CurrentValue;
+
+				if (PropertyHandle->GetValue(CurrentValue) == FPropertyAccess::Fail)
+				{
+					return;
+				}
+
+				if (FMath::IsNearlyEqual(CurrentValue[InComponent], InNewValue) == false)
+				{
+					CurrentValue[InComponent] = InNewValue;
+					PropertyHandle->SetValue(CurrentValue, EPropertyValueSetFlags::NotTransactable | EPropertyValueSetFlags::InteractiveChange);
+				}
+			}
+
+			if (StructProperty->Struct == TBaseStructure<FVector>::Get())
+			{
+				FVector CurrentValue;
+
+				if (PropertyHandle->GetValue(CurrentValue) == FPropertyAccess::Fail)
+				{
+					return;
+				}
+
+				if (FMath::IsNearlyEqual(CurrentValue[InComponent], InNewValue) == false)
+				{
+					CurrentValue[InComponent] = InNewValue;
+					PropertyHandle->SetValue(CurrentValue, EPropertyValueSetFlags::NotTransactable | EPropertyValueSetFlags::InteractiveChange);
+				}
+			}
+
+			if (StructProperty->Struct == TBaseStructure<FRotator>::Get())
+			{
+				FRotator CurrentValue;
+
+				if (PropertyHandle->GetValue(CurrentValue) == FPropertyAccess::Fail)
+				{
+					return;
+				}
+
+				EAxis::Type Axis;
+
+				switch (InComponent)
+				{
+					case 0:
+						Axis = EAxis::X;
+						break;
+
+					case 1:
+						Axis = EAxis::Y;
+						break;
+
+					case 2:
+						Axis = EAxis::Z;
+						break;
+
+					default:
+						return;
+				}
+
+				if (FMath::IsNearlyEqual(CurrentValue.GetComponentForAxis(Axis), InNewValue) == false)
+				{
+					CurrentValue.SetComponentForAxis(Axis, InNewValue);
+					PropertyHandle->SetValue(CurrentValue, EPropertyValueSetFlags::NotTransactable | EPropertyValueSetFlags::InteractiveChange);
+				}
+			}
+		}
+	}
+}
+
+#undef LOCTEXT_NAMESPACE
diff --git a/Engine/Plugins/Experimental/DynamicMaterial/Source/DynamicMaterialEditor/Private/Slate/Properties/Editors/SDMPropertyEditVector.h b/Engine/Plugins/Experimental/DynamicMaterial/Source/DynamicMaterialEditor/Private/Slate/Properties/Editors/SDMPropertyEditVector.h
new file mode 100644
--- /dev/null
+++ b/Engine/Plugins/Experimental/DynamicMaterial/Source/DynamicMaterialEditor/Private/Slate/Properties/Editors/SDMPropertyEditVector.h
@@ -0,0 +1,32 @@
+// Copyright Epic Games, Inc. All Rights Reserved.
+
+#pragma once
+
+#include "Slate/Properties/SDMPropertyEdit.h"
+#include "Templates/SharedPointer.h"
+
+class IPropertyHandle;
+class SWidget;
+
+class SDMPropertyEditVector : public SDMPropertyEdit
+{
+	SLATE_BEGIN_ARGS(SDMPropertyEditVector)
+		{}
+		SLATE_ARGUMENT(TSharedPtr<SDMComponentEdit>, ComponentEditWidget)
+	SLATE_END_ARGS()
+
+public:
+	SDMPropertyEditVector() = default;
+	virtual ~SDMPropertyEditVector() override = default;
+
+	void Construct(const FArguments& InArgs, const TSharedPtr<IPropertyHandle>& InPropertyHandle, int32 InComponentCount);
+
+protected:
+	int32 ComponentCount;
+
+	virtual TSharedRef<SWidget> GetComponentWidget(int32 InIndex) override;
+	virtual float GetMaxWidthForWidget(int32 InIndex) const override;
+
+	float GetVectorValue(int32 InComponent) const;
+	virtual void OnValueChanged(float InNewValue, int32 InComponent) const;
+};
diff --git a/Engine/Plugins/Experimental/DynamicMaterial/Source/DynamicMaterialEditor/Private/Slate/Properties/Generators/DMComponentPropertyRowGenerator.cpp b/Engine/Plugins/Experimental/DynamicMaterial/Source/DynamicMaterialEditor/Private/Slate/Properties/Generators/DMComponentPropertyRowGenerator.cpp
new file mode 100644
--- /dev/null
+++ b/Engine/Plugins/Experimental/DynamicMaterial/Source/DynamicMaterialEditor/Private/Slate/Properties/Generators/DMComponentPropertyRowGenerator.cpp
@@ -0,0 +1,105 @@
+// Copyright Epic Games, Inc. All Rights Reserved.
+
+#include "Slate/Properties/Generators/DMComponentPropertyRowGenerator.h"
+#include "DynamicMaterialEditorModule.h"
+#include "Components/DMMaterialComponent.h"
+#include "Runtime/Launch/Resources/Version.h"
+#include "Slate/SDMEditor.h"
+#include "Slate/SDMComponentEdit.h"
+
+const TSharedRef<FDMComponentPropertyRowGenerator>& FDMComponentPropertyRowGenerator::Get()
+{
+	static TSharedRef<FDMComponentPropertyRowGenerator> Generator = MakeShared<FDMComponentPropertyRowGenerator>();
+	return Generator;
+}
+
+void FDMComponentPropertyRowGenerator::AddComponentProperties(const TSharedRef<SDMComponentEdit>& InComponentEditWidget, UDMMaterialComponent* InComponent,
+	TArray<FDMPropertyHandle>& InOutPropertyRows, TSet<UDMMaterialComponent*>& InOutProcessedObjects)
+{
+	if (!IsValid(InComponent))
+	{
+		return;
+	}
+
+	if (InOutProcessedObjects.Contains(InComponent))
+	{
+		return;
+	}
+
+	InOutProcessedObjects.Add(InComponent);
+
+	const TArray<FName>& Properties = InComponent->GetEditableProperties();
+
+	for (const FName& Property : Properties)
+	{
+		if (InComponent->IsPropertyVisible(Property))
+		{
+			AddPropertyEditRows(InComponentEditWidget, InComponent, Property, InOutPropertyRows, InOutProcessedObjects);
+		}
+	}
+}
+
+void FDMComponentPropertyRowGenerator::AddPropertyEditRows(const TSharedRef<SDMComponentEdit>& InComponentEditWidget, UDMMaterialComponent* InComponent, 
+	const FName& InProperty, TArray<FDMPropertyHandle>& InOutPropertyRows, TSet<UDMMaterialComponent*>& InOutProcessedObjects)
+{
+	if (!IsValid(InComponent))
+	{
+		return;
+	}
+
+	FProperty* Property = InComponent->GetClass()->FindPropertyByName(InProperty);
+
+	if (!Property)
+	{
+		return;
+	}
+
+	void* MemoryPtr = Property->ContainerPtrToValuePtr<void>(InComponent);
+
+	if (FArrayProperty* ArrayProperty = CastField<FArrayProperty>(Property))
+	{
+		FScriptArrayHelper ArrayHelper(ArrayProperty, MemoryPtr);
+
+		for (int32 Idx = 0; Idx < ArrayHelper.Num(); ++Idx)
+		{
+#if ENGINE_MAJOR_VERSION < 5 || (ENGINE_MAJOR_VERSION == 5 && ENGINE_MINOR_VERSION <= 2)
+			void* ElemPtr = ArrayHelper.GetRawPtr(Idx);
+#else
+			void* ElemPtr = ArrayHelper.GetElementPtr(Idx);
+#endif
+			AddPropertyEditRows(InComponentEditWidget, InComponent, ArrayProperty->Inner, ElemPtr, InOutPropertyRows, InOutProcessedObjects);
+		}
+	}
+	else
+	{
+		AddPropertyEditRows(InComponentEditWidget, InComponent, Property, MemoryPtr, InOutPropertyRows, InOutProcessedObjects);
+	}
+}
+
+void FDMComponentPropertyRowGenerator::AddPropertyEditRows(const TSharedRef<SDMComponentEdit>& InComponentEditWidget, UDMMaterialComponent* InComponent,
+	FProperty* InProperty, void* MemoryPtr, TArray<FDMPropertyHandle>& InOutPropertyRows, TSet<UDMMaterialComponent*>& InOutProcessedObjects)
+{
+	if (InProperty->IsA<FArrayProperty>())
+	{
+		return;
+	}
+
+	if (FObjectPropertyBase* ObjectProperty = CastField<FObjectPropertyBase>(InProperty))
+	{
+		if (ObjectProperty->PropertyClass->IsChildOf(UDMMaterialComponent::StaticClass()))
+		{
+			UObject** ValuePtr = static_cast<UObject**>(MemoryPtr);
+			UObject* Value = *ValuePtr;
+			UDMMaterialComponent* ComponentValue = Cast<UDMMaterialComponent>(Value);
+			FDynamicMaterialEditorModule::GeneratorComponentPropertyRows(InComponentEditWidget, ComponentValue, InOutPropertyRows, InOutProcessedObjects);
+			return;
+		}
+	}
+
+	InOutPropertyRows.Add(SDMEditor::GetPropertyHandle(&*InComponentEditWidget, InComponent, InProperty->GetFName()));
+}
+
+bool FDMComponentPropertyRowGenerator::AllowKeyframeButton(UDMMaterialComponent* InComponent, FProperty* InProperty)
+{
+	return false;
+}
diff --git a/Engine/Plugins/Experimental/DynamicMaterial/Source/DynamicMaterialEditor/Private/Slate/Properties/Generators/DMInputThroughputPropertyRowGenerator.cpp b/Engine/Plugins/Experimental/DynamicMaterial/Source/DynamicMaterialEditor/Private/Slate/Properties/Generators/DMInputThroughputPropertyRowGenerator.cpp
new file mode 100644
--- /dev/null
+++ b/Engine/Plugins/Experimental/DynamicMaterial/Source/DynamicMaterialEditor/Private/Slate/Properties/Generators/DMInputThroughputPropertyRowGenerator.cpp
@@ -0,0 +1,50 @@
+// Copyright Epic Games, Inc. All Rights Reserved.
+
+#include "Slate/Properties/Generators/DMInputThroughputPropertyRowGenerator.h"
+#include "Components/DMMaterialStageThroughput.h"
+#include "Components/DMMaterialSubStage.h"
+#include "Components/MaterialStageInputs/DMMSIThroughput.h"
+
+const TSharedRef<FDMInputThroughputPropertyRowGenerator>& FDMInputThroughputPropertyRowGenerator::Get()
+{
+	static TSharedRef<FDMInputThroughputPropertyRowGenerator> Generator = MakeShared<FDMInputThroughputPropertyRowGenerator>();
+	return Generator;
+}
+
+void FDMInputThroughputPropertyRowGenerator::AddComponentProperties(const TSharedRef<SDMComponentEdit>& InComponentEditWidget,
+	UDMMaterialComponent* InComponent, TArray<FDMPropertyHandle>& InOutPropertyRows,
+	TSet<UDMMaterialComponent*>& InOutProcessedObjects)
+{
+	if (!IsValid(InComponent))
+	{
+		return;
+	}
+
+	if (InOutProcessedObjects.Contains(InComponent))
+	{
+		return;
+	}
+
+	UDMMaterialStageInputThroughput* InputThroughput = Cast<UDMMaterialStageInputThroughput>(InComponent);
+
+	if (!InputThroughput)
+	{
+		return;
+	}
+
+	InOutProcessedObjects.Add(InputThroughput);
+
+	if (UDMMaterialSubStage* SubStage = InputThroughput->GetSubStage())
+	{
+		InOutProcessedObjects.Add(SubStage);
+	}
+
+	UDMMaterialStageThroughput* Throughput = InputThroughput->GetMaterialStageThroughput();
+
+	if (!Throughput)
+	{
+		return;
+	}
+
+	FDMThroughputPropertyRowGenerator::AddComponentProperties(InComponentEditWidget, Throughput, InOutPropertyRows, InOutProcessedObjects);
+}
diff --git a/Engine/Plugins/Experimental/DynamicMaterial/Source/DynamicMaterialEditor/Private/Slate/Properties/Generators/DMInputThroughputPropertyRowGenerator.h b/Engine/Plugins/Experimental/DynamicMaterial/Source/DynamicMaterialEditor/Private/Slate/Properties/Generators/DMInputThroughputPropertyRowGenerator.h
new file mode 100644
--- /dev/null
+++ b/Engine/Plugins/Experimental/DynamicMaterial/Source/DynamicMaterialEditor/Private/Slate/Properties/Generators/DMInputThroughputPropertyRowGenerator.h
@@ -0,0 +1,17 @@
+// Copyright Epic Games, Inc. All Rights Reserved.
+
+#pragma once
+
+#include "Slate/Properties/Generators/DMThroughputPropertyRowGenerator.h"
+
+class FDMInputThroughputPropertyRowGenerator : public FDMThroughputPropertyRowGenerator
+{
+public:
+	static const TSharedRef<FDMInputThroughputPropertyRowGenerator>& Get();
+
+	FDMInputThroughputPropertyRowGenerator() = default;
+	virtual ~FDMInputThroughputPropertyRowGenerator() override = default;
+
+	virtual void AddComponentProperties(const TSharedRef<SDMComponentEdit>& InComponentEditWidget, UDMMaterialComponent* InComponent,
+		TArray<FDMPropertyHandle>& InOutPropertyRows, TSet<UDMMaterialComponent*>& InOutProcessedObjects) override;
+};
diff --git a/Engine/Plugins/Experimental/DynamicMaterial/Source/DynamicMaterialEditor/Private/Slate/Properties/Generators/DMMaterialEffectFunctionPropertyRowGenerator.h b/Engine/Plugins/Experimental/DynamicMaterial/Source/DynamicMaterialEditor/Private/Slate/Properties/Generators/DMMaterialEffectFunctionPropertyRowGenerator.h
new file mode 100644
--- /dev/null
+++ b/Engine/Plugins/Experimental/DynamicMaterial/Source/DynamicMaterialEditor/Private/Slate/Properties/Generators/DMMaterialEffectFunctionPropertyRowGenerator.h
@@ -0,0 +1,17 @@
+// Copyright Epic Games, Inc. All Rights Reserved.
+
+#pragma once
+
+#include "Slate/Properties/Generators/DMComponentPropertyRowGenerator.h"
+
+class FDMMaterialEffectFunctionPropertyRowGenerator : public FDMComponentPropertyRowGenerator
+{
+public:
+	static const TSharedRef<FDMMaterialEffectFunctionPropertyRowGenerator>& Get();
+
+	FDMMaterialEffectFunctionPropertyRowGenerator() = default;
+	virtual ~FDMMaterialEffectFunctionPropertyRowGenerator() override = default;
+
+	virtual void AddComponentProperties(const TSharedRef<SDMComponentEdit>& InComponentEditWidget, UDMMaterialComponent* InComponent,
+		TArray<FDMPropertyHandle>& InOutPropertyRows, TSet<UDMMaterialComponent*>& InOutProcessedObjects) override;
+};
diff --git a/Engine/Plugins/Experimental/DynamicMaterial/Source/DynamicMaterialEditor/Private/Slate/Properties/Generators/DMMaterialValuePropertyRowGenerator.cpp b/Engine/Plugins/Experimental/DynamicMaterial/Source/DynamicMaterialEditor/Private/Slate/Properties/Generators/DMMaterialValuePropertyRowGenerator.cpp
new file mode 100644
--- /dev/null
+++ b/Engine/Plugins/Experimental/DynamicMaterial/Source/DynamicMaterialEditor/Private/Slate/Properties/Generators/DMMaterialValuePropertyRowGenerator.cpp
@@ -0,0 +1,68 @@
+// Copyright Epic Games, Inc. All Rights Reserved.
+
+#include "Slate/Properties/Generators/DMMaterialValuePropertyRowGenerator.h"
+#include "DynamicMaterialEditorModule.h"
+#include "IDetailPropertyRow.h"
+#include "Components/DMMaterialComponent.h"
+#include "Components/DMMaterialValue.h"
+#include "Slate/SDMEditor.h"
+#include "Slate/SDMComponentEdit.h"
+
+const TSharedRef<FDMMaterialValuePropertyRowGenerator>& FDMMaterialValuePropertyRowGenerator::Get()
+{
+	static TSharedRef<FDMMaterialValuePropertyRowGenerator> Generator = MakeShared<FDMMaterialValuePropertyRowGenerator>();
+	return Generator;
+}
+
+void FDMMaterialValuePropertyRowGenerator::AddComponentProperties(const TSharedRef<SDMComponentEdit>& InComponentEditWidget, UDMMaterialComponent* InComponent,
+	TArray<FDMPropertyHandle>& InOutPropertyRows, TSet<UDMMaterialComponent*>& InOutProcessedObjects)
+{
+	if (!IsValid(InComponent))
+	{
+		return;
+	}
+
+	if (InOutProcessedObjects.Contains(InComponent))
+	{
+		return;
+	}
+
+	UDMMaterialValue* Value = Cast<UDMMaterialValue>(InComponent);
+
+	if (!Value)
+	{
+		return;
+	}
+
+	// The base material value class is abstract and not allowed.
+	if (Value->GetClass() == UDMMaterialValue::StaticClass())
+	{
+		return;
+	}
+
+	InOutProcessedObjects.Add(InComponent);
+
+	FDMPropertyHandle Handle = SDMEditor::GetPropertyHandle(&*InComponentEditWidget, Value, UDMMaterialValue::ValueName);
+
+	Handle.ResetToDefaultOverride = FResetToDefaultOverride::Create(
+		FIsResetToDefaultVisible::CreateUObject(Value, &UDMMaterialValue::CanResetToDefault),
+		FResetToDefaultHandler::CreateUObject(Value, &UDMMaterialValue::ResetToDefault)
+	);
+
+	InOutPropertyRows.Add(Handle);
+
+	const TArray<FName>& Properties = Value->GetEditableProperties();
+
+	for (const FName& Property : Properties)
+	{
+		if (Property == UDMMaterialValue::ValueName)
+		{
+			continue;
+		}
+			
+		if (InComponent->IsPropertyVisible(Property))
+		{
+			AddPropertyEditRows(InComponentEditWidget, InComponent, Property, InOutPropertyRows, InOutProcessedObjects);
+		}
+	}
+}
diff --git a/Engine/Plugins/Experimental/DynamicMaterial/Source/DynamicMaterialEditor/Private/Slate/Properties/Generators/DMMaterialValuePropertyRowGenerator.h b/Engine/Plugins/Experimental/DynamicMaterial/Source/DynamicMaterialEditor/Private/Slate/Properties/Generators/DMMaterialValuePropertyRowGenerator.h
new file mode 100644
--- /dev/null
+++ b/Engine/Plugins/Experimental/DynamicMaterial/Source/DynamicMaterialEditor/Private/Slate/Properties/Generators/DMMaterialValuePropertyRowGenerator.h
@@ -0,0 +1,17 @@
+// Copyright Epic Games, Inc. All Rights Reserved.
+
+#pragma once
+
+#include "Slate/Properties/Generators/DMComponentPropertyRowGenerator.h"
+
+class FDMMaterialValuePropertyRowGenerator : public FDMComponentPropertyRowGenerator
+{
+public:
+	static const TSharedRef<FDMMaterialValuePropertyRowGenerator>& Get();
+
+	FDMMaterialValuePropertyRowGenerator() = default;
+	virtual ~FDMMaterialValuePropertyRowGenerator() override = default;
+
+	virtual void AddComponentProperties(const TSharedRef<SDMComponentEdit>& InComponentEditWidget, UDMMaterialComponent* InComponent,
+		TArray<FDMPropertyHandle>& InOutPropertyRows, TSet<UDMMaterialComponent*>& InOutProcessedObjects) override;
+};
diff --git a/Engine/Plugins/Experimental/DynamicMaterial/Source/DynamicMaterialEditor/Private/Slate/Properties/Generators/DMStagePropertyRowGenerator.cpp b/Engine/Plugins/Experimental/DynamicMaterial/Source/DynamicMaterialEditor/Private/Slate/Properties/Generators/DMStagePropertyRowGenerator.cpp
new file mode 100644
--- /dev/null
+++ b/Engine/Plugins/Experimental/DynamicMaterial/Source/DynamicMaterialEditor/Private/Slate/Properties/Generators/DMStagePropertyRowGenerator.cpp
@@ -0,0 +1,45 @@
+// Copyright Epic Games, Inc. All Rights Reserved.
+
+#include "Slate/Properties/Generators/DMStagePropertyRowGenerator.h"
+#include "DynamicMaterialEditorModule.h"
+#include "Components/DMMaterialComponent.h"
+#include "Components/DMMaterialStage.h"
+#include "Components/DMMaterialStageSource.h"
+#include "Slate/SDMComponentEdit.h"
+
+const TSharedRef<FDMStagePropertyRowGenerator>& FDMStagePropertyRowGenerator::Get()
+{
+	static TSharedRef<FDMStagePropertyRowGenerator> Generator = MakeShared<FDMStagePropertyRowGenerator>();
+	return Generator;
+}
+
+void FDMStagePropertyRowGenerator::AddComponentProperties(const TSharedRef<SDMComponentEdit>& InComponentEditWidget, UDMMaterialComponent* InComponent,
+	TArray<FDMPropertyHandle>& InOutPropertyRows, TSet<UDMMaterialComponent*>& InOutProcessedObjects)
+{
+	if (!IsValid(InComponent))
+	{
+		return;
+	}
+
+	if (InOutProcessedObjects.Contains(InComponent))
+	{
+		return;
+	}
+
+	UDMMaterialStage* Stage = Cast<UDMMaterialStage>(InComponent);
+
+	if (!Stage)
+	{
+		return;
+	}
+
+	UDMMaterialStageSource* Source = Stage->GetSource();
+
+	if (!Source)
+	{
+		return;
+	}
+
+	FDynamicMaterialEditorModule::GeneratorComponentPropertyRows(InComponentEditWidget, Source, InOutPropertyRows, InOutProcessedObjects);
+	FDMComponentPropertyRowGenerator::AddComponentProperties(InComponentEditWidget, Stage, InOutPropertyRows, InOutProcessedObjects);
+}
\ No newline at end of file
diff --git a/Engine/Plugins/Experimental/DynamicMaterial/Source/DynamicMaterialEditor/Private/Slate/Properties/Generators/DMStagePropertyRowGenerator.h b/Engine/Plugins/Experimental/DynamicMaterial/Source/DynamicMaterialEditor/Private/Slate/Properties/Generators/DMStagePropertyRowGenerator.h
new file mode 100644
--- /dev/null
+++ b/Engine/Plugins/Experimental/DynamicMaterial/Source/DynamicMaterialEditor/Private/Slate/Properties/Generators/DMStagePropertyRowGenerator.h
@@ -0,0 +1,17 @@
+// Copyright Epic Games, Inc. All Rights Reserved.
+
+#pragma once
+
+#include "Slate/Properties/Generators/DMComponentPropertyRowGenerator.h"
+
+class FDMStagePropertyRowGenerator : public FDMComponentPropertyRowGenerator
+{
+public:
+	static const TSharedRef<FDMStagePropertyRowGenerator>& Get();
+
+	FDMStagePropertyRowGenerator() = default;
+	virtual ~FDMStagePropertyRowGenerator() override = default;
+
+	virtual void AddComponentProperties(const TSharedRef<SDMComponentEdit>& InComponentEditWidget, UDMMaterialComponent* InComponent,
+		TArray<FDMPropertyHandle>& InOutPropertyRows, TSet<UDMMaterialComponent*>& InOutProcessedObjects) override;
+};
diff --git a/Engine/Plugins/Experimental/DynamicMaterial/Source/DynamicMaterialEditor/Private/Slate/Properties/Generators/DMTextureUVPropertyRowGenerator.cpp b/Engine/Plugins/Experimental/DynamicMaterial/Source/DynamicMaterialEditor/Private/Slate/Properties/Generators/DMTextureUVPropertyRowGenerator.cpp
new file mode 100644
--- /dev/null
+++ b/Engine/Plugins/Experimental/DynamicMaterial/Source/DynamicMaterialEditor/Private/Slate/Properties/Generators/DMTextureUVPropertyRowGenerator.cpp
@@ -0,0 +1,414 @@
+// Copyright Epic Games, Inc. All Rights Reserved.
+
+#include "Slate/Properties/Generators/DMTextureUVPropertyRowGenerator.h"
+#include "Components/DMMaterialComponent.h"
+#include "Components/DMTextureUV.h"
+#include "DMEDefs.h"
+#include "Modules/ModuleManager.h"
+#include "PropertyCustomizationHelpers.h"
+#include "PropertyEditorModule.h"
+#include "ScopedTransaction.h"
+#include "Slate/SDMEditor.h"
+#include "Slate/Properties/Editors/SDMPropertyEditFloat.h"
+#include "Slate/Properties/Editors/SDMPropertyEditVector.h"
+#include "Slate/SDMComponentEdit.h"
+#include "Styling/SlateIconFinder.h"
+#include "Widgets/Layout/SBox.h"
+
+#define LOCTEXT_NAMESPACE "DMTextureUVPropertyRowGenerator"
+
+const TSharedRef<FDMTextureUVPropertyRowGenerator>& FDMTextureUVPropertyRowGenerator::Get()
+{
+	static TSharedRef<FDMTextureUVPropertyRowGenerator> Generator = MakeShared<FDMTextureUVPropertyRowGenerator>();
+	return Generator;
+}
+
+namespace UE::DynamicMaterialEditor::Private
+{
+	void AddRotationRow(const TSharedRef<SDMComponentEdit>& InComponentEditWidget, UDMTextureUV* InTextureUV, TArray<FDMPropertyHandle>& InOutPropertyRows);
+
+	void AddScaleRow(const TSharedRef<SDMComponentEdit>& InComponentEditWidget, UDMTextureUV* InTextureUV, TArray<FDMPropertyHandle>& InOutPropertyRows);
+
+	TSharedRef<SWidget> CreateScaleLinkWidget(const TSharedRef<SDMComponentEdit>& InComponentEditWidget, UDMTextureUV* InTextureUV);
+	ECheckBoxState GetScaleLinkWidgetState(UDMTextureUV* InTextureUV);
+	void OnScaleLinkWidgetChanged(ECheckBoxState InNewState, UDMTextureUV* InTextureUV);
+
+	TSharedRef<SWidget> CreateScaleValueWidget(const TSharedRef<SDMComponentEdit>& InComponentEditWidget, UDMTextureUV* InTextureUV);
+
+	void AddMirrorRow(const TSharedRef<SDMComponentEdit>& InComponentEditWidget, UDMTextureUV* InTextureUV, TArray<FDMPropertyHandle>& InOutPropertyRows);
+	TSharedRef<SWidget> CreateMirrorExtensionButtons(const TSharedRef<SDMComponentEdit>& InComponentEditWidget, UDMTextureUV* InTextureUV);
+}
+
+void FDMTextureUVPropertyRowGenerator::AddComponentProperties(const TSharedRef<SDMComponentEdit>& InComponentEditWidget, UDMMaterialComponent* InComponent,
+	TArray<FDMPropertyHandle>& InOutPropertyRows, TSet<UDMMaterialComponent*>& InOutProcessedObjects)
+{
+	if (!IsValid(InComponent))
+	{
+		return;
+	}
+
+	if (InOutProcessedObjects.Contains(InComponent))
+	{
+		return;
+	}
+
+	UDMTextureUV* TextureUV = Cast<UDMTextureUV>(InComponent);
+
+	if (!TextureUV)
+	{
+		return;
+	}
+
+	InOutProcessedObjects.Add(InComponent);
+
+	using namespace UE::DynamicMaterialEditor::Private;
+
+	const int PropertyCount = InOutPropertyRows.Num();
+
+	AddPropertyEditRows(InComponentEditWidget, InComponent, UDMTextureUV::NAME_Offset, InOutPropertyRows, InOutProcessedObjects);
+	AddRotationRow(InComponentEditWidget, TextureUV, InOutPropertyRows);
+	AddScaleRow(InComponentEditWidget, TextureUV, InOutPropertyRows);
+	AddPropertyEditRows(InComponentEditWidget, InComponent, UDMTextureUV::NAME_Pivot, InOutPropertyRows, InOutProcessedObjects);
+	AddMirrorRow(InComponentEditWidget, TextureUV, InOutPropertyRows);
+}
+
+bool FDMTextureUVPropertyRowGenerator::AllowKeyframeButton(UDMMaterialComponent* InComponent, FProperty* InProperty)
+{
+	if (InProperty)
+	{
+		const bool* AddKeyframeButtonPtr = UDMTextureUV::TextureProperties.Find(InProperty->GetFName());
+
+		if (AddKeyframeButtonPtr)
+		{
+			return *AddKeyframeButtonPtr;
+		}
+	}
+
+	return FDMComponentPropertyRowGenerator::AllowKeyframeButton(InComponent, InProperty);
+}
+
+void UE::DynamicMaterialEditor::Private::AddRotationRow(const TSharedRef<SDMComponentEdit>& InComponentEditWidget, UDMTextureUV* InTextureUV, TArray<FDMPropertyHandle>& InOutPropertyRows)
+{
+	InOutPropertyRows.Add(SDMEditor::GetPropertyHandle(&*InComponentEditWidget, InTextureUV, UDMTextureUV::NAME_Rotation));
+}
+
+void UE::DynamicMaterialEditor::Private::AddScaleRow(const TSharedRef<SDMComponentEdit>& InComponentEditWidget, UDMTextureUV* InTextureUV, TArray<FDMPropertyHandle>& InOutPropertyRows)
+{
+	InOutPropertyRows.Add(SDMEditor::GetPropertyHandle(&*InComponentEditWidget, InTextureUV, UDMTextureUV::NAME_Scale));
+}
+
+TSharedRef<SWidget> UE::DynamicMaterialEditor::Private::CreateScaleLinkWidget(const TSharedRef<SDMComponentEdit>& InComponentEditWidget, UDMTextureUV* InTextureUV)
+{
+	const FSlateBrush* Checked = FSlateIconFinder::FindIcon("PropertyWindow.Locked").GetIcon();
+	const FSlateBrush* Unchecked = FSlateIconFinder::FindIcon("PropertyWindow.Unlocked").GetIcon();
+	const FSlateBrush* NoBorder = FSlateIconFinder::FindIcon("NoBorder").GetIcon();
+
+	return SNew(SCheckBox)
+		.IsChecked_Static(&UE::DynamicMaterialEditor::Private::GetScaleLinkWidgetState, InTextureUV)
+		.OnCheckStateChanged_Static(&UE::DynamicMaterialEditor::Private::OnScaleLinkWidgetChanged, InTextureUV)
+		.CheckedImage(Checked)
+		.CheckedHoveredImage(Checked)
+		.CheckedPressedImage(Checked)
+		.UncheckedImage(Unchecked)
+		.UncheckedHoveredImage(Unchecked)
+		.UncheckedPressedImage(Unchecked)
+		.BackgroundImage(NoBorder)
+		.BackgroundHoveredImage(NoBorder)
+		.BackgroundPressedImage(NoBorder);
+}
+
+ECheckBoxState UE::DynamicMaterialEditor::Private::GetScaleLinkWidgetState(UDMTextureUV* InTextureUV)
+{
+	if (!IsValid(InTextureUV))
+	{
+		return ECheckBoxState::Unchecked;
+	}
+
+	return InTextureUV->bLinkScale ? ECheckBoxState::Checked : ECheckBoxState::Unchecked;
+}
+
+void UE::DynamicMaterialEditor::Private::OnScaleLinkWidgetChanged(ECheckBoxState InNewState, UDMTextureUV* InTextureUV)
+{
+	if (!IsValid(InTextureUV))
+	{
+		return;
+	}
+
+	InTextureUV->bLinkScale = (InNewState == ECheckBoxState::Checked);
+}
+
+namespace UE::DynamicMaterialEditor::Private
+{
+	class SDMPropertyEditVectorScale : public SDMPropertyEditVector
+	{
+	public:
+		using Super = SDMPropertyEditVector;
+
+		void Construct(const FArguments& InArgs, const TSharedPtr<SDMComponentEdit>& InComponentEditWidget, UDMTextureUV* InTextureUV)
+		{
+			if (ensure(IsValid(InTextureUV)))
+			{
+				Super::Construct(
+					Super::FArguments()
+						.ComponentEditWidget(InComponentEditWidget), 
+					InTextureUV->GetPropertyHandle(UDMTextureUV::NAME_Scale), 
+					2
+				);
+			}
+		}
+
+	protected:
+		FVector2D InitialScale;
+
+		virtual TSharedRef<SWidget> GetComponentWidget(int32 InIndex) override
+		{
+			TSharedRef<SWidget> Widget = Super::GetComponentWidget(InIndex);
+
+			if (InIndex > 0)
+			{
+				TWeakPtr<SDMPropertyEditVectorScale> ThisWeak = SharedThis(this).ToWeakPtr();
+
+				Widget->SetVisibility(TAttribute<EVisibility>::CreateLambda([ThisWeak]()
+					{
+						if (TSharedPtr<SDMPropertyEditVectorScale> This = ThisWeak.Pin())
+						{
+							TArray<UObject*> Outers;
+							This->PropertyHandle->GetOuterObjects(Outers);
+
+							if (Outers.Num() == 1 && IsValid(Outers[0]))
+							{
+								if (UDMTextureUV* TextureUV = Cast<UDMTextureUV>(Outers[0]))
+								{
+									if (TextureUV->bLinkScale && FMath::IsNearlyEqual(TextureUV->GetScale().X, TextureUV->GetScale().Y))
+									{
+										return EVisibility::Hidden;
+									}
+								}
+							}
+						}
+
+						return EVisibility::Visible;
+					}));
+				Widget->SetEnabled(TAttribute<bool>::CreateLambda([ThisWeak]()
+					{
+						if (TSharedPtr<SDMPropertyEditVectorScale> This = ThisWeak.Pin())
+						{
+							TArray<UObject*> Outers;
+							This->PropertyHandle->GetOuterObjects(Outers);
+
+							if (Outers.Num() == 1 && IsValid(Outers[0]))
+							{
+								if (UDMTextureUV* TextureUV = Cast<UDMTextureUV>(Outers[0]))
+								{
+									if (TextureUV->bLinkScale)
+									{
+										return false;
+									}
+								}
+							}
+						}
+
+						return true;
+					}));
+			}
+
+			return Widget;
+		}
+
+		virtual void StartTransaction(FText InDescription) override
+		{
+			Super::StartTransaction(InDescription);
+
+			if (PropertyHandle.IsValid() && PropertyHandle->IsValidHandle())
+			{
+				TArray<UObject*> Outers;
+				PropertyHandle->GetOuterObjects(Outers);
+
+				if (Outers.IsEmpty() == false)
+				{
+					if (UDMTextureUV* TextureUV = Cast<UDMTextureUV>(Outers[0]))
+					{
+						InitialScale = TextureUV->GetScale();
+					}
+				}
+			}
+		}
+
+		virtual void OnValueChanged(float InNewValue, int32 InComponent) const override
+		{
+			if (PropertyHandle.IsValid() && PropertyHandle->IsValidHandle())
+			{
+				TArray<UObject*> Outers;
+				PropertyHandle->GetOuterObjects(Outers);
+
+				if (Outers.IsEmpty() == false)
+				{
+					if (UDMTextureUV* TextureUV = Cast<UDMTextureUV>(Outers[0]))
+					{
+						FVector2D CurrentScale = TextureUV->GetScale();
+
+						if (FMath::IsNearlyEqual(CurrentScale[InComponent], InNewValue))
+						{
+							return;
+						}
+
+						auto IsScaledLinked = [this, TextureUV]()
+						{
+							if (FMath::IsNearlyZero(InitialScale.X) || FMath::IsNearlyZero(InitialScale.Y))
+							{
+								return false;
+							}
+
+							if (!IsValid(TextureUV))
+							{
+								return false;
+							}
+
+							return TextureUV->bLinkScale;
+						};
+
+						if (IsScaledLinked())
+						{
+							switch (InComponent)
+							{
+								case 0: // X
+									CurrentScale.Y = (InNewValue / InitialScale.X) * InitialScale.Y;
+									break;
+
+								case 1: // Y
+									CurrentScale.X = (InNewValue / InitialScale.Y) * InitialScale.X;
+									break;
+
+								default:
+									checkNoEntry();
+							}
+						}
+
+						switch (InComponent)
+						{
+							case 0: // X
+								CurrentScale.X = InNewValue;
+								break;
+
+							case 1: // Y
+								CurrentScale.Y = InNewValue;
+								break;
+
+							default:
+								checkNoEntry();
+						}
+
+						TextureUV->SetScale(CurrentScale);
+					}
+				}
+			}
+		}
+	};
+}
+
+TSharedRef<SWidget> UE::DynamicMaterialEditor::Private::CreateScaleValueWidget(const TSharedRef<SDMComponentEdit>& InComponentEditWidget, UDMTextureUV* InTextureUV)
+{
+	return SNew(UE::DynamicMaterialEditor::Private::SDMPropertyEditVectorScale, InComponentEditWidget, InTextureUV);
+}
+
+void UE::DynamicMaterialEditor::Private::AddMirrorRow(const TSharedRef<SDMComponentEdit>& InComponentEditWidget, UDMTextureUV* InTextureUV, TArray<FDMPropertyHandle>& InOutPropertyRows)
+{
+	InOutPropertyRows.Add(SDMEditor::GetPropertyHandle(&*InComponentEditWidget, InTextureUV, UDMTextureUV::NAME_bMirrorOnX));
+	InOutPropertyRows.Add(SDMEditor::GetPropertyHandle(&*InComponentEditWidget, InTextureUV, UDMTextureUV::NAME_bMirrorOnY));
+}
+
+TSharedRef<SWidget> UE::DynamicMaterialEditor::Private::CreateMirrorExtensionButtons(const TSharedRef<SDMComponentEdit>& InComponentEditWidget, UDMTextureUV* InTextureUV)
+{
+	if (!IsValid(InTextureUV))
+	{
+		return SNullWidget::NullWidget;
+	}
+
+	FProperty* XProperty = InTextureUV->GetClass()->FindPropertyByName(UDMTextureUV::NAME_bMirrorOnX);
+	FProperty* YProperty = InTextureUV->GetClass()->FindPropertyByName(UDMTextureUV::NAME_bMirrorOnY);
+
+	if (!XProperty || !YProperty)
+	{
+		return SNullWidget::NullWidget;
+	}
+
+	TSharedPtr<IPropertyHandle> XPropertyHandle = InTextureUV->GetPropertyHandle(UDMTextureUV::NAME_bMirrorOnX);
+	TSharedPtr<IPropertyHandle> YPropertyHandle = InTextureUV->GetPropertyHandle(UDMTextureUV::NAME_bMirrorOnY);
+	TSharedPtr<IDetailTreeNode> XDetailTreeNode = InTextureUV->GetDetailTreeNode(UDMTextureUV::NAME_bMirrorOnX);
+	TSharedPtr<IDetailTreeNode> YDetailTreeNode = InTextureUV->GetDetailTreeNode(UDMTextureUV::NAME_bMirrorOnY);
+
+	if (!XPropertyHandle.IsValid() || !YPropertyHandle.IsValid() || !XDetailTreeNode.IsValid() || !YDetailTreeNode.IsValid())
+	{
+		return SNullWidget::NullWidget;
+	}
+
+	ensure(XPropertyHandle->GetProperty() == XProperty && YPropertyHandle->GetProperty() == YProperty);
+
+	FPropertyEditorModule& PropertyEditorModule = FModuleManager::GetModuleChecked<FPropertyEditorModule>("PropertyEditor");
+
+	TSharedRef<SWidget> ResetButtonWidget = PropertyCustomizationHelpers::MakeResetButton(
+		FSimpleDelegate::CreateLambda([XPropertyHandle, YPropertyHandle]()
+			{
+				if (XPropertyHandle.IsValid() && XPropertyHandle->IsValidHandle()
+					&& YPropertyHandle.IsValid() && YPropertyHandle->IsValidHandle())
+				{
+					TArray<UObject*> Outers;
+					XPropertyHandle->GetOuterObjects(Outers);
+
+					if (Outers.IsEmpty() == false)
+					{
+						if (UDMTextureUV* TextureUV = Cast<UDMTextureUV>(Outers[0]))
+						{
+							if (TextureUV->GetMirrorOnX() || TextureUV->GetMirrorOnY())
+							{
+								FScopedTransaction Transaction(LOCTEXT("ResetMirror", "Reset Texture UV Mirror to default."));
+								TextureUV->Modify();
+								TextureUV->SetMirrorOnX(false);
+								TextureUV->SetMirrorOnY(false);
+							}
+						}
+					}
+				}
+			})
+	);
+	ResetButtonWidget->SetVisibility(
+		TAttribute<EVisibility>::CreateLambda([XPropertyHandle, YPropertyHandle]()
+			{
+				if (XPropertyHandle.IsValid() && XPropertyHandle->IsValidHandle()
+					&& YPropertyHandle.IsValid() && YPropertyHandle->IsValidHandle())
+				{
+					TArray<UObject*> Outers;
+					XPropertyHandle->GetOuterObjects(Outers);
+
+					if (Outers.IsEmpty() == false)
+					{
+						if (UDMTextureUV* TextureUV = Cast<UDMTextureUV>(Outers[0]))
+						{
+							return (TextureUV->GetMirrorOnX() || TextureUV->GetMirrorOnY())
+								? EVisibility::Visible
+								: EVisibility::Hidden;
+						}
+					}
+				}
+
+				return EVisibility::Hidden;
+			})
+	);
+
+	TSharedRef<SHorizontalBox> ButtonsBox = SNew(SHorizontalBox)
+		+ SHorizontalBox::Slot()
+		[
+			ResetButtonWidget
+		];
+
+	ButtonsBox->AddSlot()
+		[
+			SNew(SBox)
+			.WidthOverride(20.f)
+			.HeightOverride(18.f)
+		];
+
+	return ButtonsBox;
+}
+
+#undef LOCTEXT_NAMESPACE
diff --git a/Engine/Plugins/Experimental/DynamicMaterial/Source/DynamicMaterialEditor/Private/Slate/Properties/Generators/DMTextureUVPropertyRowGenerator.h b/Engine/Plugins/Experimental/DynamicMaterial/Source/DynamicMaterialEditor/Private/Slate/Properties/Generators/DMTextureUVPropertyRowGenerator.h
new file mode 100644
--- /dev/null
+++ b/Engine/Plugins/Experimental/DynamicMaterial/Source/DynamicMaterialEditor/Private/Slate/Properties/Generators/DMTextureUVPropertyRowGenerator.h
@@ -0,0 +1,19 @@
+// Copyright Epic Games, Inc. All Rights Reserved.
+
+#pragma once
+
+#include "Slate/Properties/Generators/DMComponentPropertyRowGenerator.h"
+
+class FDMTextureUVPropertyRowGenerator : public FDMComponentPropertyRowGenerator
+{
+public:
+	static const TSharedRef<FDMTextureUVPropertyRowGenerator>& Get();
+
+	FDMTextureUVPropertyRowGenerator() = default;
+	virtual ~FDMTextureUVPropertyRowGenerator() override = default;
+
+	virtual void AddComponentProperties(const TSharedRef<SDMComponentEdit>& InComponentEditWidget, UDMMaterialComponent* InComponent,
+		TArray<FDMPropertyHandle>& InOutPropertyRows, TSet<UDMMaterialComponent*>& InOutProcessedObjects) override;
+
+	virtual bool AllowKeyframeButton(UDMMaterialComponent* InComponent, FProperty* InProperty) override;
+};
diff --git a/Engine/Plugins/Experimental/DynamicMaterial/Source/DynamicMaterialEditor/Private/Slate/Properties/Generators/DMThroughputPropertyRowGenerator.cpp b/Engine/Plugins/Experimental/DynamicMaterial/Source/DynamicMaterialEditor/Private/Slate/Properties/Generators/DMThroughputPropertyRowGenerator.cpp
new file mode 100644
--- /dev/null
+++ b/Engine/Plugins/Experimental/DynamicMaterial/Source/DynamicMaterialEditor/Private/Slate/Properties/Generators/DMThroughputPropertyRowGenerator.cpp
@@ -0,0 +1,93 @@
+// Copyright Epic Games, Inc. All Rights Reserved.
+
+#include "Slate/Properties/Generators/DMThroughputPropertyRowGenerator.h"
+#include "DynamicMaterialEditorModule.h"
+#include "Components/DMMaterialStage.h"
+#include "Components/DMMaterialStageInput.h"
+#include "Components/DMMaterialStageThroughput.h"
+
+const TSharedRef<FDMThroughputPropertyRowGenerator>& FDMThroughputPropertyRowGenerator::Get()
+{
+	static TSharedRef<FDMThroughputPropertyRowGenerator> Generator = MakeShared<FDMThroughputPropertyRowGenerator>();
+	return Generator;
+}
+
+void FDMThroughputPropertyRowGenerator::AddComponentProperties(const TSharedRef<SDMComponentEdit>& InComponentEditWidget,
+	UDMMaterialComponent* InComponent, TArray<FDMPropertyHandle>& InOutPropertyRows,
+	TSet<UDMMaterialComponent*>& InOutProcessedObjects)
+{
+	if (!IsValid(InComponent))
+	{
+		return;
+	}
+
+	if (InOutProcessedObjects.Contains(InComponent))
+	{
+		return;
+	}
+
+	UDMMaterialStageThroughput* Throughput = Cast<UDMMaterialStageThroughput>(InComponent);
+
+	if (!Throughput)
+	{
+		return;
+	}
+
+	InOutProcessedObjects.Add(Throughput);
+
+	const TArray<FName>& ThroughputProperties = Throughput->GetEditableProperties();
+
+	for (const FName& ThroughputProperty : ThroughputProperties)
+	{
+		if (InComponent->IsPropertyVisible(ThroughputProperty))
+		{
+			AddPropertyEditRows(InComponentEditWidget, InComponent, ThroughputProperty, InOutPropertyRows, InOutProcessedObjects);
+		}
+	}
+
+	UDMMaterialStage* Stage = Throughput->GetStage();
+
+	if (!Stage)
+	{
+		return;
+	}
+
+	InOutProcessedObjects.Add(Stage);
+
+	static const FName StageInputsName = GET_MEMBER_NAME_CHECKED(UDMMaterialStage, Inputs);
+
+	{
+		const TArray<FDMMaterialStageConnector>& InputConnectors = Throughput->GetInputConnectors();
+		const TArray<FDMMaterialStageConnection>& InputMap = Stage->GetInputConnectionMap();
+		TArray<UDMMaterialStageInput*> Inputs = Stage->GetInputs();
+
+		for (int32 InputIdx = 0; InputIdx < InputConnectors.Num(); ++InputIdx)
+		{
+			if (!Throughput->IsInputVisible(InputIdx) || !Throughput->CanChangeInput(InputIdx))
+			{
+				continue;
+			}
+
+			for (const FDMMaterialStageConnectorChannel& Channel : InputMap[InputIdx].Channels)
+			{
+				const int32 StageInputIdx = Channel.SourceIndex - FDMMaterialStageConnectorChannel::FIRST_STAGE_INPUT;
+
+				if (Inputs.IsValidIndex(StageInputIdx))
+				{
+					FDynamicMaterialEditorModule::GeneratorComponentPropertyRows(InComponentEditWidget, Inputs[StageInputIdx], InOutPropertyRows, InOutProcessedObjects);
+				}
+			}
+		}
+	}
+
+	const TArray<FName>& StageProperties = Stage->GetEditableProperties();
+
+	for (const FName& StageProperty : StageProperties)
+	{
+		if (StageProperty != StageInputsName)
+		{
+			AddPropertyEditRows(InComponentEditWidget, Stage, StageProperty, InOutPropertyRows, InOutProcessedObjects);
+			continue;
+		}
+	}
+}
diff --git a/Engine/Plugins/Experimental/DynamicMaterial/Source/DynamicMaterialEditor/Private/Slate/Properties/Generators/DMThroughputPropertyRowGenerator.h b/Engine/Plugins/Experimental/DynamicMaterial/Source/DynamicMaterialEditor/Private/Slate/Properties/Generators/DMThroughputPropertyRowGenerator.h
new file mode 100644
--- /dev/null
+++ b/Engine/Plugins/Experimental/DynamicMaterial/Source/DynamicMaterialEditor/Private/Slate/Properties/Generators/DMThroughputPropertyRowGenerator.h
@@ -0,0 +1,17 @@
+// Copyright Epic Games, Inc. All Rights Reserved.
+
+#pragma once
+
+#include "Slate/Properties/Generators/DMComponentPropertyRowGenerator.h"
+
+class FDMThroughputPropertyRowGenerator : public FDMComponentPropertyRowGenerator
+{
+public:
+	static const TSharedRef<FDMThroughputPropertyRowGenerator>& Get();
+
+	FDMThroughputPropertyRowGenerator() = default;
+	virtual ~FDMThroughputPropertyRowGenerator() override = default;
+
+	virtual void AddComponentProperties(const TSharedRef<SDMComponentEdit>& InComponentEditWidget, UDMMaterialComponent* InComponent,
+		TArray<FDMPropertyHandle>& InOutPropertyRows, TSet<UDMMaterialComponent*>& InOutProcessedObjects) override;
+};
