[MaterialDesigner]

diff --git a/Engine/Plugins/Experimental/DynamicMaterial/Source/DynamicMaterialEditor/Private/Slate/Properties/SDMBlendMode.cpp b/Engine/Plugins/Experimental/DynamicMaterial/Source/DynamicMaterialEditor/Private/Slate/Properties/SDMBlendMode.cpp
new file mode 100644
--- /dev/null
+++ b/Engine/Plugins/Experimental/DynamicMaterial/Source/DynamicMaterialEditor/Private/Slate/Properties/SDMBlendMode.cpp
@@ -0,0 +1,106 @@
+// Copyright Epic Games, Inc. All Rights Reserved.
+
+#include "Slate/Properties/SDMBlendMode.h"
+#include "DMPrivate.h"
+#include "Engine/EngineTypes.h"
+#include "Widgets/Input/SComboBox.h"
+#include "Widgets/Text/STextBlock.h"
+
+namespace UE::DynamicMaterialEditor::Private
+{
+	static const TArray<EBlendMode> SupportedBlendModes =
+	{
+		EBlendMode::BLEND_Opaque,
+		EBlendMode::BLEND_Masked,
+		EBlendMode::BLEND_Translucent,
+		EBlendMode::BLEND_Additive,
+		EBlendMode::BLEND_Modulate
+	};
+
+	const TArray<FName>& SupportedBlendModeNames()
+	{
+		static TArray<FName> OutNames;
+
+		if (OutNames.IsEmpty())
+		{
+			const UEnum* const Enum = StaticEnum<EBlendMode>();
+
+			if (ensure(IsValid(Enum)))
+			{
+				for (const EBlendMode BlendMode : SupportedBlendModes)
+				{
+					const FName Name = Enum->GetNameByValue(BlendMode);
+					OutNames.Add(Name);
+				}
+			}
+		}
+
+		return OutNames;
+	}
+}
+
+void SDMBlendMode::Construct(const FArguments& InArgs)
+{
+	SelectedItem = InArgs._SelectedItem;
+	OnSelectedItemChanged = InArgs._OnSelectedItemChanged;
+
+	BlendModeEnum = StaticEnum<EBlendMode>();
+
+	const EBlendMode InitialBlendMode = InArgs._SelectedItem.Get();
+	const FName InitialBlendModeName = BlendModeEnum->GetNameByValue(InitialBlendMode);
+
+	ChildSlot
+	[
+		SNew(SComboBox<FName>)
+		.InitiallySelectedItem(InitialBlendModeName)
+		.OptionsSource(&UE::DynamicMaterialEditor::Private::SupportedBlendModeNames())
+		.OnGenerateWidget(this, &SDMBlendMode::OnGenerateWidget)
+		.OnSelectionChanged(this, &SDMBlendMode::OnSelectionChanged)
+		[
+			SNew(STextBlock)
+			.Text(this, &SDMBlendMode::GetSelectedItemText)
+		]
+	];
+}
+
+TSharedRef<SWidget> SDMBlendMode::OnGenerateWidget(const FName InItem)
+{
+	if (!BlendModeEnum->IsValidEnumName(InItem))
+	{
+		return SNullWidget::NullWidget;
+	}
+
+	return 
+		SNew(STextBlock)
+		.Text(GetBlendModeDisplayName(InItem));
+}
+
+void SDMBlendMode::OnSelectionChanged(const FName InNewItem, const ESelectInfo::Type InSelectInfoType)
+{
+	const EBlendMode BlendMode = static_cast<EBlendMode>(BlendModeEnum->GetValueByName(InNewItem));
+	OnSelectedItemChanged.ExecuteIfBound(BlendMode);
+}
+
+FText SDMBlendMode::GetSelectedItemText() const
+{
+	const FName BlendModeName = BlendModeEnum->GetNameByValue(SelectedItem.Get());
+
+	if (BlendModeName == NAME_None)
+	{
+		return FText::GetEmpty();
+	}
+
+	return GetBlendModeDisplayName(BlendModeName);
+}
+
+FText SDMBlendMode::GetBlendModeDisplayName(const FName InBlendModeName) const
+{
+	const int64 BlendModeValue = BlendModeEnum->GetValueByName(InBlendModeName);
+
+	if (BlendModeValue == INDEX_NONE)
+	{
+		return FText::GetEmpty();
+	}
+
+	return BlendModeEnum->GetDisplayNameTextByValue(BlendModeValue);
+}
diff --git a/Engine/Plugins/Experimental/DynamicMaterial/Source/DynamicMaterialEditor/Private/Slate/Properties/SDMBlendMode.h b/Engine/Plugins/Experimental/DynamicMaterial/Source/DynamicMaterialEditor/Private/Slate/Properties/SDMBlendMode.h
new file mode 100644
--- /dev/null
+++ b/Engine/Plugins/Experimental/DynamicMaterial/Source/DynamicMaterialEditor/Private/Slate/Properties/SDMBlendMode.h
@@ -0,0 +1,39 @@
+// Copyright Epic Games, Inc. All Rights Reserved.
+
+#pragma once
+
+#include "Engine/EngineTypes.h"
+#include "Widgets/SCompoundWidget.h"
+
+class UEnum;
+
+DECLARE_DELEGATE_OneParam(FDMOnBlendModeChanged, const EBlendMode)
+
+class SDMBlendMode : public SCompoundWidget
+{
+public:
+	SLATE_BEGIN_ARGS(SDMBlendMode)
+		: _SelectedItem(BLEND_Translucent)
+		{}
+		SLATE_ATTRIBUTE(EBlendMode, SelectedItem)
+		SLATE_EVENT(FDMOnBlendModeChanged, OnSelectedItemChanged)
+	SLATE_END_ARGS()
+
+	virtual ~SDMBlendMode() {}
+
+	void Construct(const FArguments& InArgs);
+
+protected:
+	TAttribute<EBlendMode> SelectedItem;
+	FDMOnBlendModeChanged OnSelectedItemChanged;
+
+	TObjectPtr<UEnum> BlendModeEnum;
+
+	TSharedRef<SWidget> OnGenerateWidget(const FName InItem);
+
+	void OnSelectionChanged(const FName InNewItem, const ESelectInfo::Type InSelectInfoType);
+
+	FText GetSelectedItemText() const;
+
+	FText GetBlendModeDisplayName(const FName InBlendModeName) const;
+};
diff --git a/Engine/Plugins/Experimental/DynamicMaterial/Source/DynamicMaterialEditor/Private/Slate/Properties/SDMDetailsGrid.cpp b/Engine/Plugins/Experimental/DynamicMaterial/Source/DynamicMaterialEditor/Private/Slate/Properties/SDMDetailsGrid.cpp
new file mode 100644
--- /dev/null
+++ b/Engine/Plugins/Experimental/DynamicMaterial/Source/DynamicMaterialEditor/Private/Slate/Properties/SDMDetailsGrid.cpp
@@ -0,0 +1,367 @@
+// Copyright Epic Games, Inc. All Rights Reserved.
+
+#include "Slate/Properties/SDMDetailsGrid.h"
+#include "Components/DMMaterialValue.h"
+#include "DMWorldSubsystem.h"
+#include "DynamicMaterialEditorStyle.h"
+#include "Engine/World.h"
+#include "Framework/MultiBox/MultiBoxBuilder.h"
+#include "IDetailKeyframeHandler.h"
+#include "Modules/ModuleManager.h"
+#include "PropertyCustomizationHelpers.h"
+#include "PropertyEditorModule.h"
+#include "Slate/Properties/SDMPropertyEdit.h"
+#include "Widgets/Colors/SColorBlock.h"
+#include "Widgets/Layout/SSplitter.h"
+
+#define LOCTEXT_NAMESPACE "SDMDetailsGrid"
+
+void SDMDetailsGrid::Construct(const FArguments& InArgs)
+{
+	MinColumnWidth = InArgs._MinColumnWidth;
+	RowSeparatorColor = InArgs._RowSeparatorColor;
+	bResizeable = InArgs._Resizeable;
+	bShowSeparators = InArgs._ShowSeparators;
+	SeparatorColor = InArgs._SeparatorColor;
+	UseEvenOddRowBackgroundColors = InArgs._UseEvenOddRowBackgroundColors;
+	OddRowBackgroundColor = InArgs._OddRowBackgroundColor;
+	EvenRowBackgroundColor = InArgs._EvenRowBackgroundColor;
+	LabelFillWidth = InArgs._LabelFillWidth;
+	ContentFillWidth = InArgs._ContentFillWidth;
+
+	ChildSlot
+		[
+			SAssignNew(MainContainer, SVerticalBox)
+		];
+}
+
+TSharedRef<SWidget> SDMDetailsGrid::CreateSplitterRow(const TSharedPtr<SWidget>& InLabel, const TSharedPtr<SWidget>& InContent, const TSharedPtr<SWidget>& InExtensions)
+{
+	TSharedRef<SSplitter> Splitter = 
+		SNew(SSplitter)
+		.Style(FAppStyle::Get(), "DetailsView.Splitter");
+
+	if (InLabel.IsValid())
+	{
+		Splitter->AddSlot()
+			.Resizable(bResizeable)
+			.SizeRule(bResizeable ? SSplitter::ESizeRule::FractionOfParent : SSplitter::ESizeRule::SizeToContent)
+			.Value(LabelFillWidth)
+			[
+				InLabel.ToSharedRef()
+			];
+	}
+
+	if (InContent.IsValid())
+	{
+		Splitter->AddSlot()
+			.Resizable(bResizeable)
+			.SizeRule(bResizeable ? SSplitter::ESizeRule::FractionOfParent : SSplitter::ESizeRule::SizeToContent)
+			.Value(ContentFillWidth)
+			[
+				InContent.ToSharedRef()
+			];
+	}
+	else
+	{
+		Splitter->AddSlot()
+			.Resizable(bResizeable)
+			.SizeRule(SSplitter::ESizeRule::FractionOfParent)
+			.Value(ContentFillWidth)
+			[
+				SNullWidget::NullWidget
+			];
+	}
+
+	if (InExtensions.IsValid())
+	{
+		Splitter->AddSlot()
+			.Resizable(false)
+			.SizeRule(SSplitter::ESizeRule::SizeToContent)
+			[
+				InExtensions.ToSharedRef()
+			];
+	}
+
+	return Splitter;
+}
+
+TSharedRef<SWidget> SDMDetailsGrid::CreateHorizontalRow(const TSharedPtr<SWidget>& InLabel, const TSharedPtr<SWidget>& InContent, const TSharedPtr<SWidget>& InExtensions)
+{
+	TSharedRef<SHorizontalBox> HBox = SNew(SHorizontalBox);
+
+	if (InLabel.IsValid())
+	{
+		HBox->AddSlot()
+			.FillWidth(LabelFillWidth)
+			.HAlign(HAlign_Left)
+			.VAlign(VAlign_Center)
+			[
+				InLabel.ToSharedRef()
+			];
+	}
+
+	if (InContent.IsValid())
+	{
+		HBox->AddSlot()
+			.FillWidth(ContentFillWidth)
+			.HAlign(HAlign_Fill)
+			.VAlign(VAlign_Center)
+			[
+				InContent.ToSharedRef()
+			];
+	}
+	else
+	{
+		HBox->AddSlot()
+			.FillWidth(ContentFillWidth)
+			.HAlign(HAlign_Fill)
+			.VAlign(VAlign_Center)
+			[
+				SNullWidget::NullWidget
+			];
+	}
+
+	if (InExtensions.IsValid())
+	{
+		if (!bResizeable && bShowSeparators)
+		{
+			HBox->AddSlot()
+				.AutoWidth()
+				.HAlign(HAlign_Left)
+				.VAlign(VAlign_Fill)
+				.Padding(0.0f, 0.0f, 5.0f, 0.0f)
+				[
+					SNew(SColorBlock)
+					.Color(SeparatorColor)
+				];
+		}
+
+		HBox->AddSlot()
+			.AutoWidth()
+			.HAlign(HAlign_Right)
+			.VAlign(VAlign_Center)
+			[
+				InExtensions.ToSharedRef()
+			];
+	}
+
+	return HBox;
+}
+
+void SDMDetailsGrid::AddRow(const TSharedRef<SWidget>& InLabel, const TSharedRef<SWidget>& InContent, const TSharedPtr<SWidget>& InExtensions)
+{
+	TSharedPtr<SWidget> NewRowWidget;
+
+	if (bResizeable)
+	{
+		NewRowWidget = CreateSplitterRow(InLabel, InContent, InExtensions);
+	}
+	else
+	{
+		NewRowWidget = CreateHorizontalRow(InLabel, InContent, InExtensions);
+	}
+
+	MainContainer->AddSlot()
+		.AutoHeight()
+		.HAlign(HAlign_Fill)
+		.VAlign(VAlign_Center)
+		[
+			NewRowWidget.ToSharedRef()
+		];
+}
+
+void SDMDetailsGrid::AddRow_TextLabel(const TAttribute<FText>& InText, const TSharedRef<SWidget>& InContent, const TSharedRef<SWidget>& InExtensions)
+{
+	AddRow(CreateDefaultLabel(InText), InContent, InExtensions);
+}
+
+void SDMDetailsGrid::AddRow_TextLabel(const TAttribute<FText>& InText, const TSharedRef<SWidget>& InContent)
+{
+	AddRow(CreateDefaultLabel(InText), InContent);
+}
+
+void SDMDetailsGrid::AddFullRowContent(const TSharedRef<SWidget>& InContent)
+{
+	MainContainer->AddSlot()
+		.AutoHeight()
+		.HAlign(HAlign_Fill)
+		.VAlign(VAlign_Center)
+		[
+			InContent
+		];
+}
+
+void SDMDetailsGrid::AddPropertyRow(UObject* InObject, const TSharedPtr<IPropertyHandle>& InPropertyHandle)
+{
+	FOnGenerateGlobalRowExtensionArgs ExtensionRowArgs;
+	ExtensionRowArgs.OwnerObject = InObject;
+	ExtensionRowArgs.Property = InPropertyHandle->GetProperty();
+	ExtensionRowArgs.PropertyPath = TEXT("Value");
+	ExtensionRowArgs.PropertyHandle = InPropertyHandle;
+
+	AddPropertyRow(InObject, InPropertyHandle, SNew(SDMPropertyEdit).PropertyHandle(InPropertyHandle), ExtensionRowArgs);
+}
+
+void SDMDetailsGrid::AddPropertyRow(UObject* InObject, const TSharedPtr<IPropertyHandle>& InPropertyHandle, const TSharedRef<SWidget>& InEditWidget)
+{
+	FOnGenerateGlobalRowExtensionArgs ExtensionRowArgs;
+	ExtensionRowArgs.OwnerObject = InObject;
+	ExtensionRowArgs.Property = InPropertyHandle->GetProperty();
+	ExtensionRowArgs.PropertyPath = TEXT("Value");
+	ExtensionRowArgs.PropertyHandle = InPropertyHandle;
+
+	AddPropertyRow(InObject, InPropertyHandle, InEditWidget, ExtensionRowArgs);
+}
+
+void SDMDetailsGrid::AddPropertyRow(UObject* InObject, const TSharedPtr<IPropertyHandle>& InPropertyHandle, const TSharedRef<SWidget>& InEditWidget,
+	const FOnGenerateGlobalRowExtensionArgs& InExtensionArgs)
+{
+	if (!IsValid(InObject) || !InPropertyHandle.IsValid() || !InPropertyHandle->IsValidHandle())
+	{
+		return;
+	}
+
+	const FText DisplayNameText = InPropertyHandle->GetPropertyDisplayName();
+
+	AddRow_TextLabel(DisplayNameText, InEditWidget, CreateExtensionButtons(InObject, InPropertyHandle, InExtensionArgs));
+}
+
+TSharedRef<SWidget> SDMDetailsGrid::CreateDefaultLabel(const TAttribute<FText>& InText)
+{
+	return
+		SNew(STextBlock)
+		.Text(InText);
+}
+
+TSharedRef<SWidget> SDMDetailsGrid::CreateExtensionButtons(UObject* InObject, const TSharedPtr<IPropertyHandle>& InPropertyHandle,
+	const FOnGenerateGlobalRowExtensionArgs& InExtensionArgs)
+{
+	FSlimHorizontalToolBarBuilder ToolBarBuilder(TSharedPtr<FUICommandList>(), FMultiBoxCustomization::None);
+ 
+	TSharedRef<SWidget> ResetButtonWidget = PropertyCustomizationHelpers::MakeResetButton(
+		FSimpleDelegate::CreateLambda([InPropertyHandle, InExtensionArgs]()
+			{
+				if (InPropertyHandle.IsValid() && InPropertyHandle->IsValidHandle())
+				{
+					InPropertyHandle->ResetToDefault();
+				}
+			})
+	);
+	ResetButtonWidget->SetVisibility(
+		TAttribute<EVisibility>::CreateLambda([InPropertyHandle, InExtensionArgs]()
+			{
+				if (InPropertyHandle.IsValid() && InPropertyHandle->IsValidHandle())
+				{
+					return InPropertyHandle->CanResetToDefault() ? EVisibility::Collapsed : EVisibility::Visible;
+				}
+				return EVisibility::Collapsed;
+			})
+	);
+
+	ToolBarBuilder.AddWidget(ResetButtonWidget);
+	
+
+	static FSlateIcon CreateKeyIcon(FAppStyle::Get().GetStyleSetName(), "Sequencer.AddKey.Details");
+ 
+	ToolBarBuilder.AddToolBarButton(
+		FUIAction(
+			FExecuteAction::CreateStatic(&SDMDetailsGrid::CreateKey, InObject, InPropertyHandle)
+		),
+		NAME_None,
+		FText::GetEmpty(),
+		LOCTEXT("CreateKeyToolTip", "Add a keyframe for this property."),
+		CreateKeyIcon
+	);
+ 
+	FPropertyEditorModule& PropertyEditorModule = FModuleManager::GetModuleChecked<FPropertyEditorModule>("PropertyEditor");
+ 
+	TArray<FPropertyRowExtensionButton> ExtensionButtons;
+	PropertyEditorModule.GetGlobalRowExtensionDelegate().Broadcast(InExtensionArgs, ExtensionButtons);
+ 
+	if (ExtensionButtons.IsEmpty() == false)
+	{
+		// Build extension toolbar 
+		ToolBarBuilder.SetLabelVisibility(EVisibility::Collapsed);
+		ToolBarBuilder.SetStyle(&FAppStyle::Get(), "DetailsView.ExtensionToolBar");
+		for (const FPropertyRowExtensionButton& Extension : ExtensionButtons)
+		{
+			ToolBarBuilder.AddToolBarButton(Extension.UIAction, NAME_None, Extension.Label, Extension.ToolTip, Extension.Icon);
+		}
+	}
+ 
+	return ToolBarBuilder.MakeWidget();
+}
+ 
+FOnGenerateGlobalRowExtensionArgs SDMDetailsGrid::CreateExtensionArgs(UDMMaterialValue* InValue)
+{
+	FOnGenerateGlobalRowExtensionArgs ExtensionRowArgs;
+
+	if (ensure(InValue))
+	{
+		FProperty* ValueProperty = InValue->GetClass()->FindPropertyByName("Value");
+		TSharedPtr<IPropertyHandle> ValuePropertyHandle = InValue->GetPropertyHandle();
+
+		if (ensure(ValuePropertyHandle.IsValid() && ValueProperty == ValuePropertyHandle->GetProperty()))
+		{
+			ExtensionRowArgs.OwnerObject = InValue;
+			ExtensionRowArgs.Property = ValueProperty;
+			ExtensionRowArgs.PropertyPath = TEXT("Value");
+			ExtensionRowArgs.OwnerTreeNode = InValue->GetDetailTreeNode();
+			ExtensionRowArgs.PropertyHandle = ValuePropertyHandle;
+		}
+	}
+
+	return ExtensionRowArgs;
+}
+
+void SDMDetailsGrid::CreateKey(UObject* InObject, TSharedPtr<IPropertyHandle> InPropertyHandle)
+{
+	if (!IsValid(InObject) || !InPropertyHandle.IsValid())
+	{
+		return;
+	}
+
+	const UWorld* const World = InObject->GetWorld();
+	if (!IsValid(World))
+	{
+		return;
+	}
+
+	const UDMWorldSubsystem* WorldSubsystem = World->GetSubsystem<UDMWorldSubsystem>();
+	if (!IsValid(WorldSubsystem))
+	{
+		return;
+	}
+
+	TSharedPtr<IDetailKeyframeHandler> KeyframeHandler = WorldSubsystem->GetKeyframeHandler();
+	if (!KeyframeHandler.IsValid())
+	{
+		return;
+	}
+
+	if (!KeyframeHandler->IsPropertyKeyable(InObject->GetClass(), *InPropertyHandle))
+	{
+		return;
+	}
+
+	KeyframeHandler->OnKeyPropertyClicked(*InPropertyHandle);
+}
+
+void SDMDetailsGrid::CreateKey(UDMMaterialValue* InMaterialValue)
+{
+	if (!IsValid(InMaterialValue))
+	{
+		return;
+	}
+
+	TSharedPtr<IPropertyHandle> PropertyHandle = InMaterialValue->GetPropertyHandle();
+	if (!PropertyHandle.IsValid())
+	{
+		return;
+	}
+
+	CreateKey(InMaterialValue, PropertyHandle);
+}
+
+#undef LOCTEXT_NAMESPACE
diff --git a/Engine/Plugins/Experimental/DynamicMaterial/Source/DynamicMaterialEditor/Private/Slate/Properties/SDMDetailsGrid.h b/Engine/Plugins/Experimental/DynamicMaterial/Source/DynamicMaterialEditor/Private/Slate/Properties/SDMDetailsGrid.h
new file mode 100644
--- /dev/null
+++ b/Engine/Plugins/Experimental/DynamicMaterial/Source/DynamicMaterialEditor/Private/Slate/Properties/SDMDetailsGrid.h
@@ -0,0 +1,87 @@
+// Copyright Epic Games, Inc. All Rights Reserved.
+
+#pragma once
+
+#include "DMEDefs.h"
+#include "Widgets/SCompoundWidget.h"
+
+class IPropertyHandle;
+class SVerticalBox;
+class UDMMaterialValue;
+class UWorld;
+struct FOnGenerateGlobalRowExtensionArgs;
+
+/**
+ * Material Designer Property Details Grid Panel
+ * 
+ * Manages width, labels, splitters, extensions and styling for a group of material properties.
+*/
+class SDMDetailsGrid : public SCompoundWidget
+{
+	SLATE_BEGIN_ARGS(SDMDetailsGrid)
+		: _MinColumnWidth(0.0f)
+		, _RowSeparatorColor(FLinearColor(0.0f, 0.0f, 0.0f, 1.0f))
+		, _Resizeable(false)
+		, _ShowSeparators(true)
+		, _SeparatorColor(FLinearColor(1, 1, 1, 0.2f))
+		, _UseEvenOddRowBackgroundColors(false)
+		, _OddRowBackgroundColor(FLinearColor(0.0f, 0.0f, 0.0f, 0.0f))
+		, _EvenRowBackgroundColor(FLinearColor(0.0f, 0.0f, 0.0f, 0.1f))
+		, _LabelFillWidth(0.2f)
+		, _ContentFillWidth(0.8f)
+	{}
+		SLATE_ARGUMENT(float, MinColumnWidth)
+		SLATE_ARGUMENT(FLinearColor, RowSeparatorColor)
+		SLATE_ARGUMENT(bool, Resizeable)
+		SLATE_ARGUMENT(bool, ShowSeparators)
+		SLATE_ARGUMENT(FLinearColor, SeparatorColor)
+		SLATE_ARGUMENT(bool, UseEvenOddRowBackgroundColors)
+		SLATE_ARGUMENT(FLinearColor, OddRowBackgroundColor)
+		SLATE_ARGUMENT(FLinearColor, EvenRowBackgroundColor)
+		SLATE_ATTRIBUTE(float, LabelFillWidth)
+		SLATE_ATTRIBUTE(float, ContentFillWidth)
+	SLATE_END_ARGS()
+
+public:
+	void Construct(const FArguments& InArgs);
+
+	/**  */
+	void AddRow(const TSharedRef<SWidget>& InLabel, const TSharedRef<SWidget>& InContent, const TSharedPtr<SWidget>& InExtensions = nullptr);
+	void AddRow_TextLabel(const TAttribute<FText>& InText, const TSharedRef<SWidget>& InContent, const TSharedRef<SWidget>& InExtensions);
+	void AddRow_TextLabel(const TAttribute<FText>& InText, const TSharedRef<SWidget>& InContent);
+	void AddFullRowContent(const TSharedRef<SWidget>& InContent);
+
+	void AddPropertyRow(UObject* InObject, const TSharedPtr<IPropertyHandle>& InPropertyHandle);
+	void AddPropertyRow(UObject* InObject, const TSharedPtr<IPropertyHandle>& InPropertyHandle, const TSharedRef<SWidget>& InEditWidget);
+	void AddPropertyRow(UObject* InObject, const TSharedPtr<IPropertyHandle>& InPropertyHandle, const TSharedRef<SWidget>& InEditWidget, 
+		const FOnGenerateGlobalRowExtensionArgs& InExtensionArgs);
+
+	static TSharedRef<SWidget> CreateDefaultLabel(const TAttribute<FText>& InText);
+
+	static TSharedRef<SWidget> CreateExtensionButtons(UObject* InObject, const TSharedPtr<IPropertyHandle>& InPropertyHandle, 
+		const FOnGenerateGlobalRowExtensionArgs& InExtensionArgs);
+
+protected:
+	float MinColumnWidth;
+	FLinearColor RowSeparatorColor;
+	bool bResizeable;
+	bool bShowSeparators;
+	FLinearColor SeparatorColor;
+	bool UseEvenOddRowBackgroundColors;
+	FLinearColor OddRowBackgroundColor;
+	FLinearColor EvenRowBackgroundColor;
+	TAttribute<float> LabelFillWidth;
+	TAttribute<float> ContentFillWidth;
+
+	TSharedPtr<SVerticalBox> MainContainer;
+
+	TSharedRef<SWidget> CreateSplitterRow(const TSharedPtr<SWidget>& InLabel, const TSharedPtr<SWidget>& InContent, 
+		const TSharedPtr<SWidget>& InExtensions);
+	TSharedRef<SWidget> CreateHorizontalRow(const TSharedPtr<SWidget>& InLabel, const TSharedPtr<SWidget>& InContent,
+		const TSharedPtr<SWidget>& InExtensions);
+
+	FOnGenerateGlobalRowExtensionArgs CreateExtensionArgs(UDMMaterialValue* InValue);
+
+	static void CreateKey(UObject* InObject, TSharedPtr<IPropertyHandle> InPropertyHandle);
+	static void CreateKey(UDMMaterialValue* InMaterialValue);
+};
diff --git a/Engine/Plugins/Experimental/DynamicMaterial/Source/DynamicMaterialEditor/Private/Slate/Properties/SDMDetailsGridRow.cpp b/Engine/Plugins/Experimental/DynamicMaterial/Source/DynamicMaterialEditor/Private/Slate/Properties/SDMDetailsGridRow.cpp
new file mode 100644
--- /dev/null
+++ b/Engine/Plugins/Experimental/DynamicMaterial/Source/DynamicMaterialEditor/Private/Slate/Properties/SDMDetailsGridRow.cpp
@@ -0,0 +1,159 @@
+// Copyright Epic Games, Inc. All Rights Reserved.
+
+#include "Slate/Properties/SDMDetailsGridRow.h"
+#include "Components/DMMaterialValue.h"
+#include "PropertyCustomizationHelpers.h"
+#include "Slate/Properties/SDMPropertyEdit.h"
+#include "Widgets/Colors/SColorBlock.h"
+#include "Widgets/Layout/SSplitter.h"
+
+#define LOCTEXT_NAMESPACE "SDMDetailsGridRow"
+
+void SDMDetailsGridRow::Construct(const FArguments& InArgs)
+{
+	MinColumnWidth = InArgs._MinColumnWidth;
+	RowSeparatorColor = InArgs._RowSeparatorColor;
+	bResizeable = InArgs._Resizeable;
+	bShowSeparators = InArgs._ShowSeparators;
+	SeparatorColor = InArgs._SeparatorColor;
+	UseEvenOddRowBackgroundColors = InArgs._UseEvenOddRowBackgroundColors;
+	OddRowBackgroundColor = InArgs._OddRowBackgroundColor;
+	EvenRowBackgroundColor = InArgs._EvenRowBackgroundColor;
+
+	LabelColumnBox = SNew(SVerticalBox);
+	ContentColumnBox = SNew(SVerticalBox);
+	ExtensionsColumnBox = SNew(SVerticalBox);
+
+	ChildSlot
+		[
+			SNew(SSplitter)
+			.Style(FAppStyle::Get(), "DetailsView.Splitter")
+			+ SSplitter::Slot()
+			.Resizable(InArgs._Resizeable)
+			.Value(0.3f)
+			[
+				LabelColumnBox.ToSharedRef()
+			]
+			+ SSplitter::Slot()
+			.Resizable(InArgs._Resizeable)
+			.Value(0.7f)
+			[
+				ContentColumnBox.ToSharedRef()
+			]
+			+ SSplitter::Slot()
+			.Resizable(false)
+			[
+				ExtensionsColumnBox.ToSharedRef()
+			]
+		];
+}
+
+void SDMDetailsGridRow::AddPropertyRow(UObject* InObject, const TSharedPtr<IPropertyHandle>& InPropertyHandle)
+{
+	FOnGenerateGlobalRowExtensionArgs ExtensionRowArgs;
+	ExtensionRowArgs.OwnerObject = InObject;
+	ExtensionRowArgs.Property = InPropertyHandle->GetProperty();
+	ExtensionRowArgs.PropertyPath = TEXT("Value");
+	ExtensionRowArgs.PropertyHandle = InPropertyHandle;
+
+	AddPropertyRow(InObject, InPropertyHandle, SNew(SDMPropertyEdit).PropertyHandle(InPropertyHandle), ExtensionRowArgs);
+}
+
+void SDMDetailsGridRow::AddPropertyRow(UObject* InObject, const TSharedPtr<IPropertyHandle>& InPropertyHandle, const TSharedRef<SWidget>& InEditWidget)
+{
+	FOnGenerateGlobalRowExtensionArgs ExtensionRowArgs;
+	ExtensionRowArgs.OwnerObject = InObject;
+	ExtensionRowArgs.Property = InPropertyHandle->GetProperty();
+	ExtensionRowArgs.PropertyPath = TEXT("Value");
+	ExtensionRowArgs.PropertyHandle = InPropertyHandle;
+
+	AddPropertyRow(InObject, InPropertyHandle, InEditWidget, ExtensionRowArgs);
+}
+
+void SDMDetailsGridRow::AddPropertyRow(UObject* InObject, const TSharedPtr<IPropertyHandle>& InPropertyHandle, const TSharedRef<SWidget>& InEditWidget,
+	const FOnGenerateGlobalRowExtensionArgs& InExtensionArgs)
+{
+	if (!IsValid(InObject) || !InPropertyHandle.IsValid() || !InPropertyHandle->IsValidHandle())
+	{
+		return;
+	}
+
+	const FText DisplayNameText = InPropertyHandle->GetPropertyDisplayName();
+
+	// ?! No idea why this is commented out.
+	//AddPropertyRow(DisplayNameText, InEditWidget, CreateExtensionButtons(InObject, InPropertyHandle, InExtensionArgs));
+}
+
+void SDMDetailsGridRow::AddPropertyRow(const FText& InText, const TSharedRef<SWidget>& InContent, const TSharedRef<SWidget>& InExtensions)
+{
+	const FMargin SeparatorPadding = FMargin(5.0f, 0.0f, 5.0f, 0.0f);
+
+	TSharedRef<SHorizontalBox> LabelWidget = SNew(SHorizontalBox);
+
+	LabelWidget->AddSlot()
+		.AutoWidth()
+		[
+			InExtensions
+		];
+
+	if (bShowSeparators)
+	{
+		LabelWidget->AddSlot()
+			.AutoWidth()
+			.HAlign(HAlign_Center)
+			.VAlign(VAlign_Fill)
+			.Padding(SeparatorPadding)
+			[
+				SNew(SColorBlock)
+				.Color(SeparatorColor)
+			];
+	}
+
+	LabelColumnBox->AddSlot()
+		.AutoHeight()
+		.HAlign(HAlign_Left)
+		.VAlign(VAlign_Center)
+		[
+			SNew(STextBlock)
+			.Text(InText)
+		];
+
+	ContentColumnBox->AddSlot()
+		.AutoHeight()
+		.HAlign(HAlign_Left)
+		.VAlign(VAlign_Center)
+		[
+			InContent
+		];
+
+	TSharedRef<SHorizontalBox> ExtensionsWidget = SNew(SHorizontalBox);
+
+	if (bShowSeparators)
+	{
+		ExtensionsWidget->AddSlot()
+			.AutoWidth()
+			.HAlign(HAlign_Center)
+			.VAlign(VAlign_Fill)
+			.Padding(SeparatorPadding)
+			[
+				SNew(SColorBlock)
+				.Color(SeparatorColor)
+			];
+	}
+
+	ExtensionsWidget->AddSlot()
+		.AutoWidth()
+		[
+			InExtensions
+		];
+
+	ExtensionsColumnBox->AddSlot()
+		.AutoHeight()
+		.HAlign(HAlign_Right)
+		.VAlign(VAlign_Center)
+		[
+			ExtensionsWidget
+		];
+}
+
+#undef LOCTEXT_NAMESPACE
diff --git a/Engine/Plugins/Experimental/DynamicMaterial/Source/DynamicMaterialEditor/Private/Slate/Properties/SDMDetailsGridRow.h b/Engine/Plugins/Experimental/DynamicMaterial/Source/DynamicMaterialEditor/Private/Slate/Properties/SDMDetailsGridRow.h
new file mode 100644
--- /dev/null
+++ b/Engine/Plugins/Experimental/DynamicMaterial/Source/DynamicMaterialEditor/Private/Slate/Properties/SDMDetailsGridRow.h
@@ -0,0 +1,58 @@
+// Copyright Epic Games, Inc. All Rights Reserved.
+
+#pragma once
+
+#include "DMEDefs.h"
+#include "Widgets/SCompoundWidget.h"
+
+class IPropertyHandle;
+class SVerticalBox;
+class UDMMaterialValue;
+class UWorld;
+struct FOnGenerateGlobalRowExtensionArgs;
+
+class SDMDetailsGridRow : public SCompoundWidget
+{
+	SLATE_BEGIN_ARGS(SDMDetailsGridRow)
+		: _MinColumnWidth(0.0f)
+		, _RowSeparatorColor(FLinearColor(0.0f, 0.0f, 0.0f, 1.0f))
+		, _Resizeable(true)
+		, _ShowSeparators(true)
+		, _SeparatorColor(FLinearColor(1, 1, 1, 0.2f))
+		, _UseEvenOddRowBackgroundColors(false)
+		, _OddRowBackgroundColor(FLinearColor(0.0f, 0.0f, 0.0f, 0.0f))
+		, _EvenRowBackgroundColor(FLinearColor(0.0f, 0.0f, 0.0f, 0.1f))
+	{}
+		SLATE_ARGUMENT(float, MinColumnWidth)
+		SLATE_ARGUMENT(FLinearColor, RowSeparatorColor)
+		SLATE_ARGUMENT(bool, Resizeable)
+		SLATE_ARGUMENT(bool, ShowSeparators)
+		SLATE_ARGUMENT(FLinearColor, SeparatorColor)
+		SLATE_ARGUMENT(bool, UseEvenOddRowBackgroundColors)
+		SLATE_ARGUMENT(FLinearColor, OddRowBackgroundColor)
+		SLATE_ARGUMENT(FLinearColor, EvenRowBackgroundColor)
+	SLATE_END_ARGS()
+
+public:
+	void Construct(const FArguments& InArgs);
+
+	void AddPropertyRow(UObject* InObject, const TSharedPtr<IPropertyHandle>& InPropertyHandle);
+	void AddPropertyRow(UObject* InObject, const TSharedPtr<IPropertyHandle>& InPropertyHandle, const TSharedRef<SWidget>& InEditWidget);
+	void AddPropertyRow(UObject* InObject, const TSharedPtr<IPropertyHandle>& InPropertyHandle, const TSharedRef<SWidget>& InEditWidget,
+		const FOnGenerateGlobalRowExtensionArgs& InExtensionArgs);
+	void AddPropertyRow(const FText& InText, const TSharedRef<SWidget>& InContent, const TSharedRef<SWidget>& InExtensions);
+
+protected:
+	float MinColumnWidth;
+	FLinearColor RowSeparatorColor;
+	bool bResizeable;
+	bool bShowSeparators;
+	FLinearColor SeparatorColor;
+	bool UseEvenOddRowBackgroundColors;
+	FLinearColor OddRowBackgroundColor;
+	FLinearColor EvenRowBackgroundColor;
+
+	TSharedPtr<SVerticalBox> LabelColumnBox;
+	TSharedPtr<SVerticalBox> ContentColumnBox;
+	TSharedPtr<SVerticalBox> ExtensionsColumnBox;
+};
diff --git a/Engine/Plugins/Experimental/DynamicMaterial/Source/DynamicMaterialEditor/Private/Slate/Properties/SDMDomain.cpp b/Engine/Plugins/Experimental/DynamicMaterial/Source/DynamicMaterialEditor/Private/Slate/Properties/SDMDomain.cpp
new file mode 100644
--- /dev/null
+++ b/Engine/Plugins/Experimental/DynamicMaterial/Source/DynamicMaterialEditor/Private/Slate/Properties/SDMDomain.cpp
@@ -0,0 +1,103 @@
+// Copyright Epic Games, Inc. All Rights Reserved.
+
+#include "Slate/Properties/SDMDomain.h"
+#include "DMPrivate.h"
+#include "Engine/EngineTypes.h"
+#include "Widgets/Input/SComboBox.h"
+#include "Widgets/Text/STextBlock.h"
+
+namespace UE::DynamicMaterialEditor::Private
+{
+	static const TArray<EMaterialDomain> SupportedDomains =
+	{
+		EMaterialDomain::MD_Surface,
+		EMaterialDomain::MD_PostProcess
+	};
+
+	const TArray<FName>& SupportedDomainNames()
+	{
+		static TArray<FName> OutNames;
+
+		if (OutNames.IsEmpty())
+		{
+			const UEnum* const Enum = StaticEnum<EMaterialDomain>();
+
+			if (ensure(IsValid(Enum)))
+			{
+				for (const EMaterialDomain Domain : SupportedDomains)
+				{
+					const FName Name = Enum->GetNameByValue(Domain);
+					OutNames.Add(Name);
+				}
+			}
+		}
+
+		return OutNames;
+	}
+}
+
+void SDMDomain::Construct(const FArguments& InArgs)
+{
+	SelectedItem = InArgs._SelectedItem;
+	OnSelectedItemChanged = InArgs._OnSelectedItemChanged;
+
+	DomainEnum = StaticEnum<EMaterialDomain>();
+
+	const EMaterialDomain InitialDomain = InArgs._SelectedItem.Get();
+	const FName InitialDomainName = DomainEnum->GetNameByValue(InitialDomain);
+
+	ChildSlot
+	[
+		SNew(SComboBox<FName>)
+			.InitiallySelectedItem(InitialDomainName)
+			.OptionsSource(&UE::DynamicMaterialEditor::Private::SupportedDomainNames())
+			.OnGenerateWidget(this, &SDMDomain::OnGenerateWidget)
+			.OnSelectionChanged(this, &SDMDomain::OnSelectionChanged)
+			[
+				SNew(STextBlock)
+				.Text(this, &SDMDomain::GetSelectedItemText)
+			]
+	];
+}
+
+TSharedRef<SWidget> SDMDomain::OnGenerateWidget(const FName InItem)
+{
+	if (!DomainEnum->IsValidEnumName(InItem))
+	{
+		return SNullWidget::NullWidget;
+	}
+
+	return
+		SNew(STextBlock)
+		.Text(GetDomainDisplayName(InItem));
+}
+
+void SDMDomain::OnSelectionChanged(const FName InNewItem, const ESelectInfo::Type InSelectInfoType)
+{
+	const EMaterialDomain Domain = static_cast<EMaterialDomain>(DomainEnum->GetValueByName(InNewItem));
+	OnSelectedItemChanged.ExecuteIfBound(Domain);
+}
+
+FText SDMDomain::GetSelectedItemText() const
+{
+	const FName DomainName = DomainEnum->GetNameByValue(SelectedItem.Get());
+
+	if (DomainName == NAME_None)
+	{
+		return FText::GetEmpty();
+	}
+
+	return GetDomainDisplayName(DomainName);
+}
+
+FText SDMDomain::GetDomainDisplayName(const FName InDomainName) const
+{
+	const int64 DomainValue = DomainEnum->GetValueByName(InDomainName);
+
+	if (DomainValue == INDEX_NONE)
+	{
+		return FText::GetEmpty();
+	}
+
+	return DomainEnum->GetDisplayNameTextByValue(DomainValue);
+}
diff --git a/Engine/Plugins/Experimental/DynamicMaterial/Source/DynamicMaterialEditor/Private/Slate/Properties/SDMDomain.h b/Engine/Plugins/Experimental/DynamicMaterial/Source/DynamicMaterialEditor/Private/Slate/Properties/SDMDomain.h
new file mode 100644
--- /dev/null
+++ b/Engine/Plugins/Experimental/DynamicMaterial/Source/DynamicMaterialEditor/Private/Slate/Properties/SDMDomain.h
@@ -0,0 +1,40 @@
+// Copyright Epic Games, Inc. All Rights Reserved.
+
+#pragma once
+
+#include "Engine/EngineTypes.h"
+#include "MaterialDomain.h"
+#include "Widgets/SCompoundWidget.h"
+
+class UEnum;
+
+DECLARE_DELEGATE_OneParam(FDMOnDomainChanged, const EMaterialDomain)
+
+class SDMDomain : public SCompoundWidget
+{
+public:
+	SLATE_BEGIN_ARGS(SDMDomain)
+		: _SelectedItem(EMaterialDomain::MD_Surface)
+		{}
+		SLATE_ATTRIBUTE(EMaterialDomain, SelectedItem)
+		SLATE_EVENT(FDMOnDomainChanged, OnSelectedItemChanged)
+	SLATE_END_ARGS()
+
+	virtual ~SDMDomain() {}
+
+	void Construct(const FArguments& InArgs);
+
+protected:
+	TAttribute<EMaterialDomain> SelectedItem;
+	FDMOnDomainChanged OnSelectedItemChanged;
+
+	TObjectPtr<UEnum> DomainEnum;
+
+	TSharedRef<SWidget> OnGenerateWidget(const FName InItem);
+
+	void OnSelectionChanged(const FName InNewItem, const ESelectInfo::Type InSelectInfoType);
+
+	FText GetSelectedItemText() const;
+
+	FText GetDomainDisplayName(const FName InDomainName) const;
+};
diff --git a/Engine/Plugins/Experimental/DynamicMaterial/Source/DynamicMaterialEditor/Private/Slate/Properties/SDMMaterialParameters.cpp b/Engine/Plugins/Experimental/DynamicMaterial/Source/DynamicMaterialEditor/Private/Slate/Properties/SDMMaterialParameters.cpp
new file mode 100644
--- /dev/null
+++ b/Engine/Plugins/Experimental/DynamicMaterial/Source/DynamicMaterialEditor/Private/Slate/Properties/SDMMaterialParameters.cpp
@@ -0,0 +1,341 @@
+// Copyright Epic Games, Inc. All Rights Reserved.
+
+#include "SDMMaterialParameters.h"
+#include "Components/DMMaterialParameter.h"
+#include "Components/DMMaterialValue.h"
+#include "DetailLayoutBuilder.h"
+#include "DMValueDefinition.h"
+#include "DynamicMaterialEditorModule.h"
+#include "DynamicMaterialEditorStyle.h"
+#include "DynamicMaterialModule.h"
+#include "Model/DynamicMaterialModel.h"
+#include "ScopedTransaction.h"
+#include "Slate/Properties/SDMDetailsGrid.h"
+#include "SlateOptMacros.h"
+#include "Styling/StyleColors.h"
+#include "Widgets/Images/SImage.h"
+#include "Widgets/Input/SButton.h"
+#include "Widgets/Input/SEditableTextBox.h"
+#include "Widgets/Layout/SBox.h"
+#include "Widgets/Layout/SExpandableArea.h"
+#include "Widgets/SBoxPanel.h"
+#include "Widgets/Text/STextBlock.h"
+
+#define LOCTEXT_NAMESPACE "SDMMaterialParameters"
+
+BEGIN_SLATE_FUNCTION_BUILD_OPTIMIZATION
+
+void SDMMaterialParameters::Construct(const FArguments& InArgs, TWeakObjectPtr<UDynamicMaterialModel> InBuilderWeak)
+{
+	ensure(InBuilderWeak.IsValid());
+	ModelWeak = InBuilderWeak;
+
+	static const FVector2D AddButtonSize = FVector2D(12.f, 12.f);
+	static const FMargin AddButtonPadding = FMargin(0.f, 0.f, 0.f, 0.f);
+	static const FMargin AddButtonContentMargin = FMargin(1.f, 1.f, 1.f, 1.f);
+
+	auto CreateAddValueButton = [this](EDMValueType ValueType)
+	{
+		FText Tooltip = FText::Format(
+			LOCTEXT("AddValueTemplate", "Add {0} Value."), 
+			UDMValueDefinitionLibrary::GetValueDefinition(ValueType).GetDisplayName()
+		);
+
+		const FName Icon = FDynamicMaterialEditorStyle::GetBrushNameForType(ValueType);
+
+		SHorizontalBox::FSlot::FSlotArguments Slot(SHorizontalBox::Slot());
+		Slot
+			.Padding(AddButtonPadding)
+			.AutoWidth()
+			[
+				SNew(SButton)
+				.ButtonStyle(FAppStyle::Get(), "SimpleButton")
+				.ContentPadding(AddButtonContentMargin)
+				.OnClicked(this, &SDMMaterialParameters::OnAddValueButtonClicked, ValueType)
+				.ToolTipText(Tooltip)
+				[
+					SNew(SImage)
+					.Image(FDynamicMaterialEditorStyle::GetBrush(Icon))
+					.DesiredSizeOverride(AddButtonSize)
+				]
+			];
+		return Slot;
+	};
+
+	ChildSlot
+	[
+		SAssignNew(ValuesArea, SExpandableArea)
+		.HeaderContent()
+		[
+			SNew(SHorizontalBox)
+			.ToolTipText(LOCTEXT("GlobalValuesTooltip", "Values which can be re-used within your material. These values just be constants or exposed as modifiable paramters if given a name."))
+			+ SHorizontalBox::Slot()
+			.Padding(0.f, 1.f, 3.f, 0.f)
+			.AutoWidth()
+			.VAlign(EVerticalAlignment::VAlign_Center)
+			[
+				SNew(STextBlock)
+				.Text(LOCTEXT("GlobalValuesHeader", "Global Values"))
+				.Font(IDetailLayoutBuilder::GetDetailFontBold())
+			]
+			+ CreateAddValueButton(EDMValueType::VT_Float1)
+			+ CreateAddValueButton(EDMValueType::VT_Float2)
+			+ CreateAddValueButton(EDMValueType::VT_Float3_XYZ)
+			+ CreateAddValueButton(EDMValueType::VT_Float3_RPY)
+			+ CreateAddValueButton(EDMValueType::VT_Float3_RGB)
+			+ CreateAddValueButton(EDMValueType::VT_Float4_RGBA)
+			+ CreateAddValueButton(EDMValueType::VT_Texture)
+			+ CreateAddValueButton(EDMValueType::VT_Bool)
+		]
+		.InitiallyCollapsed(false)
+		.HeaderPadding(FMargin(3.f, 5.f, 3.f, 5.f))
+		.BodyContent()
+		[
+			SAssignNew(AreaBody, SBox)
+			.HAlign(EHorizontalAlignment::HAlign_Left)
+			[
+				CreateValuesWidget()
+			]
+		]
+	];
+}
+
+TSharedRef<SWidget> SDMMaterialParameters::CreateValuesWidget()
+{
+	if (!ModelWeak.IsValid())
+	{
+		// @formatter:off
+		return
+			SNew(SBox)
+			.HAlign(EHorizontalAlignment::HAlign_Center)
+			.VAlign(EVerticalAlignment::VAlign_Center)
+			.Padding(10.f, 5.f, 10.f, 5.f)
+			[
+				SNew(STextBlock)
+				.Text(LOCTEXT("GlobalValuesContent", "Global Values Content"))
+				.Font(IDetailLayoutBuilder::GetDetailFont())
+			];
+		// @formatter:on
+	}
+
+	const TArray<UDMMaterialValue*>& Values = ModelWeak->GetValues();
+
+	if (Values.IsEmpty())
+	{
+		// @formatter:off
+		return
+			SNew(SBox)
+			.HAlign(EHorizontalAlignment::HAlign_Center)
+			.VAlign(EVerticalAlignment::VAlign_Center)
+			.Padding(10.f, 5.f, 10.f, 5.f)
+			[
+				SNew(STextBlock)
+				.Text(LOCTEXT("GlobalValuesNoContent", "Click a button to add a value!"))
+				.Font(IDetailLayoutBuilder::GetDetailFont())
+			];
+		// @formatter:on
+	}
+
+	TSharedRef<SDMDetailsGrid> ValuesList = SNew(SDMDetailsGrid);
+
+	static const FMargin CellPadding = FMargin(8.f, 8.f, 8.f, 6.f);
+
+	for (int32 Index = 0; Index < Values.Num(); ++Index)
+	{
+		TSharedPtr<SWidget> ValueEditWidget = FDynamicMaterialEditorModule::GetValueEditWidgetDelegate(Values[Index]->GetClass())
+			.Execute(nullptr, Values[Index]);
+
+		if (ValueEditWidget.IsValid() == false)
+		{
+			continue;
+		}
+
+		const FName Icon = FDynamicMaterialEditorStyle::GetBrushNameForType(Values[Index]->GetType());
+
+		TSharedRef<SHorizontalBox> HBox = 
+			SNew(SHorizontalBox)
+			+ SHorizontalBox::Slot()
+			.AutoWidth()
+			.HAlign(HAlign_Center)
+			.VAlign(VAlign_Center)
+			.Padding(CellPadding)
+			[
+				SNew(SBox)
+				.WidthOverride(16.0f)
+				[
+					SNew(SButton)
+					.ButtonStyle(&FAppStyle::Get().GetWidgetStyle<FButtonStyle>("SimpleButton"))
+					.ContentPadding(FMargin(6.f, 6.f, 6.f, 6.f))
+					.OnClicked(this, &SDMMaterialParameters::OnRemoveValueButtonClicked, Index)
+					.ToolTipText(LOCTEXT("RemoveValue", "Remove Value"))
+					[
+						SNew(SImage)
+						.Image(FDynamicMaterialEditorStyle::GetBrush(Icon))
+						.DesiredSizeOverride(FVector2D(12.f, 12.f))
+					]
+				]
+			]
+			+ SHorizontalBox::Slot()
+			.AutoWidth()
+			.HAlign(HAlign_Left)
+			.VAlign(VAlign_Center)
+			.Padding(CellPadding)
+			[
+				SNew(SBox)
+				.WidthOverride(75.0f)
+				[
+					SNew(SEditableTextBox)
+					.MinDesiredWidth(69.f)
+					.BackgroundColor(FStyleColors::Background.GetSpecifiedColor())
+					.ForegroundColor(FStyleColors::AccentWhite.GetSpecifiedColor())
+					.Font(IDetailLayoutBuilder::GetDetailFont())
+					.HintText(LOCTEXT("Name", "Name"))
+					.IsEnabled(true)
+					.Text(Values[Index]->GetParameter() ? FText::FromName(Values[Index]->GetParameter()->GetParameterName()) : FText::GetEmpty())
+					.OnVerifyTextChanged(this, &SDMMaterialParameters::OnVerifyValueNameChanged, Index)
+					.OnTextChanged(this, &SDMMaterialParameters::OnAcceptValueNameChanged, Index)
+				]
+			]
+			+ SHorizontalBox::Slot()
+			.AutoWidth()
+			.HAlign(HAlign_Left)
+			.VAlign(VAlign_Center)
+			.Padding(CellPadding)
+			[
+				SNew(SBox)
+				.MinDesiredWidth(100.0f)
+				[
+					ValueEditWidget.ToSharedRef()
+				]
+			];
+
+
+		ValuesList->AddRow_TextLabel(FText::AsNumber(Index + 1), HBox);
+	}
+
+	return ValuesList;
+}
+
+END_SLATE_FUNCTION_BUILD_OPTIMIZATION
+
+void SDMMaterialParameters::RefreshWidgets()
+{
+	if (AreaBody.IsValid())
+	{
+		AreaBody->SetContent(CreateValuesWidget());
+	}
+}
+
+FReply SDMMaterialParameters::OnAddValueButtonClicked(EDMValueType Type)
+{
+	UDynamicMaterialModel* MaterialModel = ModelWeak.Get();
+
+	if (!MaterialModel)
+	{
+		return FReply::Unhandled();
+	}
+
+	if (Type <= EDMValueType::VT_None || Type == EDMValueType::VT_Float_Any || Type > EDMValueType::VT_Texture)
+	{
+		return FReply::Unhandled();
+	}
+
+	FScopedTransaction Transaction(LOCTEXT("AddGlobalValue", "Add Global Value"));
+	MaterialModel->Modify();
+	MaterialModel->AddValue(Type);
+
+	if (ValuesArea.IsValid())
+	{
+		if (!ValuesArea->IsExpanded())
+		{
+			ValuesArea->SetExpanded(true);
+		}
+	}
+
+	return FReply::Handled();
+}
+
+FReply SDMMaterialParameters::OnRemoveValueButtonClicked(int32 Index)
+{
+	UDynamicMaterialModel* MaterialModel = ModelWeak.Get();
+
+	if (!MaterialModel)
+	{
+		return FReply::Unhandled();
+	}
+
+	if (MaterialModel->GetValueByIndex(Index) == nullptr)
+	{
+		return FReply::Unhandled();
+	}
+
+	FScopedTransaction Transaction(LOCTEXT("RemoveGlobalValue", "Remove Global Value"));
+	MaterialModel->Modify();
+	MaterialModel->RemoveValueByIndex(Index);
+
+	return FReply::Handled();
+}
+
+bool SDMMaterialParameters::OnVerifyValueNameChanged(const FText& InNewName, FText& OutErrorText, int32 ValueIndex)
+{
+	UDynamicMaterialModel* MaterialModel = ModelWeak.Get();
+
+	if (!MaterialModel)
+	{
+		OutErrorText = LOCTEXT("InternalError", "Internal Error");
+		return false;
+	}
+
+	const FString NewNameString = InNewName.ToString();
+
+	if (NewNameString.Len() > 20)
+	{
+		OutErrorText = LOCTEXT("NameTooLong", "That name is too long.");
+		return false;
+	};
+
+	if (NewNameString.Len() == 0)
+	{
+		return true;
+	}
+
+	const TArray<UDMMaterialValue*>& Values = MaterialModel->GetValues();
+
+	if (!Values.IsValidIndex(ValueIndex))
+	{
+		OutErrorText = LOCTEXT("InternalError", "Internal Error");
+		return false;
+	}
+
+	if (MaterialModel->HasParameterName(*NewNameString))
+	{
+		OutErrorText = LOCTEXT("NameTaken", "That name has already been used.");
+		return false;
+	}
+
+	return true;
+}
+
+void SDMMaterialParameters::OnAcceptValueNameChanged(const FText& InNewName, int32 Index)
+{
+	FText ErrorText;
+
+	if (!OnVerifyValueNameChanged(InNewName, ErrorText, Index))
+	{
+		return;
+	}
+
+	const TArray<UDMMaterialValue*>& Values = ModelWeak->GetValues();
+
+	FScopedTransaction Transaction(LOCTEXT("SetGlobalValueName", "Set Global Value Name"));
+	Values[Index]->Modify();
+
+	if (UDMMaterialParameter* Parameter = Values[Index]->GetParameter())
+	{
+		Parameter->Modify();
+	}
+
+	Values[Index]->SetParameterName(*InNewName.ToString());
+}
+
+#undef LOCTEXT_NAMESPACE
diff --git a/Engine/Plugins/Experimental/DynamicMaterial/Source/DynamicMaterialEditor/Private/Slate/Properties/SDMMaterialParameters.h b/Engine/Plugins/Experimental/DynamicMaterial/Source/DynamicMaterialEditor/Private/Slate/Properties/SDMMaterialParameters.h
new file mode 100644
--- /dev/null
+++ b/Engine/Plugins/Experimental/DynamicMaterial/Source/DynamicMaterialEditor/Private/Slate/Properties/SDMMaterialParameters.h
@@ -0,0 +1,37 @@
+// Copyright Epic Games, Inc. All Rights Reserved.
+
+#pragma once
+
+#include "DMDefs.h"
+#include "Widgets/SCompoundWidget.h"
+
+class SBox;
+class SExpandableArea;
+class SWidget;
+class UDynamicMaterialModel;
+
+class SDMMaterialParameters : public SCompoundWidget
+{
+	SLATE_BEGIN_ARGS(SDMMaterialParameters) {}
+	SLATE_END_ARGS()
+
+public:
+	virtual ~SDMMaterialParameters() {}
+
+	void Construct(const FArguments& InArgs, TWeakObjectPtr<UDynamicMaterialModel> InModelWeak);
+
+	void RefreshWidgets();
+
+protected:
+	TWeakObjectPtr<UDynamicMaterialModel> ModelWeak;
+
+	TSharedPtr<SExpandableArea> ValuesArea;
+	TSharedPtr<SBox> AreaBody;
+
+	TSharedRef<SWidget> CreateValuesWidget();
+
+	FReply OnAddValueButtonClicked(EDMValueType Type);
+	FReply OnRemoveValueButtonClicked(int32 Index);
+	bool OnVerifyValueNameChanged(const FText& InNewName, FText& OutErrorText, int32 ValueIndex);
+	void OnAcceptValueNameChanged(const FText& InNewName, int32 ValueIndex);
+};
diff --git a/Engine/Plugins/Experimental/DynamicMaterial/Source/DynamicMaterialEditor/Private/Slate/Properties/SDMMaterialProperty.cpp b/Engine/Plugins/Experimental/DynamicMaterial/Source/DynamicMaterialEditor/Private/Slate/Properties/SDMMaterialProperty.cpp
new file mode 100644
--- /dev/null
+++ b/Engine/Plugins/Experimental/DynamicMaterial/Source/DynamicMaterialEditor/Private/Slate/Properties/SDMMaterialProperty.cpp
@@ -0,0 +1,67 @@
+// Copyright Epic Games, Inc. All Rights Reserved.
+
+#include "SDMMaterialProperty.h"
+#include "Components/DMMaterialProperty.h"
+#include "DetailLayoutBuilder.h"
+#include "Model/DynamicMaterialModel.h"
+#include "Model/DynamicMaterialModelEditorOnlyData.h"
+#include "Slate/SDMEditor.h"
+#include "Slate/SDMSlot.h"
+#include "Styling/StyleColors.h"
+#include "Widgets/Colors/SColorBlock.h"
+#include "Widgets/SCanvas.h"
+#include "Widgets/SOverlay.h"
+#include "Widgets/Text/STextBlock.h"
+
+void SDMMaterialProperty::Construct(const FArguments& InArgs, const TSharedRef<SDMSlot>& InMaterialSlotWidget, EDMMaterialPropertyType InProperty)
+{
+	MaterialSlotWidgetWeak = InMaterialSlotWidget;
+	Property = InProperty;
+
+	TSharedPtr<SDMEditor> MaterialEditor = InMaterialSlotWidget->GetEditorWidget();
+
+	if (ensure(MaterialEditor.IsValid()))
+	{
+		UDynamicMaterialModel* MaterialModel = MaterialEditor->GetMaterialModel();
+
+		if (ensure(IsValid(MaterialModel)))
+		{
+			UDynamicMaterialModelEditorOnlyData* ModelEditorOnlyData = UDynamicMaterialModelEditorOnlyData::Get(MaterialModel);
+
+			if (ensure(IsValid(ModelEditorOnlyData)))
+			{
+				TMap<EDMMaterialPropertyType, UDMMaterialProperty*> EditorProperties = ModelEditorOnlyData->GetMaterialProperties();
+				UDMMaterialProperty* const* MaterialProperty = EditorProperties.Find(Property);
+
+				if (ensure(MaterialProperty) && ensure(IsValid(*MaterialProperty)))
+				{
+					ChildSlot
+						[
+							SNew(SOverlay)
+								+ SOverlay::Slot()
+								[
+									SNew(SColorBlock)
+										.Color(FLinearColor::Black)
+										.CornerRadius(FVector4(6.f, 6.f, 6.f, 6.f))
+								]
+								+ SOverlay::Slot()
+								.Padding(1.f)
+								[
+									SNew(SColorBlock)
+										.Color(FStyleColors::Highlight.GetSpecifiedColor())
+										.CornerRadius(FVector4(5.f, 5.f, 5.f, 5.f))
+								]
+								+ SOverlay::Slot()
+								[
+									SNew(STextBlock)
+										.Text((*MaterialProperty)->GetDescription())
+										.Font(IDetailLayoutBuilder::GetDetailFont())
+										.ColorAndOpacity(FLinearColor::White)
+										.Margin(FMargin(8.f, 5.f, 8.f, 4.f))
+								]
+						];
+				}
+			}
+		}
+	}
+}
diff --git a/Engine/Plugins/Experimental/DynamicMaterial/Source/DynamicMaterialEditor/Private/Slate/Properties/SDMMaterialProperty.h b/Engine/Plugins/Experimental/DynamicMaterial/Source/DynamicMaterialEditor/Private/Slate/Properties/SDMMaterialProperty.h
new file mode 100644
--- /dev/null
+++ b/Engine/Plugins/Experimental/DynamicMaterial/Source/DynamicMaterialEditor/Private/Slate/Properties/SDMMaterialProperty.h
@@ -0,0 +1,22 @@
+// Copyright Epic Games, Inc. All Rights Reserved.
+
+#pragma once
+
+#include "DMEDefs.h"
+#include "Widgets/SCompoundWidget.h"
+
+class SDMSlot;
+class SWidget;
+
+class SDMMaterialProperty : public SCompoundWidget
+{
+	SLATE_BEGIN_ARGS(SDMMaterialProperty) {}
+	SLATE_END_ARGS()
+
+public:
+	void Construct(const FArguments& InArgs, const TSharedRef<SDMSlot>& InMaterialSlotWidget, EDMMaterialPropertyType InProperty);
+
+protected:
+	TWeakPtr<SDMSlot> MaterialSlotWidgetWeak;
+	EDMMaterialPropertyType Property = EDMMaterialPropertyType::None;
+};
diff --git a/Engine/Plugins/Experimental/DynamicMaterial/Source/DynamicMaterialEditor/Private/Slate/Properties/SDMPropertyEdit.cpp b/Engine/Plugins/Experimental/DynamicMaterial/Source/DynamicMaterialEditor/Private/Slate/Properties/SDMPropertyEdit.cpp
new file mode 100644
--- /dev/null
+++ b/Engine/Plugins/Experimental/DynamicMaterial/Source/DynamicMaterialEditor/Private/Slate/Properties/SDMPropertyEdit.cpp
@@ -0,0 +1,467 @@
+// Copyright Epic Games, Inc. All Rights Reserved.
+
+#include "Slate/Properties/SDMPropertyEdit.h"
+#include "Components/DMMaterialValue.h"
+#include "DMWorldSubsystem.h"
+#include "DetailLayoutBuilder.h"
+#include "DetailRowMenuContext.h"
+#include "DynamicMaterialEditorModule.h"
+#include "DynamicMaterialEditorSettings.h"
+#include "Engine/Engine.h"
+#include "Engine/World.h"
+#include "Framework/Application/SlateApplication.h"
+#include "IDetailKeyframeHandler.h"
+#include "PropertyCustomizationHelpers.h"
+#include "PropertyHandle.h"
+#include "ScopedTransaction.h"
+#include "Slate/SDMEditor.h"
+#include "Slate/SDMSlot.h"
+#include "Slate/SDMComponentEdit.h"
+#include "ToolMenuContext.h"
+#include "ToolMenus.h"
+#include "Widgets/Colors/SColorBlock.h"
+#include "Widgets/Colors/SColorPicker.h"
+#include "Widgets/Input/SCheckBox.h"
+#include "Widgets/Input/SSpinBox.h"
+#include "Widgets/Layout/SBox.h"
+#include "Widgets/SBoxPanel.h"
+#include "Widgets/SNullWidget.h"
+
+#define LOCTEXT_NAMESPACE "SDMPropertyEdit"
+
+void SDMPropertyEdit::Construct(const FArguments& InArgs)
+{
+	ComponentEditWidgetWeak = InArgs._ComponentEditWidget;
+	PropertyHandle = InArgs._PropertyHandle;
+	PropertyMaterialValue = InArgs._PropertyMaterialValue;
+
+	if (PropertyHandle && PropertyHandle->IsValidHandle())
+	{
+		Property = PropertyHandle->GetProperty();
+	}
+	else if (PropertyMaterialValue.IsValid())
+	{
+		Property = PropertyMaterialValue->GetClass()->FindPropertyByName("Value");
+	}
+ 
+	TSharedRef<SHorizontalBox> ComponentBox = SNew(SHorizontalBox);
+
+	if (ensure(Property))
+	{
+		for (int32 ComponentIndex = 0; ComponentIndex < InArgs._InputCount; ++ComponentIndex)
+		{
+			static const FMargin FullMargin(2.f, 0.f, 2.f, 0.f);
+			static const FMargin PartialMargin(0.f, 0.f, 2.f, 0.f);
+
+			const FMargin SlotPadding = ComponentIndex == 0 ? FullMargin : PartialMargin;
+			const FMargin ComponentPadding = InArgs._InputCount > 1 ? SlotPadding : FullMargin;
+			const float MaxWidth = GetMaxWidthForWidget(ComponentIndex);
+
+			if (MaxWidth <= KINDA_SMALL_NUMBER)
+			{
+				ComponentBox->AddSlot()
+					.FillWidth(1.0f)			
+					.Padding(ComponentPadding)
+					[
+						GetComponentWidget(ComponentIndex)
+					];			
+			}
+			else
+			{
+				ComponentBox->AddSlot()
+					.FillWidth(1.0f)
+					.Padding(ComponentPadding)
+					[
+						SNew(SBox)
+						.MinDesiredWidth(MaxWidth)
+						.MaxDesiredWidth(MaxWidth)
+						[
+							GetComponentWidget(ComponentIndex)
+						]
+					];
+			}
+		}
+	}
+
+	const FMargin ParentPadding = FMargin(0.0f, 1.0f);
+ 
+	ChildSlot
+	[
+		SNew(SBox)
+		.HAlign(HAlign_Left)
+		.VAlign(VAlign_Center)
+		.Padding(ParentPadding)
+		[
+			ComponentBox
+		]
+	];
+}
+
+UDMMaterialValue* SDMPropertyEdit::GetValue() const
+{
+	return PropertyMaterialValue.Get();
+}
+
+TSharedPtr<IPropertyHandle> SDMPropertyEdit::GetPropertyHandle() const
+{
+	return PropertyHandle;
+}
+
+FReply SDMPropertyEdit::CreateRightClickDetailsMenu(const FGeometry& InGeometry, const FPointerEvent& InPointerEvent,
+	TWeakPtr<SDMPropertyEdit> InPropertyEditorWeak)
+{
+	if (InPointerEvent.GetEffectingButton() != EKeys::RightMouseButton)
+	{
+		return FReply::Unhandled();
+	}
+
+	TSharedPtr<SDMPropertyEdit> PropertyEditor = InPropertyEditorWeak.Pin();
+
+	if (!PropertyEditor.IsValid() || !PropertyEditor->IsEnabled())
+	{
+		return FReply::Unhandled();
+	}
+
+	UToolMenus* Menus = UToolMenus::Get();
+	check(Menus);
+
+	static const FName DetailViewContextMenuName = UE::PropertyEditor::RowContextMenuName;
+
+	if (!Menus->IsMenuRegistered(DetailViewContextMenuName))
+	{
+		return FReply::Unhandled();
+	}
+
+	const TSharedPtr<IPropertyHandle> RowPropertyHandle = PropertyEditor->GetPropertyHandle();
+
+	if (!RowPropertyHandle.IsValid())
+	{
+		return FReply::Unhandled();
+	}
+
+	UDetailRowMenuContext* RowMenuContext = NewObject<UDetailRowMenuContext>();
+	RowMenuContext->PropertyHandles.Add(RowPropertyHandle);
+
+	const FToolMenuContext ToolMenuContext(RowMenuContext);
+
+	FSlateApplication::Get().PushMenu(
+		PropertyEditor.ToSharedRef(),
+		FWidgetPath(),
+		Menus->GenerateWidget(DetailViewContextMenuName, ToolMenuContext),
+		FSlateApplication::Get().GetCursorPos(),
+		FPopupTransitionEffect::ContextMenu
+	);
+
+	return FReply::Handled();
+}
+
+TSharedRef<SHorizontalBox> SDMPropertyEdit::AddWidgetLabel(const FText& InLabel, TSharedRef<SWidget> ValueWidget, TSharedPtr<SHorizontalBox> AddOnToWidget)
+{
+	static const FMargin FirstSlotPadding = FMargin(2.0f);
+	static const FMargin OtherSlotPadding = FMargin(10.0f, 2.0f, 2.0f, 2.0f);
+	static const FSlateFontInfo DetailsFontInfo = IDetailLayoutBuilder::GetDetailFont();
+
+	TSharedRef<SHorizontalBox> ReturnWidget = AddOnToWidget.IsValid() ? AddOnToWidget.ToSharedRef() : SNew(SHorizontalBox);
+
+	ReturnWidget->AddSlot()
+		.AutoWidth()
+		.Padding(AddOnToWidget.IsValid() ? OtherSlotPadding : FirstSlotPadding)
+		[
+			SNew(SBox)
+			.VAlign(EVerticalAlignment::VAlign_Center)
+			[
+				SNew(STextBlock)
+				.Text(InLabel)
+				.Font(DetailsFontInfo)
+			]
+		];
+
+	ReturnWidget->AddSlot()
+		.FillWidth(1.0f)
+		.Padding(OtherSlotPadding)
+		[
+			ValueWidget
+		];
+
+	return ReturnWidget;
+}
+
+void SDMPropertyEdit::CreateKey(TWeakObjectPtr<UDMMaterialValue> InValueWeak)
+{
+	if (!InValueWeak.IsValid())
+	{
+		return;
+	}
+ 
+	TSharedPtr<IPropertyHandle> PropertyHandle = InValueWeak->GetPropertyHandle();
+	if (!PropertyHandle.IsValid())
+	{
+		return;
+	}
+ 
+	const UWorld* World = InValueWeak->GetWorld();
+	if (!World)
+	{
+		return;
+	}
+ 
+	const UDMWorldSubsystem* WorldSubsystem = World->GetSubsystem<UDMWorldSubsystem>();
+	if (!WorldSubsystem)
+	{
+		return;
+	}
+ 
+	TSharedPtr<IDetailKeyframeHandler> KeyframeHandler = WorldSubsystem->GetKeyframeHandler();
+	if (!KeyframeHandler.IsValid())
+	{
+		return;
+	}
+ 
+	if (!KeyframeHandler->IsPropertyKeyable(InValueWeak->GetClass(), *PropertyHandle))
+	{
+		return;
+	}
+ 
+	KeyframeHandler->OnKeyPropertyClicked(*PropertyHandle);
+}
+ 
+TSharedRef<SWidget> SDMPropertyEdit::GetComponentWidget(int32 InIndex)
+{
+	return SNullWidget::NullWidget;
+}
+
+float SDMPropertyEdit::GetMaxWidthForWidget(int32 InIndex) const
+{
+	return 0.f;
+}
+ 
+TSharedRef<SWidget> SDMPropertyEdit::CreateCheckbox(const TAttribute<ECheckBoxState>& InValueAttr, const FOnCheckStateChanged& InChangeFunc)
+{
+	return SNew(SCheckBox)
+		.IsChecked(InValueAttr)
+		.OnCheckStateChanged(InChangeFunc);
+}
+ 
+void SDMPropertyEdit::OnSpinBoxStartScrubbing(FText InTransactionDescription)
+{
+	StartTransaction(InTransactionDescription);
+}
+ 
+void SDMPropertyEdit::OnSpinBoxEndScrubbing(float InValue)
+{
+	EndTransaction();
+}
+ 
+void SDMPropertyEdit::OnSpinBoxValueChange(float InValue, const SSpinBox<float>::FOnValueChanged InChangeFunc)
+{
+	if (ScrubbingTransaction.IsValid())
+	{
+		InChangeFunc.ExecuteIfBound(InValue);
+	}
+}
+ 
+void SDMPropertyEdit::OnSpinBoxCommit(float InValue, ETextCommit::Type InCommitType, const SSpinBox<float>::FOnValueChanged InChangeFunc,
+	FText InTransactionDescription)
+{
+	if (!ScrubbingTransaction.IsValid())
+	{
+		StartTransaction(InTransactionDescription);
+		InChangeFunc.ExecuteIfBound(InValue);
+		EndTransaction();
+	}
+}
+ 
+TSharedRef<SSpinBox<float>> SDMPropertyEdit::CreateSpinBox(const TAttribute<float>& InValue, const SSpinBox<float>::FOnValueChanged& InOnValueChanged,
+	const FText& InTransactionDescription, const FFloatInterval* InValueRange)
+{
+	return
+		SNew(SSpinBox<float>)
+		.ClearKeyboardFocusOnCommit(false)
+		.Justification(ETextJustify::Right)
+		.Font(IDetailLayoutBuilder::GetDetailFont())
+		.MaxFractionalDigits(3)
+		.MinDesiredWidth(50.f)
+		.OnBeginSliderMovement(this, &SDMPropertyEdit::OnSpinBoxStartScrubbing, InTransactionDescription)
+		.OnEndSliderMovement(this, &SDMPropertyEdit::OnSpinBoxEndScrubbing)
+		.OnValueChanged(this, &SDMPropertyEdit::OnSpinBoxValueChange, InOnValueChanged)
+		.OnValueCommitted(this, &SDMPropertyEdit::OnSpinBoxCommit, InOnValueChanged, InTransactionDescription)
+		.Value(InValue)
+		.MinValue(InValueRange ? TOptional<float>(InValueRange->Min) : TOptional<float>())
+		.MinSliderValue(InValueRange ? TOptional<float>(InValueRange->Min) : TOptional<float>())
+		.MaxValue(InValueRange ? TOptional<float>(InValueRange->Max) : TOptional<float>())
+		.MaxSliderValue(InValueRange ? TOptional<float>(InValueRange->Max) : TOptional<float>());
+}
+ 
+void SDMPropertyEdit::OnColorPickedCommitted(FLinearColor InNewColor, const FOnLinearColorValueChanged InChangeFunc)
+{
+	InChangeFunc.ExecuteIfBound(InNewColor);
+}
+ 
+void SDMPropertyEdit::OnColorPickerCancelled(FLinearColor InOldColor, const FOnLinearColorValueChanged InChangeFunc)
+{
+	InChangeFunc.ExecuteIfBound(InOldColor);
+}
+ 
+void SDMPropertyEdit::OnColorPickerClosed(const TSharedRef<SWindow>& InClosedWindow)
+{
+	EndTransaction();
+}
+ 
+TSharedRef<SWidget> SDMPropertyEdit::CreateColorPicker(bool bInUseAlpha, const TAttribute<FLinearColor>& InValueAttr,
+	const FOnLinearColorValueChanged& InChangeFunc, const FText& InPickerTransationDescription, 
+	const FFloatInterval* InValueRange)
+{
+	// InValueRange is not currently supported
+	return 
+		SNew(SBox)
+		.WidthOverride(36.f)
+		.HeightOverride(18.f)
+		.HAlign(HAlign_Left)
+		.VAlign(VAlign_Center)
+		[
+			SNew(SColorBlock)
+			.Color(InValueAttr)
+			.Size(FVector2D(36.f, 18.f))
+			.ShowBackgroundForAlpha(true)
+			.OnMouseButtonDown(this, &SDMPropertyEdit::OnClickColorBlock, bInUseAlpha, InValueAttr, InChangeFunc, InPickerTransationDescription)
+		];
+}
+
+TSharedRef<SWidget> SDMPropertyEdit::CreateEnum(const FOnGetPropertyComboBoxValue& InValueString, const FOnGetPropertyComboBoxStrings& InGetStrings,
+	const FOnPropertyComboBoxValueSelected& InValueSet)
+{
+	FPropertyComboBoxArgs Args;
+	Args.PropertyHandle = nullptr;
+	Args.OnGetValue = InValueString;
+	Args.OnGetStrings = InGetStrings;
+	Args.OnValueSelected = InValueSet;
+
+	return PropertyCustomizationHelpers::MakePropertyComboBox(Args);
+}
+
+TSharedRef<SWidget> SDMPropertyEdit::CreateAssetPicker(UClass* InAllowedClass, const TAttribute<FString>& InPathAttr, 
+	const FOnSetObject& InChangeFunc)
+{
+	return 
+		SNew(SObjectPropertyEntryBox)
+		.AllowClear(true)
+		.AllowedClass(InAllowedClass)
+		.DisplayBrowse(true)
+		.DisplayThumbnail(true)
+		.DisplayCompactSize(true)
+		.DisplayUseSelected(true)
+		.ThumbnailPool(SDMEditor::GetThumbnailPool())
+		.EnableContentPicker(true)
+		.ObjectPath(InPathAttr)
+		.OnObjectChanged(InChangeFunc)
+		.OnShouldSetAsset(FOnShouldSetAsset::CreateLambda([](const FAssetData& InAssetData) {return false;}));
+}
+ 
+FReply SDMPropertyEdit::OnClickColorBlock(const FGeometry& MyGeometry, const FPointerEvent& MouseEvent, bool bInUseAlpha,
+	TAttribute<FLinearColor> InValueAttr, FOnLinearColorValueChanged InChangeFunc, FText InPickerTransactionDescription)
+{
+	if (MouseEvent.GetEffectingButton() != EKeys::LeftMouseButton)
+	{
+		return FReply::Unhandled();
+	}
+ 
+	FColorPickerArgs PickerArgs;
+	{
+		PickerArgs.bOnlyRefreshOnOk = true;
+		PickerArgs.DisplayGamma = TAttribute<float>::Create(TAttribute<float>::FGetter::CreateUObject(GEngine, &UEngine::GetDisplayGamma));
+		PickerArgs.OnColorCommitted = FOnLinearColorValueChanged::CreateSP(this, &SDMPropertyEdit::OnColorPickedCommitted, InChangeFunc);
+		PickerArgs.OnColorPickerCancelled = FOnColorPickerCancelled::CreateSP(this, &SDMPropertyEdit::OnColorPickerCancelled, InChangeFunc);
+		PickerArgs.OnColorPickerWindowClosed = FOnWindowClosed::CreateSP(this, &SDMPropertyEdit::OnColorPickerClosed);
+		PickerArgs.bUseAlpha = bInUseAlpha;
+		PickerArgs.InitialColor = InValueAttr.Get(FLinearColor::White);
+		PickerArgs.bOnlyRefreshOnMouseUp = false;
+		PickerArgs.bOnlyRefreshOnOk = false;
+	}
+ 
+	StartTransaction(InPickerTransactionDescription);
+	OpenColorPicker(PickerArgs);
+ 
+	return FReply::Handled();
+}
+ 
+ 
+void SDMPropertyEdit::StartTransaction(FText InDescription)
+{
+	if (ScrubbingTransaction.IsValid())
+	{
+		return;
+	}
+ 
+	UObject* Object = nullptr;
+
+	if (PropertyMaterialValue.IsValid())
+	{
+		Object = PropertyMaterialValue.Get();
+	}
+	else if (PropertyHandle.IsValid() && PropertyHandle->IsValidHandle())
+	{
+		TArray<UObject*> Outers;
+		PropertyHandle->GetOuterObjects(Outers);
+
+		if (Outers.IsEmpty() == false)
+		{
+			Object = Outers[0];
+		}
+	}
+
+	ScrubbingTransaction = MakeShared<FScopedTransaction>(InDescription);
+
+	if (IsValid(Object))
+	{
+		Object->Modify();
+	}
+}
+ 
+void SDMPropertyEdit::EndTransaction()
+{
+	ScrubbingTransaction.Reset();
+}
+ 
+TSharedRef<SWidget> SDMPropertyEdit::MakeCopyPasteMenu(const TSharedPtr<IPropertyHandle>& InPropertyHandle) const
+{
+	if (!InPropertyHandle.IsValid() || !InPropertyHandle->IsValidHandle())
+	{
+		return SNullWidget::NullWidget;
+	}
+
+	const FName MenuName = "CopyPasteMenu";
+
+	UToolMenu* NewMenu = UToolMenus::Get()->RegisterMenu(MenuName, NAME_None, EMultiBoxType::Menu, false);
+	if (!IsValid(NewMenu))
+	{
+		return SNullWidget::NullWidget;
+	}
+
+	NewMenu->bToolBarForceSmallIcons = true;
+	NewMenu->bShouldCloseWindowAfterMenuSelection = true;
+	NewMenu->bCloseSelfOnly = true;
+
+	FToolMenuSection& NewSection = NewMenu->AddDynamicSection(NAME_None, 
+		FNewToolMenuDelegate::CreateLambda([InPropertyHandle](UToolMenu* InMenu)
+		{
+			FUIAction CopyAction;
+			FUIAction PasteAction;
+			InPropertyHandle->CreateDefaultPropertyCopyPasteActions(CopyAction, PasteAction);
+
+			FToolMenuSection& NewSection = InMenu->AddSection("Edit", LOCTEXT("EditSection", "Edit"));
+			NewSection.AddMenuEntry(NAME_None,
+				LOCTEXT("CopyProperty", "Copy"),
+				LOCTEXT("CopyPropertyToolTip", "Copy this property value"),
+				FSlateIcon(FCoreStyle::Get().GetStyleSetName(), "GenericCommands.Copy"),
+				CopyAction);
+			NewSection.AddMenuEntry(NAME_None,
+				LOCTEXT("PasteProperty", "Paste"),
+				LOCTEXT("PastePropertyToolTip", "Paste the copied value here"),
+				FSlateIcon(FCoreStyle::Get().GetStyleSetName(), "GenericCommands.Paste"),
+				PasteAction);
+		}));
+
+	return UToolMenus::Get()->GenerateWidget(MenuName, FToolMenuContext());
+}
+ 
+#undef LOCTEXT_NAMESPACE
diff --git a/Engine/Plugins/Experimental/DynamicMaterial/Source/DynamicMaterialEditor/Private/Slate/Properties/SDMPropertyEdit.h b/Engine/Plugins/Experimental/DynamicMaterial/Source/DynamicMaterialEditor/Private/Slate/Properties/SDMPropertyEdit.h
new file mode 100644
--- /dev/null
+++ b/Engine/Plugins/Experimental/DynamicMaterial/Source/DynamicMaterialEditor/Private/Slate/Properties/SDMPropertyEdit.h
@@ -0,0 +1,96 @@
+// Copyright Epic Games, Inc. All Rights Reserved.
+
+#pragma once
+
+#include "DynamicMaterialEditorModule.h"
+#include "PropertyCustomizationHelpers.h"
+#include "PropertyHandle.h"
+#include "Widgets/Input/SCheckBox.h"
+#include "Widgets/Input/SSpinBox.h"
+#include "Widgets/SCompoundWidget.h"
+
+class FScopedTransaction;
+class SDMComponentEdit;
+class SHorizontalBox;
+class UDMMaterialValue;
+enum class ECheckBoxState : uint8;
+
+/**
+ * Material Designer Property Edit
+ * 
+ * Editor widget for modifying the value of a material property.
+ * Can accept a property in the form of a UDMMaterialValue or TSharedPtr<IPropertyHandle>.
+*/
+class SDMPropertyEdit : public SCompoundWidget
+{
+public:
+	static FReply CreateRightClickDetailsMenu(const FGeometry& InGeometry, const FPointerEvent& InPointerEvent,
+		TWeakPtr<SDMPropertyEdit> InPropertyEditorWeak);
+
+	SLATE_BEGIN_ARGS(SDMPropertyEdit)
+		: _InputCount(1)
+		{}
+		SLATE_ARGUMENT(int32, InputCount)
+		SLATE_ARGUMENT(TSharedPtr<SDMComponentEdit>, ComponentEditWidget)
+		SLATE_ARGUMENT(TSharedPtr<IPropertyHandle>, PropertyHandle)
+		SLATE_ARGUMENT(TWeakObjectPtr<UDMMaterialValue>, PropertyMaterialValue)
+	SLATE_END_ARGS()
+
+	static TSharedRef<SHorizontalBox> AddWidgetLabel(const FText& InLabel, TSharedRef<SWidget> ValueWidget, TSharedPtr<SHorizontalBox> AddOnToWidget = nullptr);
+
+	virtual ~SDMPropertyEdit() {}
+
+	void Construct(const FArguments& InArgs);
+
+	UDMMaterialValue* GetValue() const;
+
+	TSharedPtr<IPropertyHandle> GetPropertyHandle() const;
+
+protected:
+	TWeakPtr<SDMComponentEdit> ComponentEditWidgetWeak;
+	TSharedPtr<IPropertyHandle> PropertyHandle;
+	TWeakObjectPtr<UDMMaterialValue> PropertyMaterialValue;
+
+	FProperty* Property = nullptr;
+	TSharedPtr<FScopedTransaction> ScrubbingTransaction;
+
+	static void CreateKey(TWeakObjectPtr<UDMMaterialValue> InValueWeak);
+
+	virtual TSharedRef<SWidget> GetComponentWidget(int32 InIndex);
+	virtual float GetMaxWidthForWidget(int32 InIndex) const;
+ 
+	TSharedRef<SWidget> CreateCheckbox(const TAttribute<ECheckBoxState>& InValueAttr, const FOnCheckStateChanged& InChangeFunc);
+ 
+	TSharedRef<SWidget> CreateAssetPicker(UClass* InAllowedClass, const TAttribute<FString>& InPathAttr, const FOnSetObject& InChangeFunc);
+ 
+	void OnSpinBoxStartScrubbing(FText InTransactionDescription);
+	void OnSpinBoxEndScrubbing(float InValue);
+
+	void OnSpinBoxValueChange(float InValue, const SSpinBox<float>::FOnValueChanged InChangeFunc);
+
+	void OnSpinBoxCommit(float InValue, ETextCommit::Type InCommitType, const SSpinBox<float>::FOnValueChanged InChangeFunc,
+		FText InTransactionDescription);
+ 
+	TSharedRef<SSpinBox<float>> CreateSpinBox(const TAttribute<float>& InValue, const SSpinBox<float>::FOnValueChanged& InOnValueChanged,
+		const FText& InTransactionDescription, const FFloatInterval* InValueRange = nullptr);
+ 
+	void OnColorPickedCommitted(FLinearColor InNewColor, const FOnLinearColorValueChanged InChangeFunc);
+	void OnColorPickerCancelled(FLinearColor InOldColor, const FOnLinearColorValueChanged InChangeFunc);
+	void OnColorPickerClosed(const TSharedRef<SWindow>& InClosedWindow);
+ 
+	/** @param CommitFunc must call EndTransaction() */
+	TSharedRef<SWidget> CreateColorPicker(bool bInUseAlpha, const TAttribute<FLinearColor>& InValueAttr,
+		const FOnLinearColorValueChanged& InChangeFunc, const FText& InPickerTransationDescription,
+		const FFloatInterval* InValueRange = nullptr);
+
+	TSharedRef<SWidget> CreateEnum(const FOnGetPropertyComboBoxValue& InValueString, const FOnGetPropertyComboBoxStrings& InGetStrings,
+		const FOnPropertyComboBoxValueSelected& InValueSet);
+
+	FReply OnClickColorBlock(const FGeometry& MyGeometry, const FPointerEvent& MouseEvent, bool bInseAlpha,
+		TAttribute<FLinearColor> InValueAttr, FOnLinearColorValueChanged InChangeFunc, FText InPickerTransactionDescription);
+ 
+	virtual void StartTransaction(FText InDescription);
+	virtual void EndTransaction();
+ 
+	TSharedRef<SWidget> MakeCopyPasteMenu(const TSharedPtr<IPropertyHandle>& InPropertyHandle) const;
+};
diff --git a/Engine/Plugins/Experimental/DynamicMaterial/Source/DynamicMaterialEditor/Private/Slate/Properties/SDMSourceBlendType.cpp b/Engine/Plugins/Experimental/DynamicMaterial/Source/DynamicMaterialEditor/Private/Slate/Properties/SDMSourceBlendType.cpp
new file mode 100644
--- /dev/null
+++ b/Engine/Plugins/Experimental/DynamicMaterial/Source/DynamicMaterialEditor/Private/Slate/Properties/SDMSourceBlendType.cpp
@@ -0,0 +1,369 @@
+// Copyright Epic Games, Inc. All Rights Reserved.
+
+#include "Slate/Properties/SDMSourceBlendType.h"
+#include "Components/MaterialStageBlends/DMMSBAdd.h"
+#include "Components/MaterialStageBlends/DMMSBColor.h"
+#include "Components/MaterialStageBlends/DMMSBColorBurn.h"
+#include "Components/MaterialStageBlends/DMMSBColorDodge.h"
+#include "Components/MaterialStageBlends/DMMSBDarken.h"
+#include "Components/MaterialStageBlends/DMMSBDarkenColor.h"
+#include "Components/MaterialStageBlends/DMMSBDifference.h"
+#include "Components/MaterialStageBlends/DMMSBDivide.h"
+#include "Components/MaterialStageBlends/DMMSBExclusion.h"
+#include "Components/MaterialStageBlends/DMMSBHardLight.h"
+#include "Components/MaterialStageBlends/DMMSBHardMix.h"
+#include "Components/MaterialStageBlends/DMMSBHue.h"
+#include "Components/MaterialStageBlends/DMMSBLighten.h"
+#include "Components/MaterialStageBlends/DMMSBLightenColor.h"
+#include "Components/MaterialStageBlends/DMMSBLinearBurn.h"
+#include "Components/MaterialStageBlends/DMMSBLinearDodge.h"
+#include "Components/MaterialStageBlends/DMMSBLinearLight.h"
+#include "Components/MaterialStageBlends/DMMSBLuminosity.h"
+#include "Components/MaterialStageBlends/DMMSBMultiply.h"
+#include "Components/MaterialStageBlends/DMMSBNormal.h"
+#include "Components/MaterialStageBlends/DMMSBOverlay.h"
+#include "Components/MaterialStageBlends/DMMSBPinLight.h"
+#include "Components/MaterialStageBlends/DMMSBSaturation.h"
+#include "Components/MaterialStageBlends/DMMSBScreen.h"
+#include "Components/MaterialStageBlends/DMMSBSoftLight.h"
+#include "Components/MaterialStageBlends/DMMSBSubtract.h"
+#include "Components/MaterialStageBlends/DMMSBVividLight.h"
+#include "DMEDefs.h"
+#include "DMPrivate.h"
+#include "DynamicMaterialEditorStyle.h"
+#include "ToolMenu.h"
+#include "ToolMenuDelegates.h"
+#include "ToolMenus.h"
+#include "Widgets/Input/SComboButton.h"
+#include "Widgets/SNullWidget.h"
+#include "Widgets/Text/STextBlock.h"
+
+#define LOCTEXT_NAMESPACE "SDMSourceBlendType"
+
+namespace UE::DynamicMaterialEditor::Private
+{
+	static const FName SourceBlendMenuName = "SourceBlendMenu";
+
+	struct FDMBlendCategory
+	{
+		FText Name;
+		TArray<UClass*> Classes;
+	};
+
+	const TArray<FDMBlendCategory>& SupportedBlendCategories()
+	{
+		static TArray<FDMBlendCategory> OutBlendCategories;
+
+		if (OutBlendCategories.IsEmpty() == false)
+		{
+			return OutBlendCategories;
+		}
+
+		OutBlendCategories.Append(
+			{
+				{
+					LOCTEXT("BlendNormal", "Normal Blends"),
+					{
+						UDMMaterialStageBlendNormal::StaticClass()
+					}
+				},
+				{
+					LOCTEXT("BlendDarken", "Darken Blends"),
+					{
+						UDMMaterialStageBlendDarken::StaticClass(),
+						UDMMaterialStageBlendMultiply::StaticClass(),
+						UDMMaterialStageBlendColorBurn::StaticClass(),
+						UDMMaterialStageBlendLinearBurn::StaticClass(),
+						UDMMaterialStageBlendDarkenColor::StaticClass()
+					}
+				},
+				{
+					LOCTEXT("BlendLighten", "Lighten Blends"),
+					{
+						UDMMaterialStageBlendAdd::StaticClass(),
+						UDMMaterialStageBlendLighten::StaticClass(),
+						UDMMaterialStageBlendScreen::StaticClass(),
+						UDMMaterialStageBlendColorDodge::StaticClass(),
+						UDMMaterialStageBlendLinearDodge::StaticClass(),
+						UDMMaterialStageBlendLightenColor::StaticClass()
+					}
+				},
+				{
+					LOCTEXT("BlendContrast", "Contrast Blends"),
+					{
+						UDMMaterialStageBlendOverlay::StaticClass(),
+						UDMMaterialStageBlendSoftLight::StaticClass(),
+						UDMMaterialStageBlendHardLight::StaticClass(),
+						UDMMaterialStageBlendVividLight::StaticClass(),
+						UDMMaterialStageBlendLinearLight::StaticClass(),
+						UDMMaterialStageBlendPinLight::StaticClass(),
+						UDMMaterialStageBlendHardMix::StaticClass()
+					}
+				},
+				{
+					LOCTEXT("BlendInversion", "Inversion Blends"),
+					{
+						UDMMaterialStageBlendDifference::StaticClass(),
+						UDMMaterialStageBlendExclusion::StaticClass(),
+						UDMMaterialStageBlendSubtract::StaticClass(),
+						UDMMaterialStageBlendDivide::StaticClass()
+					}
+				},
+				{
+					LOCTEXT("BlendHSL", "HSL Blends"),
+					{
+						UDMMaterialStageBlendColor::StaticClass(),
+						UDMMaterialStageBlendHue::StaticClass(),
+						UDMMaterialStageBlendSaturation::StaticClass(),
+						UDMMaterialStageBlendLuminosity::StaticClass()
+					}
+				}
+			});
+		return OutBlendCategories;
+	};
+}
+
+TArray<TStrongObjectPtr<UClass>> SDMSourceBlendType::SupportedBlendClasses = {};
+TMap<FName, FDMBlendNameClass> SDMSourceBlendType::BlendMap = {};
+
+void SDMSourceBlendType::Construct(const FArguments& InArgs)
+{
+	SelectedItem = InArgs._SelectedItem;
+	OnSelectedItemChanged = InArgs._OnSelectedItemChanged;
+
+	EnsureBlendMap();
+	EnsureMenuRegistered();
+
+	FName InitialSourceBlendTypeName = NAME_None;
+
+	if (TSubclassOf<UDMMaterialStageBlend> Blend = SelectedItem.Get())
+	{
+		InitialSourceBlendTypeName = SelectedItem.Get()->GetFName();
+	}
+
+	const FComboBoxStyle& ComboBoxStyle = FAppStyle::Get().GetWidgetStyle<FComboBoxStyle>("ComboBox");
+	const FComboButtonStyle& ComboButtonStyle = ComboBoxStyle.ComboButtonStyle;
+	const FButtonStyle& ButtonStyle = ComboButtonStyle.ButtonStyle;
+	
+	ChildSlot
+	[
+		SNew(SComboButton)
+		.ComboButtonStyle(&ComboButtonStyle)
+		.ButtonStyle(&ButtonStyle)
+		.ContentPadding(ComboBoxStyle.ContentPadding)
+		.IsFocusable(true)
+		.ForegroundColor(FSlateColor::UseStyle())
+		.ToolTipText(LOCTEXT("SourceBlendMode", "Source Blend Mode"))
+		.OnGetMenuContent(this, &SDMSourceBlendType::MakeSourceBlendMenuWidget)
+		.ButtonContent()
+		[
+			SNew(STextBlock)
+			.Text(this, &SDMSourceBlendType::GetSelectedItemText)
+		]
+	];
+}
+
+void SDMSourceBlendType::EnsureBlendMap()
+{
+	if (SupportedBlendClasses.IsEmpty() == false)
+	{
+		return;
+	}
+
+	SupportedBlendClasses = UDMMaterialStageBlend::GetAvailableBlends();
+
+	for (const TStrongObjectPtr<UClass>& BlendClass : SupportedBlendClasses)
+	{
+		UDMMaterialStageBlend* StageBlendCDO = Cast<UDMMaterialStageBlend>(BlendClass->GetDefaultObject());
+
+		if (ensure(IsValid(StageBlendCDO)))
+		{
+			const FName BlendClassName = BlendClass->GetFName();
+			const FText BlendClassText = StageBlendCDO->GetDescription();
+			const TSubclassOf<UDMMaterialStageBlend> BlendClassObject = TSubclassOf<UDMMaterialStageBlend>(BlendClass.Get());
+
+			if (BlendClassName != NAME_None && !BlendClassText.IsEmpty())
+			{
+				BlendMap.Emplace(BlendClassName, {BlendClassText, BlendClassObject});
+			}
+		}
+	}
+}
+
+void SDMSourceBlendType::EnsureMenuRegistered()
+{
+	using namespace UE::DynamicMaterialEditor::Private;
+
+	if (UToolMenus::Get()->IsMenuRegistered(SourceBlendMenuName))
+	{
+		return;
+	}
+
+	UToolMenu* NewMenu = UToolMenus::Get()->RegisterMenu(SourceBlendMenuName, NAME_None, EMultiBoxType::Menu, false);
+	if (!IsValid(NewMenu))
+	{
+		return;
+	}
+
+	NewMenu->bToolBarForceSmallIcons = true;
+	NewMenu->bShouldCloseWindowAfterMenuSelection = true;
+	NewMenu->bCloseSelfOnly = true;
+
+	NewMenu->AddDynamicSection("PopulateToolBar", FNewToolMenuDelegate::CreateStatic(&SDMSourceBlendType::MakeSourceBlendMenu));
+}
+
+TSharedRef<SWidget> SDMSourceBlendType::OnGenerateWidget(const FName InItem)
+{
+	EnsureBlendMap();
+
+	if (!InItem.IsValid() || !BlendMap.Contains(InItem))
+	{
+		return SNullWidget::NullWidget;
+	}
+
+	return 
+		SNew(STextBlock)
+		.Text(BlendMap[InItem].BlendName);
+}
+
+void SDMSourceBlendType::OnSelectionChanged(const FName InNewItem, const ESelectInfo::Type InSelectInfoType)
+{
+	if (!BlendMap.Contains(InNewItem))
+	{
+		return;
+	}
+
+	TSubclassOf<UDMMaterialStageBlend> SelectedClass = BlendMap[InNewItem].BlendClass;
+
+	if (!SelectedClass.Get())
+	{
+		return;
+	}
+
+	OnSelectedItemChanged.ExecuteIfBound(SelectedClass);
+}
+
+FText SDMSourceBlendType::GetSelectedItemText() const
+{
+	if (TSubclassOf<UDMMaterialStageBlend> Blend = SelectedItem.Get())
+	{
+		const FName SelectedName = Blend->GetFName();
+
+		if (BlendMap.Contains(SelectedName))
+		{
+			return BlendMap[SelectedName].BlendName;
+		}
+	}
+
+	return FText::GetEmpty();
+
+}
+
+TSharedRef<SWidget> SDMSourceBlendType::MakeSourceBlendMenuWidget()
+{
+	using namespace UE::DynamicMaterialEditor::Private;
+
+	UDMSourceBlendTypeContextObject* Context = NewObject<UDMSourceBlendTypeContextObject>();
+	Context->SetBlendTypeWidget(SharedThis(this));
+
+	return UToolMenus::Get()->GenerateWidget(SourceBlendMenuName, FToolMenuContext(Context));
+}
+
+void SDMSourceBlendType::OnBlendTypeSelected(UClass* InBlendClass)
+{
+	OnSelectedItemChanged.ExecuteIfBound(InBlendClass);
+}
+
+bool SDMSourceBlendType::CanSelectBlendType(UClass* InBlendClass)
+{
+	return SelectedItem.Get() != InBlendClass;
+}
+
+bool SDMSourceBlendType::InBlendTypeSelected(UClass* InBlendClass)
+{
+	return SelectedItem.Get() == InBlendClass;
+}
+
+void SDMSourceBlendType::MakeSourceBlendMenu(UToolMenu* InToolMenu)
+{
+	using namespace UE::DynamicMaterialEditor::Private;
+
+	static const TArray<FDMBlendCategory> BlendCategories = SupportedBlendCategories();
+
+	if (SupportedBlendClasses.IsEmpty())
+	{
+		return;
+	}
+
+	UDMSourceBlendTypeContextObject* ContextObject = InToolMenu->FindContext<UDMSourceBlendTypeContextObject>();
+
+	if (!ContextObject)
+	{
+		return;
+	}
+
+	TSharedPtr<SDMSourceBlendType> SourceBlendWidget = ContextObject->GetBlendTypeWidget();
+
+	if (!SourceBlendWidget.IsValid())
+	{
+		return;
+	}
+
+	auto AddClasstoSection = [&SourceBlendWidget](FToolMenuSection& InSection, const FText& InName, UClass* InBlendClass)
+	{
+		InSection.AddMenuEntry(
+			NAME_None,
+			InName,
+			LOCTEXT("ChangeSourceBlendTooltip", "Change the source of this stage to a Material Blend."),
+			FSlateIcon(),
+			FUIAction(
+				FExecuteAction::CreateSP(SourceBlendWidget.ToSharedRef(), &SDMSourceBlendType::OnBlendTypeSelected, InBlendClass),
+				FCanExecuteAction::CreateSP(SourceBlendWidget.ToSharedRef(), &SDMSourceBlendType::CanSelectBlendType, InBlendClass),
+				FIsActionChecked::CreateSP(SourceBlendWidget.ToSharedRef(), &SDMSourceBlendType::InBlendTypeSelected, InBlendClass),
+				EUIActionRepeatMode::RepeatDisabled
+			)
+		);
+	};
+
+	for (const FDMBlendCategory& BlendCategory : BlendCategories)
+	{
+		FToolMenuSection& BlendCategorySection = InToolMenu->AddSection(FName(BlendCategory.Name.ToString()), BlendCategory.Name);
+
+		for (UClass* BlendCategoryClass : BlendCategory.Classes)
+		{
+			UDMMaterialStageBlend* BlendCDO = Cast<UDMMaterialStageBlend>(BlendCategoryClass->GetDefaultObject());
+
+			AddClasstoSection(BlendCategorySection, BlendCDO->GetDescription(), BlendCategoryClass);
+		}
+	}
+
+	FToolMenuSection& UncategorizedSection = InToolMenu->AddSection(TEXT("OtherBlends"), LOCTEXT("OtherBlends", "Other Blends"));
+
+	for (TStrongObjectPtr<UClass> BlendClass : SupportedBlendClasses)
+	{
+		UDMMaterialStageBlend* BlendCDO = Cast<UDMMaterialStageBlend>(BlendClass->GetDefaultObject());
+
+		bool bInBlendCategories = false;
+
+		for (const FDMBlendCategory& BlendCategory : BlendCategories)
+		{
+			for (UClass* BlendCategoryClass : BlendCategory.Classes)
+			{
+				if (BlendClass.Get() == BlendCategoryClass)
+				{
+					bInBlendCategories = true;
+					break;
+				}
+			}
+		}
+
+		if (bInBlendCategories)
+		{
+			continue;
+		}
+
+		AddClasstoSection(UncategorizedSection, BlendCDO->GetDescription(), BlendClass.Get());
+	}
+}
+
+#undef LOCTEXT_NAMESPACE
diff --git a/Engine/Plugins/Experimental/DynamicMaterial/Source/DynamicMaterialEditor/Private/Slate/Properties/SDMSourceBlendType.h b/Engine/Plugins/Experimental/DynamicMaterial/Source/DynamicMaterialEditor/Private/Slate/Properties/SDMSourceBlendType.h
new file mode 100644
--- /dev/null
+++ b/Engine/Plugins/Experimental/DynamicMaterial/Source/DynamicMaterialEditor/Private/Slate/Properties/SDMSourceBlendType.h
@@ -0,0 +1,76 @@
+// Copyright Epic Games, Inc. All Rights Reserved.
+
+#pragma once
+
+#include "Components/MaterialStageBlends/DMMSBNormal.h"
+#include "Templates/SharedPointer.h"
+#include "UObject/Object.h"
+#include "Widgets/SCompoundWidget.h"
+#include "SDMSourceBlendType.generated.h"
+
+class SDMSourceBlendType;
+class UDMMaterialStageBlend;
+class UToolMenu;
+
+UCLASS(MinimalAPI)
+class UDMSourceBlendTypeContextObject : public UObject
+{
+	GENERATED_BODY()
+	
+public:
+	UDMSourceBlendTypeContextObject() = default;
+	
+	TSharedPtr<SDMSourceBlendType> GetBlendTypeWidget() const { return BlendTypeWidgetWeak.Pin(); }
+	void SetBlendTypeWidget(const TSharedPtr<SDMSourceBlendType>& InBlendTypeWidget) { BlendTypeWidgetWeak = InBlendTypeWidget; }
+	
+private:
+	TWeakPtr<SDMSourceBlendType> BlendTypeWidgetWeak;
+};
+
+DECLARE_DELEGATE_OneParam(FDMOnSourceBlendTypeChanged, const TSubclassOf<UDMMaterialStageBlend>)
+
+struct FDMBlendNameClass
+{
+	FText BlendName;
+	TSubclassOf<UDMMaterialStageBlend> BlendClass;
+};
+
+class SDMSourceBlendType : public SCompoundWidget
+{
+public:
+	static TSharedRef<SWidget> MakeSourceBlendMenu(TAttribute<TSubclassOf<UDMMaterialStageBlend>> InSelectedItem, FDMOnSourceBlendTypeChanged InOnSelectedItemChanged);
+
+	SLATE_BEGIN_ARGS(SDMSourceBlendType)
+		: _SelectedItem(UDMMaterialStageBlendNormal::StaticClass())
+		{}
+		SLATE_ATTRIBUTE(TSubclassOf<UDMMaterialStageBlend>, SelectedItem)
+		SLATE_EVENT(FDMOnSourceBlendTypeChanged, OnSelectedItemChanged)
+	SLATE_END_ARGS()
+
+	virtual ~SDMSourceBlendType() = default;
+
+	void Construct(const FArguments& InArgs);
+
+protected:
+	static TArray<TStrongObjectPtr<UClass>> SupportedBlendClasses;
+	static TMap<FName, FDMBlendNameClass> BlendMap;
+
+	static void EnsureBlendMap();
+	static void EnsureMenuRegistered();
+	static void MakeSourceBlendMenu(UToolMenu* InToolMenu);
+
+	TAttribute<TSubclassOf<UDMMaterialStageBlend>> SelectedItem;
+	FDMOnSourceBlendTypeChanged OnSelectedItemChanged;
+
+	TSharedRef<SWidget> OnGenerateWidget(const FName InItem);
+
+	void OnSelectionChanged(const FName InNewItem, const ESelectInfo::Type InSelectInfoType);
+
+	FText GetSelectedItemText() const;
+
+	TSharedRef<SWidget> MakeSourceBlendMenuWidget();
+
+	void OnBlendTypeSelected(UClass* InBlendClass);
+	bool CanSelectBlendType(UClass* InBlendClass);
+	bool InBlendTypeSelected(UClass* InBlendClass);
+};
