[MaterialDesigner]

diff --git a/Engine/Plugins/Experimental/DynamicMaterial/Source/DynamicMaterialEditor/Private/Slate/SDMComponentEdit.cpp b/Engine/Plugins/Experimental/DynamicMaterial/Source/DynamicMaterialEditor/Private/Slate/SDMComponentEdit.cpp
new file mode 100644
--- /dev/null
+++ b/Engine/Plugins/Experimental/DynamicMaterial/Source/DynamicMaterialEditor/Private/Slate/SDMComponentEdit.cpp
@@ -0,0 +1,800 @@
+// Copyright Epic Games, Inc. All Rights Reserved.
+
+#include "Slate/SDMComponentEdit.h"
+#include "Components/DMMaterialEffect.h"
+#include "Components/DMMaterialEffectStack.h"
+#include "Components/DMMaterialLayer.h"
+#include "Components/DMMaterialProperty.h"
+#include "Components/DMMaterialSlot.h"
+#include "Components/DMMaterialStage.h"
+#include "Components/DMMaterialStageInput.h"
+#include "Components/DMMaterialStageSource.h"
+#include "Components/DMMaterialStageThroughput.h"
+#include "Components/DMMaterialValue.h"
+#include "Components/DMTextureUV.h"
+#include "CustomDetailsViewArgs.h"
+#include "CustomDetailsViewModule.h"
+#include "DetailLayoutBuilder.h"
+#include "DMEDefs.h"
+#include "DMWorldSubsystem.h"
+#include "DynamicMaterialEditorModule.h"
+#include "DynamicMaterialEditorStyle.h"
+#include "DynamicMaterialModule.h"
+#include "Engine/World.h"
+#include "Framework/MultiBox/MultiBoxBuilder.h"
+#include "ICustomDetailsView.h"
+#include "IDetailKeyframeHandler.h"
+#include "IDetailTreeNode.h"
+#include "ISinglePropertyView.h"
+#include "Items/CustomDetailsViewItemId.h"
+#include "Items/ICustomDetailsViewCustomItem.h"
+#include "Items/ICustomDetailsViewItem.h"
+#include "Menus/DMMaterialStageSourceMenus.h"
+#include "Misc/CoreDelegates.h"
+#include "Model/DynamicMaterialModelEditorOnlyData.h"
+#include "PropertyCustomizationHelpers.h"
+#include "PropertyEditorModule.h"
+#include "ScopedTransaction.h"
+#include "Slate/Properties/Editors/SDMPropertyEditBool.h"
+#include "Slate/Properties/Editors/SDMPropertyEditColor.h"
+#include "Slate/Properties/Editors/SDMPropertyEditEnum.h"
+#include "Slate/Properties/Editors/SDMPropertyEditFloat.h"
+#include "Slate/Properties/Editors/SDMPropertyEditObject.h"
+#include "Slate/Properties/Editors/SDMPropertyEditVector.h"
+#include "Slate/SDMSlot.h"
+#include "Slate/SDMStage.h"
+#include "Styling/AppStyle.h"
+#include "Styling/SlateTypes.h"
+#include "ToolMenu.h"
+#include "ToolMenus.h"
+#include "Widgets/Input/SComboButton.h"
+#include "Widgets/Layout/SBox.h"
+#include "Widgets/SNullWidget.h"
+#include "Widgets/Text/STextBlock.h"
+
+static const float RowPadding = 2.0f;
+
+#define LOCTEXT_NAMESPACE "SDMComponentEdit"
+
+ECheckBoxState SDMComponentEdit::IsInputPerChannelMapped(UDMMaterialStageThroughput* InThroughput, int32 InInputIdx)
+{
+	// TODO Implement per channel mapping
+	return ECheckBoxState::Unchecked;
+}
+
+FText SDMComponentEdit::GetInputChannelMapDescription(UDMMaterialStageThroughput* InThroughput, int32 InInputIdx, int32 InChannelIdx)
+{
+	if (!IsValid(InThroughput) || !InThroughput->IsComponentValid())
+	{
+		return FText::GetEmpty();
+	}
+	
+	UDMMaterialStage* Stage = InThroughput->GetStage();
+	
+	if (!IsValid(Stage) || !Stage->IsComponentValid())
+	{
+		return FText::GetEmpty();
+	}
+
+	const TArray<FDMMaterialStageConnection>& InputConnectionMap = Stage->GetInputConnectionMap();
+
+	if (!InputConnectionMap.IsValidIndex(InInputIdx))
+	{
+		return FText::GetEmpty();
+	}
+
+	const FDMMaterialStageConnection& Connection = InputConnectionMap[InInputIdx];
+
+	if (!Connection.Channels.IsValidIndex(InChannelIdx))
+	{
+		return FText::GetEmpty();
+	}
+
+	FText SourceName = LOCTEXT("?", "?");
+	EDMValueType OutputType = EDMValueType::VT_None;
+
+	const FDMMaterialStageConnectorChannel& Channel = Connection.Channels[InChannelIdx];
+
+	if (Channel.SourceIndex == FDMMaterialStageConnectorChannel::PREVIOUS_STAGE)
+	{
+		if (Stage)
+		{
+			UDMMaterialLayerObject* Layer = Stage->GetLayer();
+
+			if (!IsValid(Layer))
+			{
+				return FText::GetEmpty();
+			}
+
+			UDMMaterialSlot* Slot = Layer->GetSlot();
+
+			if (!IsValid(Slot) || !Slot->IsComponentValid())
+			{
+				return FText::GetEmpty();
+			}
+
+			UDynamicMaterialModelEditorOnlyData* ModelEditorOnlyData = Slot->GetMaterialModelEditorOnlyData();
+
+			if (!IsValid(ModelEditorOnlyData))
+			{
+				return FText::GetEmpty();
+			}
+
+			UDMMaterialProperty* PropertyObj = ModelEditorOnlyData->GetMaterialProperty(Channel.MaterialProperty);
+
+			if (!IsValid(PropertyObj) || !PropertyObj->IsComponentValid())
+			{
+				return FText::GetEmpty();
+			}
+
+			SourceName = FText::Format(
+				LOCTEXT("PrevStageFormat", "Prev Stage {0}"),
+				PropertyObj->GetDescription()
+			);
+
+			if (const UDMMaterialLayerObject* PreviousStage = Layer->GetPreviousLayer(Channel.MaterialProperty, EDMMaterialLayerStage::Base))
+			{
+				if (UDMMaterialStage* PreviousMask = PreviousStage->GetStage(EDMMaterialLayerStage::Mask))
+				{
+					if (UDMMaterialStageSource* PreviousStageSource = PreviousMask->GetSource())
+					{
+						const TArray<FDMMaterialStageConnector>& PreviousStageOutputs = PreviousStageSource->GetOutputConnectors();
+
+						if (PreviousStageOutputs.IsValidIndex(Channel.OutputIndex))
+						{
+							static const FText StageOutputFormat = LOCTEXT("StageOutputFormat", "{0}: {1}");
+
+							SourceName = FText::Format(
+								StageOutputFormat,
+								SourceName,
+								PreviousStageOutputs[Channel.OutputIndex].Name
+							);
+
+							OutputType = PreviousStageOutputs[Channel.OutputIndex].Type;
+						}
+					}
+				}
+			}
+		}
+	}
+	else
+	{
+		const TArray<UDMMaterialStageInput*>& StageInputs = Stage->GetInputs();
+		const int32 StageInputIdx = Channel.SourceIndex - FDMMaterialStageConnectorChannel::FIRST_STAGE_INPUT;
+
+		if (StageInputs.IsValidIndex(StageInputIdx))
+		{
+			SourceName = StageInputs[StageInputIdx]->GetChannelDescription(Channel);
+
+			const TArray<FDMMaterialStageConnector>& StageInputOutputConnectors = StageInputs[StageInputIdx]->GetOutputConnectors();
+
+			if (StageInputOutputConnectors.IsValidIndex(Channel.OutputIndex))
+			{
+				if (StageInputOutputConnectors.Num() > 1)
+				{
+					static const FText InputOutputFormat = LOCTEXT("InputOutputFormat", "{0}: {1}");
+
+					SourceName = FText::Format(
+						InputOutputFormat,
+						SourceName,
+						StageInputOutputConnectors[Channel.OutputIndex].Name
+					);
+				}
+
+				OutputType = StageInputOutputConnectors[Channel.OutputIndex].Type;
+			}
+		}
+	}
+
+	if (Channel.OutputChannel == FDMMaterialStageConnectorChannel::WHOLE_CHANNEL || OutputType == EDMValueType::VT_None)
+	{
+		static const FText WholeChannelFormat = LOCTEXT("WholeChannelFormat", "{0}");
+		return FText::Format(WholeChannelFormat, SourceName);
+	}
+
+	// Assume RGBA.
+	TArray<FText> Channels;
+
+	if (Channel.OutputChannel & FDMMaterialStageConnectorChannel::FIRST_CHANNEL)
+	{
+		Channels.Add(UDMValueDefinitionLibrary::GetValueDefinition(OutputType).GetChannelName(1));
+	}
+
+	if (Channel.OutputChannel & FDMMaterialStageConnectorChannel::SECOND_CHANNEL)
+	{
+		Channels.Add(UDMValueDefinitionLibrary::GetValueDefinition(OutputType).GetChannelName(2));
+	}
+
+	if (Channel.OutputChannel & FDMMaterialStageConnectorChannel::THIRD_CHANNEL)
+	{
+		Channels.Add(UDMValueDefinitionLibrary::GetValueDefinition(OutputType).GetChannelName(3));
+	}
+
+	if (Channel.OutputChannel & FDMMaterialStageConnectorChannel::FOURTH_CHANNEL)
+	{
+		Channels.Add(UDMValueDefinitionLibrary::GetValueDefinition(OutputType).GetChannelName(4));
+	}
+
+	static const FText Separator = LOCTEXT("ChannelSeparator", ",");
+	static const FText MaskedChannelFormat = LOCTEXT("WholeChannelFormat", "{0}: {1}");
+	const FText ChannelName = FText::Join(Separator, Channels);
+
+	return FText::Format(MaskedChannelFormat, SourceName, ChannelName);
+}
+
+SDMComponentEdit::~SDMComponentEdit()
+{
+	if (!FDynamicMaterialModule::AreUObjectsSafe())
+	{
+		return;
+	}
+
+	SDMEditor::ClearPropertyHandles(this);
+
+	if (UDMMaterialStage* Stage = Cast<UDMMaterialStage>(ComponentWeak.Get()))
+	{
+		Stage->GetOnUpdate().RemoveAll(this);
+	}
+	else if (UDMMaterialEffect* Effect = Cast<UDMMaterialEffect>(ComponentWeak.Get()))
+	{
+		Effect->GetOnUpdate().RemoveAll(this);
+	}
+}
+
+void SDMComponentEdit::Construct(const FArguments& InArgs, UDMMaterialComponent* InComponent, const TWeakPtr<SDMSlot>& InSlotWidget)
+{
+	ComponentWeak = InComponent;
+	SlotWidgetWeak = InSlotWidget;
+
+	KeyframeHandler = nullptr;
+
+	if (ensure(IsValid(InComponent)))
+	{
+		if (UDMMaterialStage* Stage = Cast<UDMMaterialStage>(InComponent))
+		{
+			Stage->GetOnUpdate().AddSP(this, &SDMComponentEdit::OnStageUpdate);
+		}
+		else if (UDMMaterialEffect* Effect = Cast<UDMMaterialEffect>(InComponent))
+		{
+			Effect->GetOnUpdate().AddSP(this, &SDMComponentEdit::OnEffectUpdated);
+		}
+
+		if (const UWorld* const World = InComponent->GetWorld())
+		{
+			if (const UDMWorldSubsystem* const WorldSubsystem = World->GetSubsystem<UDMWorldSubsystem>())
+			{
+				KeyframeHandler = WorldSubsystem->GetKeyframeHandler();
+			}
+		}
+
+		ChildSlot
+			.HAlign(HAlign_Fill)
+			.VAlign(VAlign_Top)
+			[
+				SAssignNew(Container, SBox)
+				[
+					CreateEditWidget()
+				]
+			];
+	}
+}
+
+TSharedRef<SWidget> SDMComponentEdit::CreateEditWidget()
+{
+	SDMEditor::ClearPropertyHandles(this);
+
+	FCustomDetailsViewArgs Args;
+	Args.KeyframeHandler = KeyframeHandler;
+	Args.bAllowGlobalExtensions = true;
+	Args.bAllowResetToDefault = true;
+	Args.bShowCategories = false;
+
+	TSharedRef<ICustomDetailsView> DetailsView = ICustomDetailsViewModule::Get().CreateCustomDetailsView(Args);
+	FCustomDetailsViewItemId RootId = DetailsView->GetRootItem()->GetItemId();
+
+	if (UDMMaterialStage* Stage = Cast<UDMMaterialStage>(ComponentWeak.Get()))
+	{
+		TSharedPtr<ICustomDetailsViewCustomItem> TypeSelectorItem = DetailsView->CreateCustomItem(
+			FName(TEXT("SamplerType")),
+			LOCTEXT("SamplerType", "Type"),
+			LOCTEXT("SamplerTypeTooltip", "Source Function")
+		);
+
+		if (TypeSelectorItem.IsValid())
+		{
+			TypeSelectorItem->SetValueWidget(CreateSourceTypeEditWidget());
+
+			DetailsView->ExtendTree(
+				RootId,
+				ECustomDetailsTreeInsertPosition::FirstChild,
+				TypeSelectorItem->AsItem()
+			);
+		}
+	}
+
+	TArray<FDMPropertyHandle> EditRows = GetEditRows();
+
+	for (const FDMPropertyHandle& EditRow : EditRows)
+	{
+		ECustomDetailsTreeInsertPosition Position = ECustomDetailsTreeInsertPosition::Child;
+
+		if (EditRow.DetailTreeNode->CreatePropertyHandle()->HasMetaData("HighPriority"))
+		{
+			Position = ECustomDetailsTreeInsertPosition::FirstChild;
+		}
+		else if (EditRow.DetailTreeNode->CreatePropertyHandle()->HasMetaData("LowPriority"))
+		{
+			Position = ECustomDetailsTreeInsertPosition::LastChild;
+		}
+
+		TSharedRef<ICustomDetailsViewItem> Item = DetailsView->CreateDetailTreeItem(EditRow.DetailTreeNode.ToSharedRef());
+
+		if (EditRow.NameOverride.IsSet())
+		{
+			Item->SetOverrideWidget(
+				ECustomDetailsViewWidgetType::Name,
+				SNew(STextBlock)
+				.Font(IDetailLayoutBuilder::GetDetailFont())
+				.Text(EditRow.NameOverride.GetValue())
+				.ToolTipText(EditRow.NameToolTipOverride.Get(FText::GetEmpty()))
+			);
+		}
+
+		if (EditRow.DetailTreeNode->CreatePropertyHandle()->HasMetaData("NotKeyframeable"))
+		{
+			Item->SetKeyframeEnabled(false);
+		}
+
+		if (EditRow.ResetToDefaultOverride.IsSet())
+		{
+			Item->SetResetToDefaultOverride(EditRow.ResetToDefaultOverride.GetValue());
+		}
+
+		DetailsView->ExtendTree(RootId, Position, Item);
+	}
+
+	DetailsView->RebuildTree(ECustomDetailsViewBuildType::InstantBuild);
+
+	return DetailsView;
+}
+
+TArray<FDMPropertyHandle> SDMComponentEdit::GetEditRows()
+{
+	TArray<FDMPropertyHandle> PropertyRows;
+	TSet<UDMMaterialComponent*> ProcessedObjects;
+
+	FDynamicMaterialEditorModule::GeneratorComponentPropertyRows(SharedThis(this), ComponentWeak.Get(), PropertyRows, ProcessedObjects);
+
+	return PropertyRows;
+}
+
+TSharedRef<SWidget> SDMComponentEdit::CreateExtensionButtons(const TSharedPtr<SDMComponentEdit>& InComponentEditWidget,
+	UDMMaterialComponent* InComponent, const FName& InPropertyName, bool bInAllowKeyframe)
+{
+	return CreateExtensionButtons(
+		InComponentEditWidget, 
+		InComponent, 
+		InPropertyName, 
+		bInAllowKeyframe,
+		FSimpleDelegate::CreateLambda([ComponentEditWidgetWeak = InComponentEditWidget.ToWeakPtr()]()
+			{
+				if (TSharedPtr<SDMComponentEdit> ComponentEditWidget = ComponentEditWidgetWeak.Pin())
+				{
+					if (TSharedPtr<SDMSlot> SlotWidget = ComponentEditWidget->GetSlotWidget())
+					{
+						SlotWidget->InvalidateComponentEditWidget();
+					}
+				}
+			})
+	);
+}
+
+TSharedRef<SWidget> SDMComponentEdit::CreateExtensionButtons(const TSharedPtr<SWidget>& InPropertyOwner, UDMMaterialComponent* InComponent,
+	const FName& InPropertyName, bool bInAllowKeyframe, FSimpleDelegate InOnResetDelegate)
+{
+	if (!IsValid(InComponent))
+	{
+		return SNew(SBox)
+			.WidthOverride(42.f)
+			.HeightOverride(18.f);
+	}
+
+	FProperty* Property = InComponent->GetClass()->FindPropertyByName(InPropertyName);
+
+	if (!Property)
+	{
+		return SNew(SBox)
+			.WidthOverride(42.f)
+			.HeightOverride(18.f);
+	}
+
+	TSharedPtr<IPropertyHandle> PropertyHandle = nullptr;
+	TSharedPtr<IDetailTreeNode> DetailTreeNode = nullptr;
+	FString PropertyName = InPropertyName.ToString();
+
+	if (UDMTextureUV* TextureUV = Cast<UDMTextureUV>(InComponent))
+	{
+		PropertyHandle = TextureUV->GetPropertyHandle(InPropertyName);
+		DetailTreeNode = TextureUV->GetDetailTreeNode(InPropertyName);
+		PropertyName = FString("TextureUV->") + PropertyName;
+	}
+	else if (UDMMaterialValue* Value = Cast<UDMMaterialValue>(InComponent))
+	{
+		PropertyHandle = Value->GetPropertyHandle();
+		DetailTreeNode = Value->GetDetailTreeNode();
+	}
+	else
+	{
+		FDMPropertyHandle DMPropertyHandle = SDMEditor::GetPropertyHandle(InPropertyOwner.Get(), InComponent, InPropertyName);
+		PropertyHandle = DMPropertyHandle.PropertyHandle;
+		DetailTreeNode = DMPropertyHandle.DetailTreeNode;
+	}
+
+	if (PropertyHandle.IsValid())
+	{
+		ensure(PropertyHandle->GetProperty() == Property);
+	}
+
+	FOnGenerateGlobalRowExtensionArgs ExtensionRowArgs;
+	ExtensionRowArgs.OwnerObject = InComponent;
+	ExtensionRowArgs.Property = Property;
+	ExtensionRowArgs.PropertyPath = PropertyName;
+	ExtensionRowArgs.OwnerTreeNode = DetailTreeNode;
+	ExtensionRowArgs.PropertyHandle = PropertyHandle;
+
+	FSlimHorizontalToolBarBuilder ToolBarBuilder(TSharedPtr<FUICommandList>(), FMultiBoxCustomization::None);
+
+	FPropertyEditorModule& PropertyEditorModule = FModuleManager::GetModuleChecked<FPropertyEditorModule>("PropertyEditor");
+
+	TArray<FPropertyRowExtensionButton> ExtensionButtons;
+	PropertyEditorModule.GetGlobalRowExtensionDelegate().Broadcast(ExtensionRowArgs, ExtensionButtons);
+
+	TSharedRef<SWidget> ResetButtonWidget = PropertyCustomizationHelpers::MakeResetButton(
+		FSimpleDelegate::CreateLambda([PropertyHandle, InOnResetDelegate]()
+			{
+				if (PropertyHandle.IsValid() && PropertyHandle->IsValidHandle())
+				{
+					TArray<UObject*> Outers;
+					PropertyHandle->GetOuterObjects(Outers);
+
+					if (Outers.IsEmpty() == false && IsValid(Outers[0]))
+					{
+						FScopedTransaction Transaction(LOCTEXT("ResetValue", "Reset Value to default."));
+
+						if (FProperty* Property = PropertyHandle->GetProperty())
+						{
+							Outers[0]->Modify();
+
+							if (UDMMaterialValue* Value = Cast<UDMMaterialValue>(Outers[0]))
+							{
+								Value->ApplyDefaultValue();
+							}
+							else
+							{
+								void* DefaultValue = Property->ContainerPtrToValuePtr<void>(Outers[0]->GetClass()->GetDefaultObject(true));
+
+								if (DefaultValue != nullptr)
+								{
+									if (Property->HasSetter())
+									{
+										Property->CallSetter(Outers[0], DefaultValue);
+									}
+									else
+									{
+										Outers[0]->PreEditChange(Property);
+
+										Property->SetValue_InContainer(Outers[0], DefaultValue);
+
+										FPropertyChangedEvent PCE = FPropertyChangedEvent(Property, EPropertyChangeType::ValueSet);
+										Outers[0]->PostEditChangeProperty(PCE);
+									}
+								}
+							}
+						}
+						else
+						{
+							PropertyHandle->ResetToDefault();
+						}
+
+						InOnResetDelegate.ExecuteIfBound();
+					}
+				}
+			})
+	);
+	ResetButtonWidget->SetVisibility(
+		TAttribute<EVisibility>::CreateLambda([PropertyHandle]()->EVisibility
+			{
+				if (PropertyHandle.IsValid() && PropertyHandle->IsValidHandle())
+				{
+					if (FProperty* Property = PropertyHandle->GetProperty())
+					{
+						TArray<UObject*> Outers;
+						PropertyHandle->GetOuterObjects(Outers);
+
+						if (Outers.IsEmpty() == false && IsValid(Outers[0]))
+						{
+							if (UDMMaterialValue* Value = Cast<UDMMaterialValue>(Outers[0]))
+							{
+								return Value->IsDefaultValue()
+									? EVisibility::Hidden
+									: EVisibility::Visible;
+							}
+
+							return Property->Identical_InContainer(Outers[0], Outers[0]->GetClass()->GetDefaultObject(true))
+								? EVisibility::Hidden
+								: EVisibility::Visible;
+						}
+					}
+
+					return PropertyHandle->CanResetToDefault() ? EVisibility::Visible : EVisibility::Hidden;
+				}
+
+				return EVisibility::Hidden;
+			})
+	);
+
+	ToolBarBuilder.AddWidget(ResetButtonWidget);
+
+	// Sequencer relies on getting the Keyframe Handler via the Details View of the IDetailTreeNode, but it's null since
+	// there's no Details View here. Instead add it manually.
+	if (bInAllowKeyframe && InComponent->GetWorld() && PropertyHandle.IsValid())
+	{
+		ToolBarBuilder.AddToolBarButton(
+			FUIAction(FExecuteAction::CreateStatic(&SDMComponentEdit::CreateKeyFrame, PropertyHandle)	),
+			NAME_None,
+			FText::GetEmpty(),
+			LOCTEXT("CreateKeyToolTip", "Add a keyframe for this property."),
+			FSlateIcon(FAppStyle::Get().GetStyleSetName(), "Sequencer.AddKey.Details")
+		);
+	}
+	else
+	{
+		// Maintain space
+		ToolBarBuilder.AddWidget(
+			SNew(SBox)
+			.WidthOverride(20.f)
+			.HeightOverride(18.f)
+		);
+	}
+
+	if (!ExtensionButtons.IsEmpty())
+	{
+		// Build extension toolbar 
+		ToolBarBuilder.SetLabelVisibility(EVisibility::Collapsed);
+		ToolBarBuilder.SetStyle(&FAppStyle::Get(), "DetailsView.ExtensionToolBar");
+
+		for (const FPropertyRowExtensionButton& Extension : ExtensionButtons)
+		{
+			ToolBarBuilder.AddToolBarButton(Extension.UIAction, NAME_None, Extension.Label, Extension.ToolTip, Extension.Icon);
+		}
+	}
+
+	return ToolBarBuilder.MakeWidget();
+}
+
+TSharedPtr<SWidget> SDMComponentEdit::CreateSinglePropertyEditWidget(UDMMaterialComponent* InComponent, const FName& InPropertyName)
+{
+	if (!IsValid(InComponent))
+	{
+		return nullptr;
+	}
+
+	FProperty* ClassProperty = InComponent->GetClass()->FindPropertyByName(InPropertyName);
+
+	if (!ClassProperty)
+	{
+		return nullptr;
+	}
+
+	FDMPropertyHandle PropertyHandle = SDMEditor::GetPropertyHandle(this, InComponent, InPropertyName);
+
+	if (PropertyHandle.PropertyHandle.IsValid())
+	{
+		if (ClassProperty->IsA<FBoolProperty>())
+		{
+			return SNew(SDMPropertyEditBool, PropertyHandle.PropertyHandle)
+				.ComponentEditWidget(SharedThis(this));
+		}
+		else if (ClassProperty->IsA<FEnumProperty>())
+		{
+			return SNew(SDMPropertyEditEnum, PropertyHandle.PropertyHandle)
+				.ComponentEditWidget(SharedThis(this));
+		}
+		else if (ClassProperty->IsA<FFloatProperty>())
+		{
+			return SNew(SDMPropertyEditFloat, PropertyHandle.PropertyHandle)
+				.ComponentEditWidget(SharedThis(this));
+		}
+		else if (FObjectPropertyBase* ObjectProperty = CastField<FObjectPropertyBase>(ClassProperty))
+		{
+			return SNew(SDMPropertyEditObject, PropertyHandle.PropertyHandle, ObjectProperty->PropertyClass)
+				.ComponentEditWidget(SharedThis(this));
+		}
+		else if (FStructProperty* StructProperty = CastField<FStructProperty>(ClassProperty))
+		{
+			if (StructProperty->Struct == TBaseStructure<FVector2D>::Get())
+			{
+				return SNew(SDMPropertyEditVector, PropertyHandle.PropertyHandle, 2)
+					.ComponentEditWidget(SharedThis(this));
+			}
+			else if (StructProperty->Struct == TBaseStructure<FVector>::Get()
+				|| StructProperty->Struct == TBaseStructure<FRotator>::Get())
+			{
+				return SNew(SDMPropertyEditVector, PropertyHandle.PropertyHandle, 3)
+					.ComponentEditWidget(SharedThis(this));
+			}
+			else if (StructProperty->Struct == TBaseStructure<FLinearColor>::Get())
+			{
+				return SNew(SDMPropertyEditColor, PropertyHandle.PropertyHandle)
+					.ComponentEditWidget(SharedThis(this));
+			}
+		}
+
+		return nullptr;
+	}
+
+	// This is a complete fallback that does not layout correctly.
+	FPropertyEditorModule& PropertyEditor = FModuleManager::Get().LoadModuleChecked<FPropertyEditorModule>(TEXT("PropertyEditor"));
+
+	FSinglePropertyParams InitParams;
+	InitParams.NamePlacement = EPropertyNamePlacement::Hidden;
+	InitParams.NotifyHook = InComponent;
+
+	return PropertyEditor.CreateSingleProperty(InComponent, InPropertyName, InitParams);
+}
+
+TSharedRef<SWidget> SDMComponentEdit::MakeSourceTypeEditWidgetMenuContent()
+{
+	if (UDMMaterialStage* Stage = Cast<UDMMaterialStage>(ComponentWeak.Get()))
+	{
+		if (TSharedPtr<SDMSlot> SlotWidget = GetSlotWidget())
+		{
+			if (TSharedPtr<SDMStage> StageWidget = SlotWidget->FindStageWidget(Stage))
+			{
+				return FDMMaterialStageSourceMenus::MakeChangeSourceMenu(SlotWidget, StageWidget);
+			}
+		}
+	}
+
+	return SNullWidget::NullWidget;
+}
+
+FText SDMComponentEdit::GetSourceTypeEditWidgetText() const
+{
+	if (UDMMaterialStage* Stage = Cast<UDMMaterialStage>(ComponentWeak.Get()))
+	{
+		if (UDMMaterialStageSource* Source = Stage->GetSource())
+		{
+			return Source->GetStageDescription();
+		}
+	}
+
+	return FText::GetEmpty();
+}
+
+void SDMComponentEdit::OnStageUpdate(UDMMaterialComponent* InComponent, EDMUpdateType InUpdateType)
+{
+	if (InUpdateType == EDMUpdateType::Structure)
+	{
+		if (InComponent == ComponentWeak && !UpdateHandle.IsValid())
+		{
+			UpdateHandle = FCoreDelegates::OnEndFrame.AddSPLambda(this, [this]()
+				{
+					if (Container.IsValid())
+					{
+						Container->SetContent(CreateEditWidget());
+						FCoreDelegates::OnEndFrame.Remove(UpdateHandle);
+					}
+				});
+		}
+	}
+}
+
+void SDMComponentEdit::OnEffectUpdated(UDMMaterialComponent* InComponent, EDMUpdateType InUpdateType)
+{
+	if (InUpdateType == EDMUpdateType::Structure)
+	{
+		if (InComponent == ComponentWeak.Get() && !UpdateHandle.IsValid())
+		{
+			UpdateHandle = FCoreDelegates::OnEndFrame.AddSPLambda(this, [this]()
+				{
+					if (Container.IsValid())
+					{
+						Container->SetContent(CreateEditWidget());
+						FCoreDelegates::OnEndFrame.Remove(UpdateHandle);
+					}
+				});
+		}
+	}
+}
+
+void SDMComponentEdit::OnUndo()
+{
+	if (UDMMaterialStage* Stage = Cast<UDMMaterialStage>(ComponentWeak.Get()))
+	{
+		if (const UDMMaterialLayerObject* Layer = Stage->GetLayer())
+		{
+			if (Layer->GetStageType(Stage) == EDMMaterialLayerStage::Mask)
+			{
+				if (bCreatedWithLinkedUVs != Layer->IsTextureUVLinkEnabled())
+				{
+					OnStageUpdate(Stage, EDMUpdateType::Structure);
+				}
+			}
+		}
+	}
+}
+
+TSharedRef<SWidget> SDMComponentEdit::CreateSourceTypeEditWidget()
+{
+	TWeakPtr<SDMComponentEdit> ThisWeak = SharedThis(this);
+
+	return SNew(SBox)
+		.HAlign(EHorizontalAlignment::HAlign_Left)
+		.VAlign(EVerticalAlignment::VAlign_Center)
+		.Padding(0.f, 0.f, 0.f, 0.2f)
+		[
+			SNew(SComboButton)
+				.HasDownArrow(false)
+				.IsFocusable(true)
+				.ContentPadding(4.0f)
+				.ButtonStyle(FDynamicMaterialEditorStyle::Get(), "HoverHintOnly.Bordered.Dark")
+				.ToolTipText(LOCTEXT("ChangeLayer", "Change Stage Type"))
+				.OnGetMenuContent(this, &SDMComponentEdit::MakeSourceTypeEditWidgetMenuContent)
+				.ButtonContent()
+				[
+					SNew(STextBlock)
+					.Text(this, &SDMComponentEdit::GetSourceTypeEditWidgetText)
+				]
+		];
+}
+
+void SDMComponentEdit::CreateKeyFrame(TSharedPtr<IPropertyHandle> InPropertyHandle)
+{
+	TArray<UObject*> Outers;
+	InPropertyHandle->GetOuterObjects(Outers);
+
+	if (Outers.IsEmpty() || !IsValid(Outers[0]))
+	{
+		return;
+	}
+
+	UDMMaterialComponent* Component = Cast<UDMMaterialComponent>(Outers[0]);
+
+	if (!IsValid(Component))
+	{
+		return;
+	}
+
+	if (!InPropertyHandle.IsValid())
+	{
+		return;
+	}
+
+	const UWorld* const World = Component->GetWorld();
+	if (!World)
+	{
+		return;
+	}
+
+	const UDMWorldSubsystem* const WorldSubsystem = World->GetSubsystem<UDMWorldSubsystem>();
+	if (!WorldSubsystem)
+	{
+		return;
+	}
+
+	TSharedPtr<IDetailKeyframeHandler> KeyframeHandler = WorldSubsystem->GetKeyframeHandler();
+	if (!KeyframeHandler.IsValid())
+	{
+		return;
+	}
+
+	if (!KeyframeHandler->IsPropertyKeyable(Component->GetClass(), *InPropertyHandle))
+	{
+		return;
+	}
+
+	KeyframeHandler->OnKeyPropertyClicked(*InPropertyHandle);
+}
+
+#undef LOCTEXT_NAMESPACE
diff --git a/Engine/Plugins/Experimental/DynamicMaterial/Source/DynamicMaterialEditor/Private/Slate/SDMEditor.cpp b/Engine/Plugins/Experimental/DynamicMaterial/Source/DynamicMaterialEditor/Private/Slate/SDMEditor.cpp
new file mode 100644
--- /dev/null
+++ b/Engine/Plugins/Experimental/DynamicMaterial/Source/DynamicMaterialEditor/Private/Slate/SDMEditor.cpp
@@ -0,0 +1,1509 @@
+// Copyright Epic Games, Inc. All Rights Reserved.
+
+#include "SDMEditor.h"
+#include "AssetThumbnail.h"
+#include "AssetToolsModule.h"
+#include "Components/DMMaterialLayer.h"
+#include "Components/DMMaterialProperty.h"
+#include "Components/DMMaterialSlot.h"
+#include "Components/DMMaterialValue.h"
+#include "Components/MaterialStageExpressions/DMMSETextureSample.h"
+#include "Components/MaterialValues/DMMaterialValueFloat1.h"
+#include "Components/PrimitiveComponent.h"
+#include "DetailLayoutBuilder.h"
+#include "DMBlueprintFunctionLibrary.h"
+#include "DMWorldSubsystem.h"
+#include "DynamicMaterialEditorCommands.h"
+#include "DynamicMaterialEditorModule.h"
+#include "DynamicMaterialEditorSettings.h"
+#include "DynamicMaterialEditorStyle.h"
+#include "DynamicMaterialModule.h"
+#include "Engine/World.h"
+#include "Framework/Commands/GenericCommands.h"
+#include "Framework/Commands/InputChord.h"
+#include "HAL/PlatformApplicationMisc.h"
+#include "IDetailTreeNode.h"
+#include "IDocumentation.h"
+#include "IPropertyRowGenerator.h"
+#include "MaterialDomain.h"
+#include "Menus/DMToolBarMenus.h"
+#include "Misc/CoreDelegates.h"
+#include "Model/DynamicMaterialModel.h"
+#include "Model/DynamicMaterialModelEditorOnlyData.h"
+#include "PropertyCustomizationHelpers.h"
+#include "PropertyHandle.h"
+#include "ScopedTransaction.h"
+#include "SEnumCombo.h"
+#include "Slate/Properties/Editors/SDMPropertyEditOpacity.h"
+#include "Slate/Properties/SDMBlendMode.h"
+#include "Slate/Properties/SDMDomain.h"
+#include "Slate/Properties/SDMMaterialParameters.h"
+#include "Slate/SDMSlot.h"
+#include "Slate/SDMComponentEdit.h"
+#include "Slate/SDMToolBar.h"
+#include "SlateOptMacros.h"
+#include "Styling/StyleColors.h"
+#include "ThumbnailRendering/ThumbnailManager.h"
+#include "Widgets/Colors/SColorBlock.h"
+#include "Widgets/Input/SButton.h"
+#include "Widgets/Layout/SBox.h"
+#include "Widgets/Layout/SConstraintCanvas.h"
+#include "Widgets/Layout/SExpandableArea.h"
+#include "Widgets/Layout/SScrollBox.h"
+#include "Widgets/SBoxPanel.h"
+#include "Widgets/SNullWidget.h"
+#include "Widgets/Text/STextBlock.h"
+
+#define LOCTEXT_NAMESPACE "SDMEditor"
+
+TSharedPtr<FAssetThumbnailPool> SDMEditor::ThumbnailPool = nullptr;
+TMap<const SWidget*, TArray<FDMPropertyHandle>> SDMEditor::PropertyHandleMap;
+
+TSharedRef<FAssetThumbnailPool> SDMEditor::GetThumbnailPool()
+{
+	if (TSharedPtr<FAssetThumbnailPool> SharedPool = UThumbnailManager::Get().GetSharedThumbnailPool())
+	{
+		return SharedPool.ToSharedRef();
+	}
+
+	if (!ThumbnailPool.IsValid())
+	{
+		ThumbnailPool = MakeShared<FAssetThumbnailPool>(1024);
+
+		FCoreDelegates::OnEnginePreExit.AddLambda([]()
+			{
+				ThumbnailPool.Reset();
+			});
+	}
+
+	return ThumbnailPool.ToSharedRef();
+}
+
+void SDMEditor::Construct(const FArguments& InArgs, TWeakObjectPtr<UDynamicMaterialModel> InModelWeak)
+{
+	SetCanTick(true);
+
+	ActiveSlotIndex = 0;
+
+	BindCommands();
+	
+	ChildSlot
+	.HAlign(HAlign_Fill)
+	.VAlign(VAlign_Fill)
+	[
+		SAssignNew(Container, SBox)
+		[
+			CreateMainLayout()
+		]
+	];
+
+	SetMaterialModel(InModelWeak.Get());
+
+	UDynamicMaterialEditorSettings::Get()->OnSettingsChanged.AddSP(this, &SDMEditor::OnSettingsChanged);
+
+	static bool bAddedEnginePreExitDelegate = false;
+
+	if (!bAddedEnginePreExitDelegate)
+	{
+		bAddedEnginePreExitDelegate = true;
+
+		FCoreDelegates::OnEnginePreExit.AddLambda(
+			[]
+			{
+				PropertyHandleMap.Reset();
+			}
+		);
+	}
+}
+
+void SDMEditor::PrivateRegisterAttributes(struct FSlateAttributeDescriptor::FInitializer&)
+{
+	
+}
+
+SDMEditor::~SDMEditor()
+{
+	if (FDynamicMaterialModule::AreUObjectsSafe())
+	{
+		if (UDynamicMaterialModel* MaterialModel = MaterialModelWeak.Get())
+		{
+			if (UDynamicMaterialModelEditorOnlyData* ModelEditorOnlyData = UDynamicMaterialModelEditorOnlyData::Get(MaterialModel))
+			{
+				ModelEditorOnlyData->GetOnMaterialBuiltDelegate().RemoveAll(this);
+				ModelEditorOnlyData->GetOnValueListUpdateDelegate().RemoveAll(this);
+				ModelEditorOnlyData->GetOnSlotListUpdateDelegate().RemoveAll(this);
+			}
+		}
+	}
+}
+
+void SDMEditor::ClearEditor()
+{
+	SetMaterialModel(nullptr);
+}
+
+void SDMEditor::SetActiveSlotIndex(int InSlotIndex)
+{
+	if (InSlotIndex != ActiveSlotIndex && SlotWidgets.IsValidIndex(InSlotIndex))
+	{
+		ActiveSlotIndex = InSlotIndex;
+		SlotWidgets[InSlotIndex]->ClearSelection();
+		RefreshSlotsList();
+	}
+}
+
+FDMPropertyHandle SDMEditor::GetPropertyHandle(const SWidget* InOwner, UDMMaterialComponent* InComponent, const FName& InPropertyName)
+{
+	TArray<FDMPropertyHandle>& PropertyHandles = PropertyHandleMap.FindOrAdd(InOwner);
+
+	for (const FDMPropertyHandle& ExistingHandle : PropertyHandles)
+	{
+		if (ExistingHandle.PropertyHandle->GetProperty()->GetFName() == InPropertyName)
+		{
+			TArray<UObject*> Outers;
+			ExistingHandle.PropertyHandle->GetOuterObjects(Outers);
+
+			if (Outers.IsEmpty() == false && Outers[0] == InComponent)
+			{
+				return ExistingHandle;
+			}
+		}
+	}
+
+	FDMPropertyHandle NewHandle = CreatePropertyHandle(InOwner, InComponent, InPropertyName);
+	PropertyHandles.Add(NewHandle);
+
+	return NewHandle;
+}
+
+void SDMEditor::ClearPropertyHandles(const SWidget* InOwner)
+{
+	PropertyHandleMap.Remove(InOwner);
+}
+
+FDMPropertyHandle SDMEditor::CreatePropertyHandle(const void* InOwner, UDMMaterialComponent* InComponent,
+	const FName& InPropertyName)
+{
+	FDMPropertyHandle PropertyHandle;
+
+	FPropertyEditorModule& PropertyEditor = FModuleManager::Get().LoadModuleChecked<FPropertyEditorModule>(TEXT("PropertyEditor"));
+
+	FPropertyRowGeneratorArgs RowGeneratorArgs;
+	RowGeneratorArgs.NotifyHook = InComponent;
+
+	PropertyHandle.PropertyRowGenerator = PropertyEditor.CreatePropertyRowGenerator(RowGeneratorArgs);
+	PropertyHandle.PropertyRowGenerator->SetObjects({InComponent});
+
+	for (const TSharedRef<IDetailTreeNode>& CategoryNode : PropertyHandle.PropertyRowGenerator->GetRootTreeNodes())
+	{
+		if (CategoryNode->GetNodeName() != TEXT("Material Designer"))
+		{
+			continue;
+		}
+
+		TArray<TSharedRef<IDetailTreeNode>> ChildNodes;
+		CategoryNode->GetChildren(ChildNodes);
+
+		for (const TSharedRef<IDetailTreeNode>& ChildNode : ChildNodes)
+		{
+			if (ChildNode->GetNodeType() != EDetailNodeType::Item)
+			{
+				continue;
+			}
+
+			if (ChildNode->GetNodeName() != InPropertyName)
+			{
+				continue;
+			}
+
+			PropertyHandle.DetailTreeNode = ChildNode;
+			PropertyHandle.PropertyHandle = ChildNode->CreatePropertyHandle();
+			return PropertyHandle;
+		}
+	}
+
+	return PropertyHandle;
+}
+
+void SDMEditor::SetMaterialModel(UDynamicMaterialModel* InMaterialModel)
+{
+	MaterialModelWeak = InMaterialModel;
+
+	SlotsContainer.Reset();
+	SlotWidgets.Empty();
+
+	Toolbar->SetMaterialModel(InMaterialModel);
+
+	if (!InMaterialModel)
+	{
+		Container->SetContent(SDMEditor::GetEmptyContent());
+		return;
+	}
+
+	Container->SetContent(CreateMainLayout());
+
+	if (UDynamicMaterialModelEditorOnlyData* ModelEditorOnlyData = UDynamicMaterialModelEditorOnlyData::Get(InMaterialModel))
+	{
+		ModelEditorOnlyData->GetOnMaterialBuiltDelegate().AddSP(this, &SDMEditor::OnMaterialBuilt);
+		ModelEditorOnlyData->GetOnValueListUpdateDelegate().AddSP(this, &SDMEditor::OnValuesUpdated);
+		ModelEditorOnlyData->GetOnSlotListUpdateDelegate().AddSP(this, &SDMEditor::OnSlotsUpdated);
+	}
+}
+
+void SDMEditor::SetMaterialObjectProperty(const FDMObjectMaterialProperty& InObjectProperty)
+{
+	UDynamicMaterialModel* MaterialModel = InObjectProperty.GetMaterialModel();
+
+	if (IsValid(MaterialModel))
+	{
+		ObjectProperty = InObjectProperty;
+		SetMaterialModel(MaterialModel);
+	}
+	else
+	{
+		ObjectProperty.Reset();
+	}
+}
+
+void SDMEditor::SetMaterialActor(AActor* InActor)
+{
+	if (UDynamicMaterialEditorSettings* Settings = UDynamicMaterialEditorSettings::Get())
+	{
+		if (!Settings->bFollowSelection)
+		{
+			return;
+		}
+	}
+
+	if (!Container.IsValid())
+	{
+		return;
+	}
+
+	if (!IsValid(InActor))
+	{
+		Container->SetContent(SDMEditor::GetEmptyContent());
+		return;
+	}
+
+	Toolbar->SetMaterialActor(InActor);
+
+	Container->SetContent(
+		SNew(SBox)
+		.HAlign(HAlign_Center)
+		.Padding(5.0f, 5.0f, 5.0f, 5.0f)
+		[
+			CreateActorMaterialSlotSelector(InActor)
+		]
+	);
+}
+
+TSharedRef<SWidget> SDMEditor::CreateActorMaterialSlotSelector(const AActor* InActor)
+{
+	TArray<TSharedPtr<FDMObjectMaterialProperty>> MaterialProperties = Toolbar->GetMaterialProperties();
+	if (MaterialProperties.IsEmpty())
+	{
+		return 
+			SNew(STextBlock)
+			.Justification(ETextJustify::Center)
+			.AutoWrapText(true)
+			.Text(LOCTEXT("NoMaterialSlot", "\n\nThe selected actor contains no primitive components with material slots."));
+	}
+
+	TSharedRef<SVerticalBox> ListOuter = 
+		SNew(SVerticalBox)
+		+ SVerticalBox::Slot()
+		.AutoHeight()
+		.HAlign(HAlign_Center)
+		.Padding(0.0f, 20.0f, 0.0f, 20.0f)
+		[
+			SNew(STextBlock)
+			.TextStyle(FDynamicMaterialEditorStyle::Get(), "ActorNameBig")
+			.Text(Toolbar.Get(), &SDMToolBar::GetSlotActorDisplayName)
+		];
+
+	const UObject* CurrentOuter = nullptr;
+
+	for (const TSharedPtr<FDMObjectMaterialProperty>& MaterialSlot : MaterialProperties)
+	{
+		if (!MaterialSlot.IsValid())
+		{
+			continue;
+		}
+
+		// Only show material slots on the selector
+		if (MaterialSlot->Property)
+		{
+			continue;
+		}
+
+		const UObject* Outer = MaterialSlot->OuterWeak.Get();
+		if (!IsValid(Outer))
+		{
+			continue;
+		}
+
+		UPrimitiveComponent* PrimComponent = Cast<UPrimitiveComponent>(MaterialSlot->OuterWeak.Get());
+
+		if (!PrimComponent)
+		{
+			continue;
+		}
+
+		if (Outer != CurrentOuter)
+		{
+			ListOuter->AddSlot()
+				.AutoHeight()
+				.Padding(0.f, CurrentOuter == nullptr ? 0.f : 10.f, 0.f, 5.f)
+				[
+					SNew(STextBlock)
+					.TextStyle(FDynamicMaterialEditorStyle::Get(), "ComponentNameBig")
+					.Text(FText::FromString(Outer->GetName()))
+				];
+
+			CurrentOuter = Outer;
+		}
+
+		TWeakPtr<FDMObjectMaterialProperty> MaterialSlotWeak = MaterialSlot;
+
+		constexpr int32 ThumbnailSize = 48;
+		TSharedRef<FAssetThumbnail> Thumbnail = MakeShared<FAssetThumbnail>(PrimComponent->GetMaterial(MaterialSlot->Index), ThumbnailSize, ThumbnailSize, UThumbnailManager::Get().GetSharedThumbnailPool());
+
+		FAssetThumbnailConfig ThumbnailConfig;
+		ThumbnailConfig.GenericThumbnailSize = ThumbnailSize;
+
+		ListOuter->AddSlot()
+			.AutoHeight()
+			.HAlign(EHorizontalAlignment::HAlign_Left)
+			[
+				SNew(SHorizontalBox)
+
+				+ SHorizontalBox::Slot()
+				.AutoWidth()
+				.Padding(0.f, 5.f, 5.f, 5.f)
+				.VAlign(EVerticalAlignment::VAlign_Center)
+				[
+					Thumbnail->MakeThumbnailWidget(ThumbnailConfig)
+				]
+
+				+ SHorizontalBox::Slot()
+				.AutoWidth()
+				.Padding(0.f, 5.f, 0.5, 5.f)
+				[
+					SNew(SVerticalBox)
+
+					+ SVerticalBox::Slot()
+					.AutoHeight()
+					.Padding(0.f, 5.f, 0.f, 5.f)
+					[
+						SNew(STextBlock)
+						.Text(MaterialSlot->GetPropertyName(true))
+						.Font(IDetailLayoutBuilder::GetDetailFont())
+					]
+
+					+ SVerticalBox::Slot()
+					.AutoHeight()
+					.Padding(0.f, 0.f, 0.f, 5.f)
+					[
+						SNew(SButton)
+						.ContentPadding(FMargin(2.f, 2.f, 2.f, 2.f))
+						.OnClicked(this, &SDMEditor::OnCreateMaterialButtonClicked, MaterialSlotWeak)
+						[
+							SNew(STextBlock)
+							.Text(LOCTEXT("CreateMaterial", "Create Material"))
+							.Font(IDetailLayoutBuilder::GetDetailFont())
+						]
+					]
+				]
+			];
+	}
+
+	return SNew(SScrollBox)
+		.Orientation(EOrientation::Orient_Vertical)
+		+ SScrollBox::Slot()
+		[
+			ListOuter
+		];
+}
+
+BEGIN_SLATE_FUNCTION_BUILD_OPTIMIZATION
+
+TSharedRef<SWidget> SDMEditor::GetEmptyContent()
+{
+	return SNew(SBox)
+		.HAlign(HAlign_Center)
+		.Padding(5.0f, 5.0f, 5.0f, 5.0f)
+		[
+			SNew(STextBlock)
+			.Text(LOCTEXT("NoActiveMaterial", "No active Material Designer Instance."))
+		];
+}
+
+void SDMEditor::BindCommands()
+{
+	CommandList = MakeShared<FUICommandList>();
+
+	CommandList->MapAction(
+		FDynamicMaterialEditorCommands::Get().AddDefaultLayer,
+		FExecuteAction::CreateSP(this, &SDMEditor::AddNewLayer),
+		FCanExecuteAction::CreateSP(this, &SDMEditor::CanAddNewLayer)
+	);
+
+	CommandList->MapAction(
+		FDynamicMaterialEditorCommands::Get().InsertDefaultLayerAbove,
+		FExecuteAction::CreateSP(this, &SDMEditor::InsertNewLayer),
+		FCanExecuteAction::CreateSP(this, &SDMEditor::CanInsertNewLayer)
+	);
+
+	CommandList->MapAction(
+		FGenericCommands::Get().Copy,
+		FExecuteAction::CreateSP(this, &SDMEditor::CopySelectedLayer),
+		FCanExecuteAction::CreateSP(this, &SDMEditor::CanCopySelectedLayer)
+	);
+
+	CommandList->MapAction(
+		FGenericCommands::Get().Cut,
+		FExecuteAction::CreateSP(this, &SDMEditor::CutSelectedLayer),
+		FCanExecuteAction::CreateSP(this, &SDMEditor::CanCutSelectedLayer)
+	);
+
+	CommandList->MapAction(
+		FGenericCommands::Get().Paste,
+		FExecuteAction::CreateSP(this, &SDMEditor::PasteLayer),
+		FCanExecuteAction::CreateSP(this, &SDMEditor::CanPasteLayer)
+	);
+
+	CommandList->MapAction(
+		FGenericCommands::Get().Duplicate,
+		FExecuteAction::CreateSP(this, &SDMEditor::DuplicateSelectedLayer),
+		FCanExecuteAction::CreateSP(this, &SDMEditor::CanDuplicateSelectedLayer)
+	);
+
+	CommandList->MapAction(
+		FGenericCommands::Get().Delete,
+		FExecuteAction::CreateSP(this, &SDMEditor::DeleteSelectedLayer),
+		FCanExecuteAction::CreateSP(this, &SDMEditor::CanDeleteSelectedLayer)
+	);
+}
+
+TSharedRef<SWidget> SDMEditor::CreateMainLayout()
+{
+	return 
+		SNew(SVerticalBox)
+		+ SVerticalBox::Slot()
+		.AutoHeight()
+		.HAlign(HAlign_Fill)
+		.VAlign(VAlign_Top)
+		[
+			SAssignNew(Toolbar, SDMToolBar)
+			.MaterialModel(MaterialModelWeak.Get())
+			.OnSlotChanged(this, &SDMEditor::OnToolBarPropertyChanged)
+			.OnGetSettingsMenu(this, &SDMEditor::MakeToolBarSettingsMenu)
+		]
+		+ SVerticalBox::Slot()
+		.AutoHeight()
+		.HAlign(HAlign_Fill)
+		.VAlign(VAlign_Fill)
+		[
+			SAssignNew(GlobalOpacityContainer, SBox)
+			.HAlign(HAlign_Fill)
+			.VAlign(VAlign_Fill)
+			[
+				SNew(SBorder)
+				.HAlign(HAlign_Fill)
+				.VAlign(VAlign_Center)
+				.Padding(3.0f, 3.0f, 3.0f, 3.0f)
+				.BorderImage(FDynamicMaterialEditorStyle::GetBrush("Border.Bottom"))
+				.BorderBackgroundColor(FLinearColor(1, 1, 1, 0.05f))
+				[
+					CreateMaterialSettingsRow()
+				]
+			]
+		]
+		+ SVerticalBox::Slot()
+		.AutoHeight()
+		.HAlign(HAlign_Fill)
+		.VAlign(VAlign_Fill)
+		[
+			SAssignNew(GlobalOpacityContainer, SBox)
+			.HAlign(HAlign_Fill)
+			.VAlign(VAlign_Fill)
+			[
+				CreateGlobalOpacityWidget()
+			]
+		]
+		+ SVerticalBox::Slot()
+		.AutoHeight()
+		[
+			UE::DynamicMaterialEditor::bGlobalValuesEnabled && MaterialModelWeak.IsValid()
+				? CreateParametersArea()
+				: SNullWidget::NullWidget
+		]
+		+ SVerticalBox::Slot()
+		.AutoHeight()
+		.HAlign(HAlign_Fill)
+		.VAlign(VAlign_Top)
+		[
+			SAssignNew(SlotPickerContainer, SBox)
+			[
+				CreateSlotPickerWidget()
+			]
+		]
+		+ SVerticalBox::Slot()
+		.FillHeight(1.0f)
+		.HAlign(HAlign_Fill)
+		.VAlign(VAlign_Fill)
+		[
+			SAssignNew(SlotsContainer, SBox)
+			.HAlign(HAlign_Fill)
+			.VAlign(VAlign_Fill)
+			[
+				CreateSlotsWidget()
+			]
+		];
+}
+
+TSharedRef<SWidget> SDMEditor::CreateMaterialSettingsRow()
+{
+	return 
+		SNew(SHorizontalBox)
+		+ SHorizontalBox::Slot()
+		.FillWidth(1.0f)
+		.HAlign(HAlign_Fill)
+		.VAlign(VAlign_Center)
+		.Padding(5.0f, 0.0f, 0.0f, 2.0f)
+		[
+			SNew(SDMDomain)
+			.ToolTipText(LOCTEXT("MaterialDesignerInstanceDomainTooltip", "Change the Material Designer Instance Domain"))
+			.SelectedItem(this, &SDMEditor::GetSelectedDomain)
+			.OnSelectedItemChanged(this, &SDMEditor::OnDomainChanged)
+			.IsEnabled(this, &SDMEditor::CanChangeDomain)
+		]
+		+ SHorizontalBox::Slot()
+		.FillWidth(1.0f)
+		.HAlign(HAlign_Fill)
+		.VAlign(VAlign_Center)
+		.Padding(5.0f, 0.0f, 0.0f, 2.0f)
+		[
+			SNew(SDMBlendMode)
+			.ToolTipText(LOCTEXT("MaterialDesignerInstanceBlendModeTooltip", "Change the Material Designer Instance Blend Mode"))
+			.SelectedItem(this, &SDMEditor::GetSelectedBlendMode)
+			.OnSelectedItemChanged(this, &SDMEditor::OnBlendModeChanged)
+			.IsEnabled(this, &SDMEditor::CanChangeBlendType)
+		]
+		+ SHorizontalBox::Slot()
+		.AutoWidth()
+		.HAlign(HAlign_Right)
+		.VAlign(VAlign_Center)
+		.Padding(10.0f, 0.0f, 0.0f, 2.0f)
+		[
+			SNew(SHorizontalBox)
+			.ToolTipText(LOCTEXT("MaterialDesignerInstanceTypeTooltip", "Enables the material's TSR pixel animation flag."))
+			+ SHorizontalBox::Slot()
+			.AutoWidth()
+			.HAlign(HAlign_Left)
+			.VAlign(VAlign_Center)
+			[
+				SNew(SBox)
+				.Padding(0.f, 0.f, 5.f, 0.f)
+				[
+					SNew(STextBlock)
+					.Text(LOCTEXT("MaterialDesignerInstanceAnimaed", "Animated"))
+				]
+			]
+			+ SHorizontalBox::Slot()
+			.AutoWidth()
+			.HAlign(HAlign_Left)
+			.VAlign(VAlign_Center)
+			[
+				SNew(SCheckBox)
+				.IsChecked(this, &SDMEditor::IsMaterialAnimated)
+				.OnCheckStateChanged(this, &SDMEditor::OnMaterialAnimatedChanged)
+			]
+		]
+		+ SHorizontalBox::Slot()
+		.AutoWidth()
+		.HAlign(HAlign_Right)
+		.VAlign(VAlign_Center)
+		.Padding(10.f, 0.0f, 5.0f, 2.0f)
+		[
+			SNew(SHorizontalBox)
+			.ToolTipText(LOCTEXT("MaterialDesignerInstanceUnlitTooltip", "Toggle between Default Lit and Unlit for this Material Designer Instance."))
+			+ SHorizontalBox::Slot()
+			.AutoWidth()
+			.HAlign(HAlign_Left)
+			.VAlign(VAlign_Center)
+			.Padding(0.0f, 0.0f, 5.0f, 0.0f)
+			[
+				SNew(STextBlock)
+				.Text(LOCTEXT("MaterialDesignerInstanceUnlit", "Unlit"))
+			]
+			+ SHorizontalBox::Slot()
+			.AutoWidth()
+			.HAlign(HAlign_Left)
+			.VAlign(VAlign_Center)
+			[
+				SNew(SCheckBox)
+				.IsChecked(this, &SDMEditor::IsMaterialUnlit)
+				.OnCheckStateChanged(this, &SDMEditor::OnMaterialUnlitChanged)
+				.IsEnabled(this, &SDMEditor::CanChangeMaterialShadingModel)
+			]
+		];
+}
+
+TSharedRef<SWidget> SDMEditor::CreateGlobalOpacityWidget()
+{
+	UDMMaterialValueFloat1* OpacityValue = nullptr;
+
+	if (UDynamicMaterialModel* MaterialModel = MaterialModelWeak.Get())
+	{
+		OpacityValue = MaterialModel->GetGlobalOpacityValue();
+	}
+
+	if (!OpacityValue)
+	{
+		return SNullWidget::NullWidget;
+	}
+
+	TSharedRef<SDMPropertyEdit> GlobalOpacityWidget = SNew(SDMPropertyEditOpacity, SharedThis(this), OpacityValue);
+	GlobalOpacityWidget->SetEnabled(TAttribute<bool>::CreateSP(this, &SDMEditor::IsGlobalOpacityEnabled));
+
+	TSharedRef<SWidget> GlobalOpacityButtons = SDMComponentEdit::CreateExtensionButtons(SharedThis(this), OpacityValue, UDMMaterialValue::ValueName, true, FSimpleDelegate());
+	GlobalOpacityButtons->SetEnabled(TAttribute<bool>::CreateSP(this, &SDMEditor::IsGlobalOpacityEnabled));
+
+	TSharedPtr<SWidget> RowLabel;
+
+	TSharedRef<SWidget> Row =
+		SNew(SHorizontalBox)
+		+ SHorizontalBox::Slot()
+		.AutoWidth()
+		.HAlign(HAlign_Left)
+		.VAlign(VAlign_Top)
+		.Padding(10.0f, 5.0f)
+		[
+			SAssignNew(RowLabel, SBox)
+			.Padding(0.f, 0.f, 5.f, 0.f)
+			[
+				SNew(STextBlock)
+				.Text(LOCTEXT("GlobalOpacity", "Global Opacity"))
+			]
+		]
+		+ SHorizontalBox::Slot()
+		.FillWidth(1.0f)
+		.HAlign(HAlign_Right)
+		.VAlign(VAlign_Top)
+		[
+			GlobalOpacityWidget
+		]
+		+ SHorizontalBox::Slot()
+		.AutoWidth()
+		.HAlign(HAlign_Left)
+		.VAlign(VAlign_Top)
+		[
+			GlobalOpacityButtons
+		];
+
+	RowLabel->SetOnMouseButtonDown(FPointerEventHandler::CreateStatic(&SDMPropertyEdit::CreateRightClickDetailsMenu, GlobalOpacityWidget.ToWeakPtr()));
+
+	return Row;
+}
+
+TSharedRef<SWidget> SDMEditor::CreateParametersArea()
+{
+	return 
+		SNew(SExpandableArea)
+		.InitiallyCollapsed(false)
+		.HeaderPadding(FMargin(3.0f, 5.0f, 3.0f, 5.0f))
+		.HeaderContent()
+		[
+			SNew(SBox)
+			.VAlign(EVerticalAlignment::VAlign_Center)
+			[
+				SNew(STextBlock)
+				.Text(LOCTEXT("MaterialParameters", "Material Parameters"))
+				.Font(IDetailLayoutBuilder::GetDetailFontBold())
+			]
+		]
+		.BodyContent()
+		[
+			SAssignNew(ParametersWidget, SDMMaterialParameters, MaterialModelWeak)
+		];
+}
+
+TSharedRef<SWidget> SDMEditor::CreateSlotPickerWidget()
+{
+	TSharedRef<SHorizontalBox> SlotSelector = SNew(SHorizontalBox);
+
+	UDynamicMaterialModelEditorOnlyData* ModelEditorOnlyData = UDynamicMaterialModelEditorOnlyData::Get(MaterialModelWeak);
+
+	if (!ModelEditorOnlyData)
+	{
+		return SlotSelector;
+	}
+
+	const TArray<UDMMaterialSlot*>& Slots = ModelEditorOnlyData->GetSlots();
+
+	if (Slots.IsEmpty())
+	{
+		return SlotSelector;
+	}
+
+	if (Slots.IsValidIndex(0))
+	{
+		SlotSelector->AddSlot()
+			[
+				SNew(SButton)
+				.HAlign(EHorizontalAlignment::HAlign_Center)
+				.ButtonColorAndOpacity(this, &SDMEditor::GetRGBButtonColorAndOpacity_HasSlot)
+				.OnClicked(this, &SDMEditor::OnRGBButtonClicked_HasSlot)
+				[
+					SNew(STextBlock)
+					.Text(LOCTEXT("RGBSlot", "RGB"))
+				]
+			];		
+	}
+	else
+	{
+		SlotSelector->AddSlot()
+			[
+				SNew(SButton)
+				.HAlign(EHorizontalAlignment::HAlign_Center)
+				.OnClicked(this, &SDMEditor::OnRGBButtonClicked_NoSlot)
+				[
+					SNew(STextBlock)
+					.Text(LOCTEXT("AddRGBStack", "RGB"))
+				]
+			];
+	}
+
+	if (Slots.IsValidIndex(1))
+	{
+		SlotSelector->AddSlot()
+			[
+				SNew(SButton)
+				.HAlign(EHorizontalAlignment::HAlign_Center)
+				.IsEnabled(this, &SDMEditor::GetOpacityButtonEnabled_HasSlot)
+				.ButtonColorAndOpacity(this, &SDMEditor::GetOpacityButtonColorAndOpacity_HasSlot)
+				.OnClicked(this, &SDMEditor::OnOpacityButtonClicked_HasSlot)
+				[
+					SNew(STextBlock)
+					.Text(LOCTEXT("OpacitySlot", "Opacity"))
+				]
+			];		
+	}
+	else
+	{
+		SlotSelector->AddSlot()
+			[
+				SNew(SButton)
+				.HAlign(EHorizontalAlignment::HAlign_Center)
+				.IsEnabled(this, &SDMEditor::GetOpacityButtonEnabled_NoSlot)
+				.OnClicked(this, &SDMEditor::OnOpacityButtonClicked_NoSlot)
+				[
+					SNew(STextBlock)
+					.Text(LOCTEXT("AddOpacity", "Opacity"))
+				]
+			];
+	}
+
+	return SlotSelector;
+}
+
+END_SLATE_FUNCTION_BUILD_OPTIMIZATION
+
+TSharedPtr<SDMSlot> SDMEditor::GetSlotWidget(UDMMaterialSlot* Slot) const
+{
+	if (ensure(IsValid(Slot)))
+	{
+		for (const TSharedPtr<SDMSlot>& SlotWidget : SlotWidgets)
+		{
+			if (SlotWidget.IsValid() && SlotWidget->GetSlot() == Slot)
+			{
+				return SlotWidget;
+			}
+		}
+	}
+
+	return nullptr;
+}
+
+void SDMEditor::RefreshGlobalOpacitySlider()
+{
+	if (GlobalOpacityContainer.IsValid())
+	{
+		GlobalOpacityContainer->SetContent(CreateGlobalOpacityWidget());
+	}
+}
+
+void SDMEditor::RefreshParametersList()
+{
+	if (ParametersWidget.IsValid())
+	{
+		ParametersWidget->RefreshWidgets();
+	}
+}
+
+void SDMEditor::RefreshSlotPickerList()
+{
+	if (SlotPickerContainer.IsValid())
+	{
+		SlotPickerContainer->SetContent(CreateSlotPickerWidget());
+	}
+}
+
+void SDMEditor::RefreshSlotsList()
+{
+	if (SlotsContainer.IsValid())
+	{
+		SlotsContainer->SetContent(CreateSlotsWidget());
+	}
+}
+
+void SDMEditor::Tick(const FGeometry& AllottedGeometry, const double InCurrentTime, const float InDeltaTime)
+{
+	SCompoundWidget::Tick(AllottedGeometry, InCurrentTime, InDeltaTime);
+
+	if (!MaterialModelWeak.IsValid())
+	{
+		return;
+	}
+
+	bool bValidModel = false;
+
+	if (ObjectProperty.IsValid())
+	{
+		bValidModel = IsValid(ObjectProperty.GetMaterialModel());
+	}
+	else if (UDynamicMaterialModel* MaterialModel = MaterialModelWeak.Get())
+	{
+		auto CheckValidity = [MaterialModel]()
+		{
+			if (UWorld* World = MaterialModel->GetWorld())
+			{
+				if (UDMWorldSubsystem* WorldSubsystem = World->GetSubsystem<UDMWorldSubsystem>())
+				{
+					// ExecuteIfBound doesn't work with return values
+					if (WorldSubsystem->GetIsValidDelegate().IsBound()
+						&& WorldSubsystem->GetIsValidDelegate().Execute(MaterialModel) == false)
+					{
+						return false;
+					}
+				}
+			}
+
+			UActorComponent* ComponentOuter = MaterialModel->GetTypedOuter<UActorComponent>();
+			if (ComponentOuter && !IsValid(ComponentOuter))
+			{
+				return false;
+			}
+
+			AActor* ActorOuter = MaterialModel->GetTypedOuter<AActor>();
+			if (ActorOuter && !IsValid(ActorOuter))
+			{
+				return false;
+			}
+
+			UPackage* PackageOuter = MaterialModel->GetPackage();
+			if (PackageOuter && !IsValid(PackageOuter))
+			{
+				return false;
+			}
+
+			return true;
+		};
+
+		bValidModel = CheckValidity();
+	}
+
+	if (!bValidModel)
+	{
+		ClearEditor();
+	}
+	else if (FDynamicMaterialModule::AreUObjectsSafe())
+	{
+		if (Toolbar.IsValid() && Toolbar->GetMaterialModel() != MaterialModelWeak)
+		{
+			Toolbar->SetMaterialModel(MaterialModelWeak.Get());
+		}
+		else
+		{
+			bool bHasInvalidSlotWidget = false;
+			for (const TWeakPtr<SDMSlot> SlotWeak : SlotWidgets)
+			{
+				TSharedPtr<SDMSlot> Slot = SlotWeak.Pin();
+				if (Slot.IsValid() && !Slot->CheckValidity())
+				{
+					bHasInvalidSlotWidget = true;
+					break;
+				}
+			}
+			if (bHasInvalidSlotWidget)
+			{
+				RefreshSlotsList();
+			}
+		}
+	}
+}
+
+FReply SDMEditor::OnAddSlotButtonClicked()
+{
+	if (!MaterialModelWeak.IsValid())
+	{
+		return FReply::Unhandled();
+	}
+
+	UDynamicMaterialModelEditorOnlyData* ModelEditorOnlyData = UDynamicMaterialModelEditorOnlyData::Get(MaterialModelWeak);
+
+	if (!ModelEditorOnlyData)
+	{
+		return FReply::Unhandled();
+	}
+
+	ModelEditorOnlyData->AddSlot();
+
+	return FReply::Handled();
+}
+
+void SDMEditor::OnMaterialBuilt(UDynamicMaterialModel* InMaterialModel)
+{
+}
+
+void SDMEditor::OnValuesUpdated(UDynamicMaterialModel* InMaterialModel)
+{
+	RefreshParametersList();
+}
+
+void SDMEditor::OnSlotsUpdated(UDynamicMaterialModel* InMaterialModel)
+{
+	RefreshSlotsList();
+}
+
+bool SDMEditor::IsGlobalOpacityEnabled() const
+{
+	if (UDynamicMaterialModelEditorOnlyData* ModelEditorOnlyData = UDynamicMaterialModelEditorOnlyData::Get(MaterialModelWeak))
+	{
+		return ModelEditorOnlyData->GetBlendMode() == BLEND_Translucent || ModelEditorOnlyData->GetBlendMode() == BLEND_Masked;
+	}
+
+	return false;
+}
+
+FReply SDMEditor::OnRGBButtonClicked_NoSlot()
+{
+	if (UDynamicMaterialModelEditorOnlyData* ModelEditorOnlyData = UDynamicMaterialModelEditorOnlyData::Get(MaterialModelWeak))
+	{
+		if (ModelEditorOnlyData->AddSlot())
+		{
+			RefreshSlotPickerList();
+			SetActiveSlotIndex(0);
+		}
+	}
+
+	return FReply::Handled();
+}
+
+FSlateColor SDMEditor::GetRGBButtonColorAndOpacity_HasSlot() const
+{
+	if (GetActiveSlotIndex() == 0)
+	{
+		return FSlateColor(FStyleColors::AccentBlue.GetSpecifiedColor());
+	}
+
+	return FSlateColor(FStyleColors::AccentGray.GetSpecifiedColor());
+}
+
+FReply SDMEditor::OnRGBButtonClicked_HasSlot()
+{
+	SetActiveSlotIndex(0);
+
+	return FReply::Handled();
+}
+
+FReply SDMEditor::OnOpacityButtonClicked_NoSlot()
+{
+	if (UDynamicMaterialModelEditorOnlyData* ModelEditorOnlyData = UDynamicMaterialModelEditorOnlyData::Get(MaterialModelWeak))
+	{
+		if (UDMMaterialSlot* NewSlot = ModelEditorOnlyData->AddSlot())
+		{
+			if (UDMMaterialLayerObject* Layer = NewSlot->GetLayer(0))
+			{
+				FScopedTransaction Transaction(LOCTEXT("AddOpacitySlot", "Material Designer Add Opacity Slot"));
+				Layer->Modify();
+
+				switch (ModelEditorOnlyData->GetBlendMode())
+				{
+					case EBlendMode::BLEND_Masked:
+						Layer->SetMaterialProperty(EDMMaterialPropertyType::OpacityMask);
+						break;
+
+					default:
+						Layer->SetMaterialProperty(EDMMaterialPropertyType::Opacity);
+						break;
+				}
+			}
+
+			RefreshSlotPickerList();
+			SetActiveSlotIndex(1);
+		}
+	}
+
+	return FReply::Handled();
+}
+
+bool SDMEditor::GetOpacityButtonEnabled_NoSlot() const
+{
+	if (UDynamicMaterialModelEditorOnlyData* ModelEditorOnlyData = UDynamicMaterialModelEditorOnlyData::Get(MaterialModelWeak))
+	{
+		return ModelEditorOnlyData->GetDomain() != EMaterialDomain::MD_PostProcess
+			&& ModelEditorOnlyData->GetBlendMode() != EBlendMode::BLEND_Opaque;
+	}
+
+	return false;
+}
+
+FSlateColor SDMEditor::GetOpacityButtonColorAndOpacity_HasSlot() const
+{
+	if (GetActiveSlotIndex() == 1)
+	{
+		return FSlateColor(FStyleColors::AccentBlue.GetSpecifiedColor());
+	}
+
+	return FSlateColor(FStyleColors::AccentGray.GetSpecifiedColor());
+}
+
+FReply SDMEditor::OnOpacityButtonClicked_HasSlot()
+{
+	SetActiveSlotIndex(1);
+
+	return FReply::Handled();
+}
+
+bool SDMEditor::GetOpacityButtonEnabled_HasSlot() const
+{
+	if (UDynamicMaterialModelEditorOnlyData* ModelEditorOnlyData = UDynamicMaterialModelEditorOnlyData::Get(MaterialModelWeak))
+	{
+		return ModelEditorOnlyData->GetBlendMode() != EBlendMode::BLEND_Opaque;
+	}
+
+	return false;
+}
+
+FReply SDMEditor::OnCreateMaterialButtonClicked(TWeakPtr<FDMObjectMaterialProperty> InMaterialProperty)
+{
+	if (TSharedPtr<FDMObjectMaterialProperty> MaterialProperty = InMaterialProperty.Pin())
+	{
+		UDynamicMaterialModel* NewModel = UDMBlueprintFunctionLibrary::CreateDynamicMaterialInObject(*MaterialProperty);
+
+		if (NewModel)
+		{
+			if (Toolbar.IsValid())
+			{
+				Toolbar->SetMaterialModel(NewModel);
+			}
+		}
+	}
+
+	return FReply::Handled();
+}
+
+void SDMEditor::OnToolBarPropertyChanged(TSharedPtr<FDMObjectMaterialProperty> InNewSelectedProperty)
+{
+	if (InNewSelectedProperty.IsValid())
+	{
+		SetMaterialModel(InNewSelectedProperty->GetMaterialModel());
+	}
+}
+
+TSharedRef<SWidget> SDMEditor::MakeToolBarSettingsMenu()
+{
+	return FDMToolBarMenus::MakeEditorLayoutMenu(SharedThis(this));
+}
+
+void SDMEditor::OnSettingsChanged(const FPropertyChangedEvent& InPropertyChangedEvent)
+{
+	SetMaterialModel(MaterialModelWeak.Get());
+}
+
+bool SDMEditor::CanAddNewLayer() const
+{
+	return SlotWidgets.IsValidIndex(ActiveSlotIndex)
+		&& MaterialModelWeak.IsValid();
+}
+
+void SDMEditor::AddNewLayer()
+{
+	if (SlotWidgets.IsValidIndex(ActiveSlotIndex))
+	{
+		SlotWidgets[ActiveSlotIndex]->AddNewLayer_Expression(
+			TSubclassOf<UDMMaterialStageExpression>(UDMMaterialStageExpressionTextureSample::StaticClass()),
+			EDMMaterialLayerStage::All
+		);
+
+		SlotWidgets[ActiveSlotIndex]->InvalidateMainWidget();
+	}
+}
+
+bool SDMEditor::CanInsertNewLayer() const
+{
+	return SlotWidgets.IsValidIndex(ActiveSlotIndex)
+		&& SlotWidgets[ActiveSlotIndex]->GetSelectedLayerIndices().Num() == 1
+		&& MaterialModelWeak.IsValid();
+}
+
+void SDMEditor::InsertNewLayer()
+{
+	if (SlotWidgets.IsValidIndex(ActiveSlotIndex))
+	{
+		if (UDMMaterialLayerObject* SelectedLayer = SlotWidgets[ActiveSlotIndex]->GetSelectedLayer())
+		{
+			if (UDMMaterialSlot* Slot = SelectedLayer->GetSlot())
+			{
+				// Added here because stuff is done after the layer is added
+				FScopedTransaction Transaction(LOCTEXT("InsertNewLayer", "Material Designer Insert Layer"));
+				Slot->Modify();
+
+				SlotWidgets[ActiveSlotIndex]->AddNewLayer_Expression(
+					TSubclassOf<UDMMaterialStageExpression>(UDMMaterialStageExpressionTextureSample::StaticClass()),
+					EDMMaterialLayerStage::All
+				);
+
+				Slot->MoveLayerAfter(Slot->GetLayers().Last(), SelectedLayer);
+
+				SlotWidgets[ActiveSlotIndex]->InvalidateMainWidget();
+			}
+		}
+	}
+}
+
+bool SDMEditor::CanCopySelectedLayer() const
+{
+	return SlotWidgets.IsValidIndex(ActiveSlotIndex)
+		&& SlotWidgets[ActiveSlotIndex]->GetSelectedLayerIndices().Num() == 1
+		&& MaterialModelWeak.IsValid();
+}
+
+void SDMEditor::CopySelectedLayer()
+{
+	if (SlotWidgets.IsValidIndex(ActiveSlotIndex)
+		&& SlotWidgets[ActiveSlotIndex]->GetSelectedLayerIndices().Num() == 1)
+	{
+		if (const UDMMaterialLayerObject* Layer = SlotWidgets[ActiveSlotIndex]->GetSelectedLayer())
+		{
+			FPlatformApplicationMisc::ClipboardCopy(*Layer->SerializeToString());
+		}
+	}
+}
+
+bool SDMEditor::CanCutSelectedLayer() const
+{
+	return CanCopySelectedLayer() && CanDeleteSelectedLayer();
+}
+
+void SDMEditor::CutSelectedLayer()
+{
+	CopySelectedLayer();
+	DeleteSelectedLayer();
+}
+
+bool SDMEditor::CanPasteLayer() const
+{
+	if (!MaterialModelWeak.IsValid())
+	{
+		return false;
+	}
+
+	FString ClipboardContent;
+	FPlatformApplicationMisc::ClipboardPaste(ClipboardContent);
+
+	return !ClipboardContent.IsEmpty();
+}
+
+void SDMEditor::PasteLayer()
+{
+	if (UDynamicMaterialModelEditorOnlyData* ModelEditorOnlyData = UDynamicMaterialModelEditorOnlyData::Get(MaterialModelWeak))
+	{
+		if (UDMMaterialSlot* Slot = ModelEditorOnlyData->GetSlot(ActiveSlotIndex))
+		{
+			FString SerializeString;
+			FPlatformApplicationMisc::ClipboardPaste(SerializeString);
+
+			if (UDMMaterialLayerObject* PastedLayer = UDMMaterialLayerObject::DeserializeFromString(Slot, SerializeString))
+			{
+				FScopedTransaction Transaction(LOCTEXT("PasteLayer", "Material Designer Paste Layer"));
+				Slot->Modify();
+				Slot->PasteLayer(PastedLayer);
+
+				SlotWidgets[ActiveSlotIndex]->InvalidateMainWidget();
+			}
+		}
+	}
+}
+
+bool SDMEditor::CanDuplicateSelectedLayer() const
+{
+	// There's no "can add" check, so only copy is tested.
+	return CanCopySelectedLayer();
+}
+
+void SDMEditor::DuplicateSelectedLayer()
+{
+	FString PastedText;
+	FPlatformApplicationMisc::ClipboardPaste(PastedText);
+
+	// Added here to set the transaction description
+	FScopedTransaction Transaction(LOCTEXT("DuplicateLayer", "Material Designer Duplicate Layer"));
+
+	CopySelectedLayer();
+	PasteLayer();
+
+	FPlatformApplicationMisc::ClipboardCopy(*PastedText);
+}
+
+bool SDMEditor::CanDeleteSelectedLayer() const
+{
+	if (SlotWidgets.IsValidIndex(ActiveSlotIndex)
+		&& SlotWidgets[ActiveSlotIndex]->GetSelectedLayerIndices().IsEmpty() == false)
+	{
+		return SlotWidgets[ActiveSlotIndex]->GetLayerRowsButtonsCanRemove();
+	}
+
+	return false;
+}
+
+void SDMEditor::DeleteSelectedLayer()
+{
+	if (SlotWidgets.IsValidIndex(ActiveSlotIndex)
+		&& SlotWidgets[ActiveSlotIndex]->GetSelectedLayerIndices().IsEmpty() == false)
+	{
+		SlotWidgets[ActiveSlotIndex]->OnLayerRowButtonsRemoveClicked();
+	}
+}
+
+FReply SDMEditor::OnKeyDown(const FGeometry& MyGeometry, const FKeyEvent& InKeyEvent)
+{
+	if (CommandList.IsValid())
+	{
+		if (CommandList->ProcessCommandBindings(InKeyEvent))
+		{
+			return FReply::Handled();
+		}
+	}
+
+	// We accept the delete key bind, so we don't want this accidentally deleting actors and such.
+	// Always return handled to stop the event bubbling.
+	const TArray<TSharedRef<const FInputChord>> DeleteChords = {
+		FGenericCommands::Get().Delete->GetActiveChord(EMultipleKeyBindingIndex::Primary),
+		FGenericCommands::Get().Delete->GetActiveChord(EMultipleKeyBindingIndex::Secondary)
+	};
+
+	for (const TSharedRef<const FInputChord>& DeleteChord : DeleteChords)
+	{
+		if (DeleteChord->Key == InKeyEvent.GetKey())
+		{
+			return FReply::Handled();
+		}
+	}
+
+	return FReply::Unhandled();
+}
+
+void SDMEditor::PostUndo(bool bSuccess)
+{
+	OnUndo();
+}
+
+void SDMEditor::PostRedo(bool bSuccess)
+{
+	OnUndo();
+}
+
+TSharedRef<SWidget> SDMEditor::CreateSlotsWidget()
+{
+	SlotWidgets.Empty();
+
+	auto CreateEmptySlotsContent = []() -> TSharedRef<SBox>
+	{
+		return 
+			SNew(SBox)
+			.HAlign(EHorizontalAlignment::HAlign_Center)
+			.VAlign(EVerticalAlignment::VAlign_Center)
+			.Padding(10.f, 5.f, 10.f, 5.f)
+			[
+				SNew(STextBlock)
+				.Text(LOCTEXT("SlotsContent", "Slots Content"))
+				.Font(IDetailLayoutBuilder::GetDetailFont())
+			];
+	};
+	
+	UDynamicMaterialModelEditorOnlyData* ModelEditorOnlyData = UDynamicMaterialModelEditorOnlyData::Get(MaterialModelWeak);
+
+	if (!ModelEditorOnlyData)
+	{
+		return CreateEmptySlotsContent();
+	}
+
+	const TArray<UDMMaterialSlot*>& Slots = ModelEditorOnlyData->GetSlots();
+
+	if (Slots.IsEmpty())
+	{
+		return CreateEmptySlotsContent();
+	}
+
+	for (int32 SlotIdx = 0; SlotIdx < Slots.Num(); ++SlotIdx)
+	{
+		UDMMaterialSlot* Slot = Slots[SlotIdx];
+
+		TSharedRef<SDMSlot> SlotWidget =
+			SNew(SDMSlot, SharedThis(this), Slot)
+			.SlotPreviewSize_Lambda([]() { return UDynamicMaterialEditorSettings::Get()->SlotPreviewSize; })
+			.LayerPreviewSize_Lambda([]() { return UDynamicMaterialEditorSettings::Get()->LayerPreviewSize; });
+
+		SlotWidgets.Add(SlotWidget);
+	}
+
+	return SlotWidgets.IsValidIndex(ActiveSlotIndex)
+		? SlotWidgets[ActiveSlotIndex].ToSharedRef()
+		: SNullWidget::NullWidget;
+}
+
+EMaterialDomain SDMEditor::GetSelectedDomain() const
+{
+	if (UDynamicMaterialModelEditorOnlyData* ModelEditorOnlyData = UDynamicMaterialModelEditorOnlyData::Get(MaterialModelWeak))
+	{
+		return ModelEditorOnlyData->GetDomain();
+	}
+
+	return EMaterialDomain::MD_Surface;
+}
+
+void SDMEditor::OnDomainChanged(const EMaterialDomain InDomain)
+{
+	if (UDynamicMaterialModelEditorOnlyData* ModelEditorOnlyData = UDynamicMaterialModelEditorOnlyData::Get(MaterialModelWeak))
+	{
+		FScopedTransaction Transaction(LOCTEXT("ChangeDomain", "Material Designer Change Domain"));
+		ModelEditorOnlyData->Modify();
+		ModelEditorOnlyData->SetDomain(InDomain);
+
+		RefreshSlotPickerList();
+
+		if (InDomain == EMaterialDomain::MD_PostProcess)
+		{
+			SetActiveSlotIndex(0);
+		}
+	}
+}
+
+bool SDMEditor::CanChangeDomain() const
+{
+	return true;
+}
+
+EBlendMode SDMEditor::GetSelectedBlendMode() const
+{
+	if (UDynamicMaterialModelEditorOnlyData* ModelEditorOnlyData = UDynamicMaterialModelEditorOnlyData::Get(MaterialModelWeak))
+	{
+		return ModelEditorOnlyData->GetBlendMode();
+	}
+
+	return EBlendMode::BLEND_Opaque;
+}
+
+void SDMEditor::OnBlendModeChanged(const EBlendMode InBlendMode)
+{
+	if (UDynamicMaterialModelEditorOnlyData* ModelEditorOnlyData = UDynamicMaterialModelEditorOnlyData::Get(MaterialModelWeak))
+	{
+		FScopedTransaction Transaction(LOCTEXT("ChangeBlendMode", "Material Designer Change Blend Mode"));
+		ModelEditorOnlyData->Modify();
+		ModelEditorOnlyData->SetBlendMode(InBlendMode);
+
+		RefreshSlotPickerList();
+
+		if (InBlendMode == EBlendMode::BLEND_Opaque)
+		{
+			SetActiveSlotIndex(0);
+		}
+	}
+}
+
+bool SDMEditor::CanChangeBlendType() const
+{
+	if (UDynamicMaterialModelEditorOnlyData* ModelEditorOnlyData = UDynamicMaterialModelEditorOnlyData::Get(MaterialModelWeak))
+	{
+		return ModelEditorOnlyData->GetDomain() != EMaterialDomain::MD_PostProcess;
+	}
+
+	return false;
+}
+
+ECheckBoxState SDMEditor::IsMaterialUnlit() const
+{
+	if (UDynamicMaterialModelEditorOnlyData* ModelEditorOnlyData = UDynamicMaterialModelEditorOnlyData::Get(MaterialModelWeak))
+	{
+		return ModelEditorOnlyData->GetShadingModel() == EDMMaterialShadingModel::Unlit
+			? ECheckBoxState::Checked : ECheckBoxState::Unchecked;
+	}
+
+	return ECheckBoxState::Unchecked;
+}
+
+void SDMEditor::OnMaterialUnlitChanged(const ECheckBoxState InNewCheckState)
+{
+	if (UDynamicMaterialModelEditorOnlyData* ModelEditorOnlyData = UDynamicMaterialModelEditorOnlyData::Get(MaterialModelWeak))
+	{
+		FScopedTransaction Transaction(LOCTEXT("ChangeShadingModel", "Material Designer Change Shading Model"));
+		ModelEditorOnlyData->Modify();
+		ModelEditorOnlyData->SetShadingModel(InNewCheckState == ECheckBoxState::Checked
+			? EDMMaterialShadingModel::Unlit : EDMMaterialShadingModel::DefaultLit);
+	}
+}
+
+bool SDMEditor::CanChangeMaterialShadingModel() const
+{
+	if (UDynamicMaterialModelEditorOnlyData* ModelEditorOnlyData = UDynamicMaterialModelEditorOnlyData::Get(MaterialModelWeak))
+	{
+		return ModelEditorOnlyData->GetDomain() != EMaterialDomain::MD_PostProcess;
+	}
+
+	return false;
+}
+
+ECheckBoxState SDMEditor::IsMaterialAnimated() const
+{
+	if (UDynamicMaterialModelEditorOnlyData* ModelEditorOnlyData = UDynamicMaterialModelEditorOnlyData::Get(MaterialModelWeak))
+	{
+		return ModelEditorOnlyData->IsPixelAnimationFlagSet()
+			? ECheckBoxState::Checked : ECheckBoxState::Unchecked;
+	}
+
+	return ECheckBoxState::Unchecked;
+}
+
+void SDMEditor::OnMaterialAnimatedChanged(const ECheckBoxState InNewCheckState)
+{
+	if (UDynamicMaterialModelEditorOnlyData* ModelEditorOnlyData = UDynamicMaterialModelEditorOnlyData::Get(MaterialModelWeak))
+	{
+		FScopedTransaction Transaction(LOCTEXT("ToggleAnimated", "Material Designer Toggle Animated"));
+		ModelEditorOnlyData->Modify();
+		ModelEditorOnlyData->SetPixelAnimationFlag(InNewCheckState == ECheckBoxState::Checked);
+	}
+}
+
+void SDMEditor::OnUndo()
+{
+	UDynamicMaterialModel* MaterialModel = MaterialModelWeak.Get();
+
+	if (!IsValid(MaterialModel))
+	{
+		return;
+	}
+
+	if (UDynamicMaterialModelEditorOnlyData* ModelEditorOnlyData = UDynamicMaterialModelEditorOnlyData::Get(MaterialModelWeak))
+	{
+		if (ModelEditorOnlyData->GetBlendMode() == BLEND_Opaque)
+		{
+			SetActiveSlotIndex(0);
+		}
+	}
+}
+
+#undef LOCTEXT_NAMESPACE
diff --git a/Engine/Plugins/Experimental/DynamicMaterial/Source/DynamicMaterialEditor/Private/Slate/SDMEditor.h b/Engine/Plugins/Experimental/DynamicMaterial/Source/DynamicMaterialEditor/Private/Slate/SDMEditor.h
new file mode 100644
--- /dev/null
+++ b/Engine/Plugins/Experimental/DynamicMaterial/Source/DynamicMaterialEditor/Private/Slate/SDMEditor.h
@@ -0,0 +1,190 @@
+// Copyright Epic Games, Inc. All Rights Reserved.
+
+#pragma once
+
+#include "Components/DMMaterialStage.h"
+#include "DMEDefs.h"
+#include "DMObjectMaterialProperty.h"
+#include "EditorUndoClient.h"
+#include "Widgets/SCompoundWidget.h"
+
+class FUICommandList;
+class IDetailTreeNode;
+class IPropertyHandle;
+class IPropertyRowGenerator;
+class IStructureDetailsView;
+class SBox;
+class SDMMaterialParameters;
+class SDMPropertyEdit;
+class SDMSlot;
+class SDMToolBar;
+class SDockTab;
+class SExpandableArea;
+class SScrollBox;
+class SWidgetSwitcher;
+class UDMMaterialSlot;
+class UDMMaterialStage;
+class UDMMaterialStageExpression;
+class UDMMaterialValueFloat1;
+class UDynamicMaterialModel;
+enum class ECheckBoxState : uint8;
+enum class EDMExpressionMenu : uint8;
+enum EMaterialDomain : int;
+struct FDMActorMaterialSlot;
+
+class SDMEditor : public SCompoundWidget, public FSelfRegisteringEditorUndoClient
+{
+	SLATE_DECLARE_WIDGET(SDMEditor, SCompoundWidget);
+
+	SLATE_BEGIN_ARGS(SDMEditor) {}
+	SLATE_END_ARGS()
+
+public:
+	static TSharedRef<FAssetThumbnailPool> GetThumbnailPool();
+
+	static TSharedRef<SWidget> GetEmptyContent();
+
+	static FDMPropertyHandle GetPropertyHandle(const SWidget* InOwner, UDMMaterialComponent* InComponent, const FName& InPropertyName);
+	static void ClearPropertyHandles(const SWidget* InOwner);
+
+	void Construct(const FArguments& InArgs, TWeakObjectPtr<UDynamicMaterialModel> InModelWeak);
+	virtual ~SDMEditor() override;
+
+	UDynamicMaterialModel* GetMaterialModel() const { return MaterialModelWeak.Get(); }
+	void SetMaterialModel(UDynamicMaterialModel* InMaterialModel);
+
+	const FDMObjectMaterialProperty& GetMaterialObjectProperty() const { return ObjectProperty; }
+	void SetMaterialObjectProperty(const FDMObjectMaterialProperty& InObjectProperty);
+
+	void SetMaterialActor(AActor* InActor);
+
+	const TArray<TSharedPtr<SDMSlot>>& GetSlotWidgets() const { return SlotWidgets; }
+	TSharedPtr<SDMSlot> GetSlotWidget(UDMMaterialSlot* Slot) const;
+
+	void RefreshGlobalOpacitySlider();
+	void RefreshParametersList();
+	void RefreshSlotPickerList();
+	void RefreshSlotsList();
+
+	//~ Begin SWidget
+	virtual bool SupportsKeyboardFocus() const override { return true; }
+	virtual void Tick(const FGeometry& AllottedGeometry, const double InCurrentTime, const float InDeltaTime) override;
+	//~ End SWidget
+
+	void ClearEditor();
+
+	int32 GetActiveSlotIndex() const { return ActiveSlotIndex; }
+	void SetActiveSlotIndex(int InSlotIndex);
+
+	const TSharedPtr<FUICommandList>& GetCommandList() const { return CommandList; }
+
+	bool CanAddNewLayer() const;
+	void AddNewLayer();
+
+	bool CanInsertNewLayer() const;
+	void InsertNewLayer();
+
+	bool CanCopySelectedLayer() const;
+	void CopySelectedLayer();
+
+	bool CanCutSelectedLayer() const;
+	void CutSelectedLayer();
+
+	bool CanPasteLayer() const;
+	void PasteLayer();
+
+	bool CanDuplicateSelectedLayer() const;
+	void DuplicateSelectedLayer();
+
+	bool CanDeleteSelectedLayer() const;
+	void DeleteSelectedLayer();
+
+	//~ Begin SWidget
+	virtual FReply OnKeyDown(const FGeometry& MyGeometry, const FKeyEvent& InKeyEvent) override;
+	//~ End SWidget
+
+	//~ Begin FUndoClient
+	virtual void PostUndo(bool bSuccess) override;
+	virtual void PostRedo(bool bSuccess) override;
+	//~ End FUndoClient
+
+protected:
+	static TSharedPtr<FAssetThumbnailPool> ThumbnailPool;
+
+	static TMap<const SWidget*, TArray<FDMPropertyHandle>> PropertyHandleMap;
+
+	static FDMPropertyHandle CreatePropertyHandle(const void* InOwner, UDMMaterialComponent* InComponent, const FName& InPropertyName);
+
+	TSharedPtr<FUICommandList> CommandList;
+
+	TSharedPtr<SBox> Container;
+	TSharedPtr<SDMToolBar> Toolbar;
+	TSharedPtr<SBox> GlobalOpacityContainer;
+	TSharedPtr<SDMMaterialParameters> ParametersWidget;
+	TSharedPtr<SBox> SlotsContainer;
+	TSharedPtr<SBox> SlotPickerContainer;
+	TSharedPtr<SWidgetSwitcher> SlotSwitcher;
+	int32 ActiveSlotIndex;
+
+	TWeakObjectPtr<UDynamicMaterialModel> MaterialModelWeak;
+	FDMObjectMaterialProperty ObjectProperty;
+
+	TArray<TSharedPtr<SDMSlot>> SlotWidgets;
+
+	void BindCommands();
+
+	TSharedRef<SWidget> CreateMainLayout();
+	TSharedRef<SWidget> CreateMaterialSettingsRow();
+	TSharedRef<SWidget> CreateGlobalOpacityWidget();
+	TSharedRef<SWidget> CreateParametersArea();
+
+	TSharedRef<SWidget> CreateSlotPickerWidget();
+	TSharedRef<SWidget> CreateSlotsWidget();
+
+	TSharedRef<SWidget> CreateActorMaterialSlotSelector(const AActor* InActor);
+
+	FReply OnAddSlotButtonClicked();
+
+	void OnMaterialBuilt(UDynamicMaterialModel* InMaterialModel);
+	void OnValuesUpdated(UDynamicMaterialModel* InMaterialModel);
+	void OnSlotsUpdated(UDynamicMaterialModel* InMaterialModel);
+
+	bool IsGlobalOpacityEnabled() const;
+
+	FReply OnRGBButtonClicked_NoSlot();
+
+	FSlateColor GetRGBButtonColorAndOpacity_HasSlot() const;
+	FReply OnRGBButtonClicked_HasSlot();
+
+	FReply OnOpacityButtonClicked_NoSlot();
+	bool GetOpacityButtonEnabled_NoSlot() const;
+
+	FSlateColor GetOpacityButtonColorAndOpacity_HasSlot() const;
+	FReply OnOpacityButtonClicked_HasSlot();
+	bool GetOpacityButtonEnabled_HasSlot() const;
+
+	FReply OnCreateMaterialButtonClicked(TWeakPtr<FDMObjectMaterialProperty> InMaterialProperty);
+
+	void OnToolBarPropertyChanged(TSharedPtr<FDMObjectMaterialProperty> InNewSelectedProperty);
+	TSharedRef<SWidget> MakeToolBarSettingsMenu();
+
+	void OnSettingsChanged(const FPropertyChangedEvent& InPropertyChangedEvent);
+
+	EMaterialDomain GetSelectedDomain() const;
+	void OnDomainChanged(const EMaterialDomain InDomain);
+	bool CanChangeDomain() const;
+
+	EBlendMode GetSelectedBlendMode() const;
+	void OnBlendModeChanged(const EBlendMode InBlendMode);
+	bool CanChangeBlendType() const;
+
+	ECheckBoxState IsMaterialUnlit() const;
+	void OnMaterialUnlitChanged(const ECheckBoxState InNewCheckState);
+
+	ECheckBoxState IsMaterialAnimated() const;
+	void OnMaterialAnimatedChanged(const ECheckBoxState InNewCheckState);
+
+	bool CanChangeMaterialShadingModel() const;
+
+	void OnUndo();
+};
