[MaterialDesigner]

diff --git a/Engine/Plugins/Experimental/DynamicMaterial/Source/DynamicMaterialEditor/Private/Slate/SDMSlot.cpp b/Engine/Plugins/Experimental/DynamicMaterial/Source/DynamicMaterialEditor/Private/Slate/SDMSlot.cpp
new file mode 100644
--- /dev/null
+++ b/Engine/Plugins/Experimental/DynamicMaterial/Source/DynamicMaterialEditor/Private/Slate/SDMSlot.cpp
@@ -0,0 +1,1504 @@
+// Copyright Epic Games, Inc. All Rights Reserved.
+
+#include "SDMSlot.h"
+#include "Components/DMMaterialLayer.h"
+#include "Components/DMMaterialSlot.h"
+#include "Components/DMMaterialStage.h"
+#include "Components/DMMaterialStageBlend.h"
+#include "Components/DMMaterialStageExpression.h"
+#include "Components/DMMaterialStageGradient.h"
+#include "Components/DMMaterialStageThroughputLayerBlend.h"
+#include "Components/DMMaterialValue.h"
+#include "Components/MaterialStageBlends/DMMSBNormal.h"
+#include "Components/MaterialStageExpressions/DMMSESceneTexture.h"
+#include "Components/MaterialStageExpressions/DMMSETextureSample.h"
+#include "Components/MaterialStageInputs/DMMSITextureUV.h"
+#include "Components/MaterialStageInputs/DMMSIValue.h"
+#include "Components/MaterialValues/DMMaterialValueFloat1.h"
+#include "DetailLayoutBuilder.h"
+#include "DMBlueprintFunctionLibrary.h"
+#include "DMPrivate.h"
+#include "DynamicMaterialEditorModule.h"
+#include "DynamicMaterialEditorSettings.h"
+#include "DynamicMaterialEditorStyle.h"
+#include "DynamicMaterialModule.h"
+#include "Menus/DMMaterialSlotMenus.h"
+#include "Model/DynamicMaterialModel.h"
+#include "Model/DynamicMaterialModelEditorOnlyData.h"
+#include "Properties/Editors/SDMPropertyEditOpacity.h"
+#include "ScopedTransaction.h"
+#include "Slate/Layers/SDMSlotLayerItem.h"
+#include "Slate/Layers/SDMSlotLayerView.h"
+#include "Slate/Properties/SDMMaterialProperty.h"
+#include "Slate/Properties/SDMSourceBlendType.h"
+#include "Slate/SDMComponentEdit.h"
+#include "Slate/SDMEditor.h"
+#include "Slate/SDMStage.h"
+#include "SlateOptMacros.h"
+#include "Styling/AppStyle.h"
+#include "Styling/StyleColors.h"
+#include "Widgets/Colors/SColorBlock.h"
+#include "Widgets/Input/SButton.h"
+#include "Widgets/Input/SComboButton.h"
+#include "Widgets/Layout/SBorder.h"
+#include "Widgets/Layout/SScrollBox.h"
+#include "Widgets/Layout/SWidgetSwitcher.h"
+#include "Widgets/Layout/SWrapBox.h"
+#include "Widgets/SBoxPanel.h"
+#include "Widgets/SCanvas.h"
+#include "Widgets/SNullWidget.h"
+#include "Widgets/SOverlay.h"
+#include "Widgets/Text/STextBlock.h"
+
+#define LOCTEXT_NAMESPACE "SDMSlot"
+
+SDMSlot::~SDMSlot()
+{
+	if (OnEndFrameDelegateHandle.IsValid())
+	{
+		FCoreDelegates::OnEndFrame.Remove(OnEndFrameDelegateHandle);
+	}
+
+	if (!FDynamicMaterialModule::AreUObjectsSafe())
+	{
+		return;
+	}
+
+	if (UDMMaterialSlot* Slot = SlotWeak.Get())
+	{
+		Slot->GetOnPropertiesUpdateDelegate().RemoveAll(this);
+		Slot->GetOnLayersUpdateDelegate().RemoveAll(this);
+	}
+
+	SDMEditor::ClearPropertyHandles(this);
+}
+
+BEGIN_SLATE_FUNCTION_BUILD_OPTIMIZATION
+
+void SDMSlot::Construct(const FArguments& InArgs, const TSharedRef<SDMEditor>& InEditor, UDMMaterialSlot* InSlot)
+{
+	EditorWidgetWeak = InEditor;
+	SlotWeak = InSlot;
+
+	SlotPreviewSize = InArgs._SlotPreviewSize;
+	LayerPreviewSize = InArgs._LayerPreviewSize;
+
+	if (ensure(IsValid(InSlot)))
+	{
+		InSlot->GetOnPropertiesUpdateDelegate().AddSP(this, &SDMSlot::OnSlotPropertiesUpdated);
+		InSlot->GetOnLayersUpdateDelegate().AddSP(this, &SDMSlot::OnSlotLayersUpdated);
+
+		for (const UDMMaterialLayerObject* Layer : InSlot->GetLayers())
+		{
+			Layer->ForEachValidStage(
+				EDMMaterialLayerStage::All,
+				[](UDMMaterialStage* InStage)
+				{
+					InStage->SetBeingEdited(false);
+				});
+		}
+
+		RefreshMainWidget();
+
+		bInvalidateMainWidget = false;
+		bInvalidateHeaderWidget = false;
+		bInvalidateSettingsWidget = false;
+		bInvalidateComponentEditWidget = false;
+		OnEndFrameDelegateHandle = FCoreDelegates::OnEndFrame.AddSP(this, &SDMSlot::HandleEndFrameRefresh);
+	}
+}
+
+void SDMSlot::OnLayerSelected(TSharedPtr<FDMMaterialLayerReference> InLayerItem, const int32 InLayerIndex)
+{
+	InvalidateSlotSettingsRowWidget();
+	InvalidateComponentEditWidget();
+}
+
+void SDMSlot::OnLayerStageSelected(const bool bInSelected, const TSharedRef<SDMStage>& InStageWidget)
+{
+	InvalidateSlotSettingsRowWidget();
+	SetEditedComponent(InStageWidget->GetStage());
+}
+
+void SDMSlot::HandleEndFrameRefresh()
+{
+	if (bInvalidateMainWidget)
+	{
+		RefreshMainWidget();
+	}
+	else
+	{
+		if (bInvalidateHeaderWidget)
+		{
+			RefreshHeaderPropertyListWidget();
+		}
+
+		if (bInvalidateSettingsWidget)
+		{
+			RefreshSlotSettingsRowWidget();
+		}
+
+		if (bInvalidateComponentEditWidget)
+		{
+			RefreshComponentEditWidget();
+		}
+	}
+}
+
+FText SDMSlot::GetLayerButtonsDescription() const
+{
+	UDMMaterialSlot* Slot = SlotWeak.Get();
+
+	if (!Slot)
+	{
+		return FText::GetEmpty();
+	}
+
+	const int32 SlotLayerCount = Slot->GetLayers().Num();
+	return SlotLayerCount == 1
+		? LOCTEXT("SlotLayerInfo_OneLayer", "1 Layer")
+		: FText::Format(LOCTEXT("SlotLayerInfo_MultipleLayers", "{0} Layers"), SlotLayerCount);
+}
+
+TSharedRef<SWidget> SDMSlot::GetLayerButtonsMenuContent()
+{
+	return FDMMaterialSlotMenus::MakeAddLayerButtonMenu(SharedThis(this));
+}
+
+bool SDMSlot::GetLayerRowsButtonsCanDuplicate() const
+{
+	return LayerView->GetSelectedItems().IsEmpty() == false;
+}
+
+FReply SDMSlot::OnLayerRowButtonsDuplicateClicked()
+{
+	if (TSharedPtr<SDMEditor> EditorWidget = EditorWidgetWeak.Pin())
+	{
+		if (EditorWidget->CanDuplicateSelectedLayer())
+		{
+			EditorWidget->DuplicateSelectedLayer();
+		}
+	}
+
+	return FReply::Handled();
+}
+
+bool SDMSlot::GetLayerRowsButtonsCanRemove() const
+{
+	const UDMMaterialLayerObject* Layer = GetSelectedLayer();
+
+	if (!Layer)
+	{
+		return false;
+	}
+
+	return CanRemoveLayerByIndex(Layer->FindIndex());
+}
+
+FReply SDMSlot::OnLayerRowButtonsRemoveClicked()
+{
+	const UDMMaterialLayerObject* Layer = GetSelectedLayer();
+
+	if (!Layer)
+	{
+		return FReply::Unhandled();
+	}
+
+	RemoveLayerByIndex(Layer->FindIndex());
+
+	return FReply::Handled();
+}
+
+UDMMaterialLayerObject* SDMSlot::GetSelectedLayer() const
+{
+	if (!LayerView.IsValid())
+	{
+		return nullptr;
+	}
+
+	const int32 SelectedLayerIndex = LayerView->GetSelectedLayerIndex();
+
+	if (SelectedLayerIndex == INDEX_NONE)
+	{
+		return nullptr;
+	}
+
+	if (!LayerView->GetLayerItems().IsValidIndex(SelectedLayerIndex))
+	{
+		return nullptr;
+	}
+
+	return LayerView->GetLayerItems()[SelectedLayerIndex]->GetLayer();
+}
+
+void SDMSlot::SetSelectedLayer(UDMMaterialLayerObject* InLayer) const
+{
+	if (!LayerView.IsValid())
+	{
+		return;
+	}
+
+	LayerView->ClearSelection();
+	LayerView->SelectLayerItem(InLayer, /* Mask */ false, /* Selected */ true, ESelectInfo::Direct);
+}
+
+UDMMaterialComponent* SDMSlot::GetEditedComponent() const
+{
+	return EditedComponent.Get();
+}
+
+void SDMSlot::SetEditedComponent(UDMMaterialComponent* InComponent)
+{
+	EditedComponent = InComponent;
+
+	InvalidateComponentEditWidget();
+}
+
+int32 SDMSlot::GetSourceBlendTypeSwitcherIndex() const
+{
+	if (LayerView->GetLayerItems().IsValidIndex(LayerView->GetSelectedLayerIndex()))
+	{
+		if (TSharedPtr<SDMSlotLayerItem> LayerWidget = LayerView->WidgetFromLayerItem(LayerView->GetLayerItems()[LayerView->GetSelectedLayerIndex()]))
+		{
+			if (TSharedPtr<SDMStage> BaseStageWidget = LayerWidget->GetBaseStageWidget())
+			{
+				if (UDMMaterialStage* BaseStage = BaseStageWidget->GetStage())
+				{
+					return BaseStage->CanChangeSource() ? 2 : 1;
+				}
+			}
+		}
+	}
+
+	return 0;
+}
+
+TSubclassOf<UDMMaterialStageBlend> SDMSlot::GetSelectedSourceBlendType() const
+{
+	if (LayerView->GetLayerItems().IsValidIndex(LayerView->GetSelectedLayerIndex()))
+	{
+		if (TSharedPtr<SDMSlotLayerItem> LayerWidget = LayerView->WidgetFromLayerItem(LayerView->GetLayerItems()[LayerView->GetSelectedLayerIndex()]))
+		{
+			if (TSharedPtr<SDMStage> BaseStageWidget = LayerWidget->GetBaseStageWidget())
+			{
+				if (UDMMaterialStage* BaseStage = BaseStageWidget->GetStage())
+				{
+					if (UDMMaterialStageSource* BaseStageSource = BaseStage->GetSource())
+					{
+						return BaseStageSource->GetClass();
+					}
+				}
+			}
+		}
+	}
+
+	return UDMMaterialStageBlendNormal::StaticClass();
+}
+
+void SDMSlot::OnSourceBlendTypedSelected(const TSubclassOf<UDMMaterialStageBlend> InNewItem)
+{
+	if (LayerView->GetLayerItems().IsValidIndex(LayerView->GetSelectedLayerIndex()))
+	{
+		if (TSharedPtr<SDMSlotLayerItem> LayerWidget = LayerView->WidgetFromLayerItem(LayerView->GetLayerItems()[LayerView->GetSelectedLayerIndex()]))
+		{
+			if (TSharedPtr<SDMStage> BaseStageWidget = LayerWidget->GetBaseStageWidget())
+			{
+				if (UDMMaterialStage* BaseStage = BaseStageWidget->GetStage())
+				{
+					FScopedTransaction Transaction(LOCTEXT("SetStageBlendMode", "Material Designer Set Blend Mode"));
+					BaseStage->Modify();
+					BaseStage->ChangeSource_Blend(InNewItem);
+
+					RefreshSlotSettingsRowWidget();
+				}
+			}
+		}
+	}
+}
+
+bool SDMSlot::IsOpacityPropertyEnabled() const
+{
+	return LayerView->GetLayerItems().IsValidIndex(LayerView->GetSelectedLayerIndex());
+}
+
+FReply SDMSlot::OnRemoveSlotClicked()
+{
+	if (CanRemoveSlot())
+	{
+		RemoveSlot();
+		return FReply::Handled();
+	}
+
+	return FReply::Unhandled();
+}
+
+bool SDMSlot::CanRemoveSlot() const
+{
+	if (UDMMaterialSlot* Slot = SlotWeak.Get())
+	{
+		// Cannot remove RGB slot.
+		const TArray<TObjectPtr<UDMMaterialLayerObject>>& Layers = Slot->GetLayers();
+
+		for (const TObjectPtr<UDMMaterialLayerObject>& Layer : Layers)
+		{
+			if (Layer->GetMaterialProperty() == EDMMaterialPropertyType::BaseColor
+				|| Layer->GetMaterialProperty() == EDMMaterialPropertyType::EmissiveColor)
+			{
+				return false;
+			}
+		}
+
+		return true;
+	}
+
+	return false;
+}
+
+void SDMSlot::RefreshMainWidget()
+{
+	bInvalidateMainWidget = false;
+	bInvalidateHeaderWidget = false;
+	bInvalidateSettingsWidget = false;
+	bInvalidateComponentEditWidget = false;
+
+	const float SplitterLocation = FMath::Clamp(UDynamicMaterialEditorSettings::Get()->SplitterLocation, 0.3f, 0.7f);
+	const float LayerViewSplitterLocation = SplitterLocation;
+	const float ExtraSpaceSplitterLocation = 1.0f - SplitterLocation;
+
+	LayerView = SNew(SDMSlotLayerView, SharedThis(this))
+		.OnLayerSelected(this, &SDMSlot::OnLayerSelected)
+		.OnLayerStageSelected(this, &SDMSlot::OnLayerStageSelected)
+		.PreviewSize(LayerPreviewSize);
+
+	SplitterContainer =
+		SNew(SSplitter)
+		.Style(FAppStyle::Get(), "DetailsView.Splitter")
+		.Orientation(Orient_Vertical)
+		.ResizeMode(ESplitterResizeMode::Fill)
+		.PhysicalSplitterHandleSize(3.0f)
+		.HitDetectionSplitterHandleSize(4.0f)
+		.OnSplitterFinishedResizing(this, &SDMSlot::OnSplitterResized)
+
+		+ SSplitter::Slot()
+		.Resizable(false)
+		.SizeRule(SSplitter::ESizeRule::SizeToContent)
+		[
+			SAssignNew(SlotSettingsRowContainer, SBox)
+			.Padding(3.0f, 4.0f, 3.0f, 4.0f)
+			[
+				CreateSlotSettingsRow()
+			]
+		]
+
+	+ SSplitter::Slot()
+		.Expose(LayerViewSplitterSlot)
+		.Resizable(true)
+		.SizeRule(SSplitter::ESizeRule::FractionOfParent)
+		.MinSize(50)
+		.Value(LayerViewSplitterLocation)
+		[
+			SNew(SBorder)
+			.Padding(2.0f)
+			.BorderImage(FDynamicMaterialEditorStyle::GetBrush("LayerView.Background"))
+			.HAlign(HAlign_Fill)
+			.VAlign(VAlign_Fill)
+			[
+				LayerView.ToSharedRef()
+			]
+		]
+
+	+ SSplitter::Slot()
+		.Resizable(false)
+		.SizeRule(SSplitter::ESizeRule::SizeToContent)
+		[
+			CreateLayerButtonsRowWidget()
+		]
+
+	+ SSplitter::Slot()
+		.Expose(ExtraSpaceSplitterSlot)
+		.Resizable(true)
+		.SizeRule(SSplitter::ESizeRule::FractionOfParent)
+		.MinSize(50)
+		.Value(ExtraSpaceSplitterLocation)
+		[
+			SAssignNew(ComponentEditContainer, SScrollBox)
+			+ SScrollBox::Slot()
+			.AutoSize()
+			[
+				CreateComponentEditWidget()
+			]
+		];
+
+	TSharedRef<SVerticalBox> Child = SNew(SVerticalBox);
+
+	if constexpr (UE::DynamicMaterialEditor::bMultipleSlotPropertiesEnabled)
+	{
+		Child->AddSlot()
+			.AutoHeight()
+			[
+				SAssignNew(HeaderPropertyListWidget, SBox)
+					.Padding(2.0f, 2.0f, 2.0f, 2.0f)
+					.HAlign(HAlign_Fill)
+					.VAlign(VAlign_Top)
+					[
+						CreateHeaderPropertyListWidget()
+					]
+			];
+	}
+
+	Child->AddSlot()
+		.FillHeight(1.0f)
+		.HAlign(HAlign_Fill)
+		.VAlign(VAlign_Fill)
+		.Padding(0.0f, 2.0f, 0.0f, 0.0f)
+		[
+			SplitterContainer.ToSharedRef()
+		];
+
+	ChildSlot
+		.HAlign(HAlign_Fill)
+		.VAlign(VAlign_Fill)
+		[
+			Child
+		];
+
+	if (UDMMaterialSlot* Slot = SlotWeak.Get())
+	{
+		Slot->SetEditingLayers(true);
+
+		const TArray<TObjectPtr<UDMMaterialLayerObject>>& SlotLayers = Slot->GetLayers();
+
+		for (const TObjectPtr<UDMMaterialLayerObject>& Layer : SlotLayers)
+		{
+			if (UDMMaterialStage* Stage = Layer->GetFirstStageBeingEdited(EDMMaterialLayerStage::All))
+			{
+				LayerView->SelectLayerItem(Stage, true, ESelectInfo::OnMouseClick);
+				break;
+			}
+		}
+	}
+}
+
+TSharedRef<SWidget> SDMSlot::CreateHeaderPropertyListWidget()
+{
+	UDMMaterialSlot* Slot = SlotWeak.Get();
+
+	if (!ensure(Slot))
+	{
+		return SNullWidget::NullWidget;
+	}
+
+	TSharedPtr<SDMEditor> MaterialEditor = EditorWidgetWeak.Pin();
+	if (!ensure(MaterialEditor.IsValid()))
+	{
+		return SNullWidget::NullWidget;
+	}
+
+	UDynamicMaterialModel* MaterialModel = MaterialEditor->GetMaterialModel();
+	if (!ensure(IsValid(MaterialModel)))
+	{
+		return SNullWidget::NullWidget;
+	}
+
+	UDynamicMaterialModelEditorOnlyData* ModelEditorOnlyData = UDynamicMaterialModelEditorOnlyData::Get(MaterialModel);
+	if (!ensure(IsValid(ModelEditorOnlyData)))
+	{
+		return SNullWidget::NullWidget;
+	}
+
+	TSharedRef<SWrapBox> PropertyContainer = SNew(SWrapBox)
+		.InnerSlotPadding(FVector2D(2.f, 2.f))
+		.Orientation(EOrientation::Orient_Horizontal)
+		.UseAllottedSize(true)
+		.UseAllottedWidth(true);
+
+	const TArray<EDMMaterialPropertyType> Properties = ModelEditorOnlyData->GetMaterialPropertiesForSlot(Slot);
+
+	for (EDMMaterialPropertyType Property : Properties)
+	{
+		PropertyContainer->AddSlot()
+			.FillEmptySpace(false)
+			[
+				SNew(SDMMaterialProperty, SharedThis(this), Property)
+			];
+	}
+
+	const TMap<TWeakObjectPtr<UDMMaterialSlot>, int32>& SlotsReferencedBy = Slot->GetSlotsReferencedBy();
+
+	for (const TPair<TWeakObjectPtr<UDMMaterialSlot>, int32>& Pair : SlotsReferencedBy)
+	{
+		PropertyContainer->AddSlot()
+			.FillEmptySpace(false)
+			[
+				SNew(SOverlay)
+				+ SOverlay::Slot()
+				.HAlign(EHorizontalAlignment::HAlign_Fill)
+				[
+					SNew(SColorBlock)
+					.Color(FStyleColors::AccentGreen.GetSpecifiedColor())
+					.CornerRadius(FVector4(5.f, 5.f, 5.f, 5.f))
+				]
+				+ SOverlay::Slot()
+				.HAlign(EHorizontalAlignment::HAlign_Fill)
+				[
+					SNew(SCanvas)
+					+ SCanvas::Slot()
+					.Position(FVector2D(8.f, 8.f))
+					.Size(FVector2D(4.f, 4.f))
+					[
+						SNew(SColorBlock)
+						.Color(FLinearColor::White)
+						.CornerRadius(FVector4(2.f, 2.f, 2.f, 2.f))
+					]
+				]
+				+ SOverlay::Slot()
+				.HAlign(EHorizontalAlignment::HAlign_Fill)
+				[
+					SNew(STextBlock)
+					.Text(Pair.Key->GetDescription())
+					.Font(IDetailLayoutBuilder::GetDetailFont())
+					.ColorAndOpacity(FLinearColor::White)
+					.Margin(FMargin(16.f, 4.f, 4.f, 4.f))
+				]
+			];
+	}
+
+	return PropertyContainer;
+}
+
+TSharedRef<SWidget> SDMSlot::CreateLayerButtonsRowWidget()
+{
+	return 
+		SNew(SHorizontalBox)
+		+ SHorizontalBox::Slot()
+		.AutoWidth()
+		.HAlign(HAlign_Left)
+		.VAlign(VAlign_Top)
+		.Padding(5.0f, 2.0f, 10.0f, 2.0f)
+		[
+			SNew(SButton)
+			.ContentPadding(4.0f)
+			.ButtonStyle(FDynamicMaterialEditorStyle::Get(), "HoverHintOnly.Bordered.Dark")
+			.ToolTipText(LOCTEXT("RemoveSlot", "Remove Slot\n\nRGB Slots cannot be removed."))
+			.IsEnabled(this, &SDMSlot::CanRemoveSlot)
+			.OnClicked(this, &SDMSlot::OnRemoveSlotClicked)
+			[
+				SNew(SImage)
+				.Image(FAppStyle::Get().GetBrush("GenericCommands.Delete"))
+				.DesiredSizeOverride(FVector2D(22.0f))
+			]
+		]
+		+ SHorizontalBox::Slot()
+		.FillWidth(1.0f)
+		.HAlign(HAlign_Left)
+		.VAlign(VAlign_Center)
+		.Padding(5.0f, 0.f, 0.f, 0.f)
+		[
+			SNew(STextBlock)
+			.TextStyle(FDynamicMaterialEditorStyle::Get(), "SlotLayerInfo")
+			.Text(this, &SDMSlot::GetLayerButtonsDescription)
+		]
+		+ SHorizontalBox::Slot()
+		.AutoWidth()
+		.HAlign(HAlign_Right)
+		.VAlign(VAlign_Center)
+		[
+			SNew(SComboButton)
+			.HasDownArrow(false)
+			.IsFocusable(true)
+			.ContentPadding(4.0f)
+			.ButtonStyle(FDynamicMaterialEditorStyle::Get(), "HoverHintOnly.Bordered.Dark")
+			.ToolTipText(LOCTEXT("AddLayer", "Add New Layer"))
+			.OnGetMenuContent(this, &SDMSlot::GetLayerButtonsMenuContent)
+			.ButtonContent()
+			[
+				SNew(SImage)
+				.Image(FAppStyle::Get().GetBrush("Icons.Plus"))
+				.ColorAndOpacity(FDynamicMaterialEditorStyle::Get().GetColor("Color.Stage.Enabled"))
+				.DesiredSizeOverride(FVector2D(22.0f))
+			]
+		]
+		+ SHorizontalBox::Slot()
+		.AutoWidth()
+		.HAlign(HAlign_Right)
+		.VAlign(VAlign_Center)
+		.Padding(5.0f, 2.0f, 0.0f, 2.0f)
+		[
+			SNew(SButton)
+			.ContentPadding(4.0f)
+			.ButtonStyle(FDynamicMaterialEditorStyle::Get(), "HoverHintOnly.Bordered.Dark")
+			.ToolTipText(LOCTEXT("DuplicateLayer", "Duplicate Selected Layer"))
+			.IsEnabled(this, &SDMSlot::GetLayerRowsButtonsCanDuplicate)
+			.OnClicked(this, &SDMSlot::OnLayerRowButtonsDuplicateClicked)
+			[
+				SNew(SImage)
+				.Image(FDynamicMaterialEditorStyle::GetBrush("LayerView.DuplicateIcon"))
+				.DesiredSizeOverride(FVector2D(22.0f))
+			]
+		]
+		+ SHorizontalBox::Slot()
+		.AutoWidth()
+		.HAlign(HAlign_Right)
+		.VAlign(VAlign_Center)
+		.Padding(5.0f, 2.0f, 0.0f, 2.0f)
+		[
+			SNew(SButton)
+			.ContentPadding(4.0f)
+			.ButtonStyle(FDynamicMaterialEditorStyle::Get(), "HoverHintOnly.Bordered.Dark")
+			.ToolTipText(LOCTEXT("RemoveLayer", "Remove Selected Layer\n\nThe last layer cannot be removed."))
+			.IsEnabled(this, &SDMSlot::GetLayerRowsButtonsCanRemove)
+			.OnClicked(this, &SDMSlot::OnLayerRowButtonsRemoveClicked)
+			[
+				SNew(SImage)
+				.Image(FAppStyle::Get().GetBrush("Icons.Delete"))
+				.DesiredSizeOverride(FVector2D(22.0f))
+			]
+		];
+}
+
+TSharedRef<SWidget> SDMSlot::CreateComponentEditWidget()
+{
+	if (UDMMaterialComponent* Component = EditedComponent.Get())
+	{
+		return
+			SNew(SBorder)
+			.HAlign(HAlign_Fill)
+			.VAlign(VAlign_Fill)
+			.Padding(3.0f)
+			.BorderImage(FDynamicMaterialEditorStyle::GetBrush("LayerView.Details.Background"))
+			[
+				SNew(SDMComponentEdit, Component, SharedThis(this))
+			];
+	}
+
+	return SNullWidget::NullWidget;
+}
+
+TSharedRef<SWidget> SDMSlot::CreateSlotSettingsRow()
+{
+	TSharedPtr<SDMPropertyEdit> OpacityPropertyEditWidget = nullptr;
+	TSharedRef<SWidget> OpacityExtensionWidget = SNullWidget::NullWidget;
+
+	if (const UDMMaterialLayerObject* SelectedLayer = GetSelectedLayer())
+	{
+		if (UDMMaterialStage* ValidStage = SelectedLayer->GetFirstValidStage(EDMMaterialLayerStage::All))
+		{
+			SelectedOpacityStageInputValue = UDMBlueprintFunctionLibrary::FindDefaultStageOpacityInputValue(ValidStage);
+
+			if (SelectedOpacityStageInputValue.IsValid())
+			{
+				UDMMaterialValueFloat1* OpacityValue = Cast<UDMMaterialValueFloat1>(SelectedOpacityStageInputValue->GetValue());
+				SelectedOpacityInputValue = OpacityValue;
+
+				if (SelectedOpacityInputValue.IsValid())
+				{
+					UDMMaterialStageSource* StageSource = ValidStage->GetSource();
+
+					if (IsValid(StageSource))
+					{
+						OpacityPropertyEditWidget = SNew(SDMPropertyEditOpacity, SharedThis(this), OpacityValue);
+						OpacityPropertyEditWidget->SetEnabled(TAttribute<bool>::CreateSP(this, &SDMSlot::IsOpacityPropertyEnabled));
+
+						OpacityExtensionWidget = SDMComponentEdit::CreateExtensionButtons(
+							SharedThis(this),
+							OpacityValue,
+							UDMMaterialValue::ValueName,
+							true,
+							FSimpleDelegate::CreateWeakLambda(OpacityValue, [OpacityValue]()
+								{
+									OpacityValue->ApplyDefaultValue();
+								})
+						);
+					}
+				}
+			}
+		}
+	}
+
+	TSharedPtr<SWidget> OpacityLabel;
+
+	TSharedRef<SWidget> Row =
+		SNew(SHorizontalBox)
+		+ SHorizontalBox::Slot()
+		.FillWidth(1.0f)
+		.HAlign(HAlign_Fill)
+		.VAlign(VAlign_Center)
+		.Padding(5.0f, 0.0f, 0.0f, 0.0f)
+		[
+			SNew(SHorizontalBox)
+			.ToolTipText(LOCTEXT("MaterialDesignerInstanceBlendModeTooltip", "Change the Blend Mode for this Material Designer Instance."))
+			+ SHorizontalBox::Slot()
+			.AutoWidth()
+			.HAlign(HAlign_Left)
+			.VAlign(VAlign_Center)
+			[
+				SNew(SBox)
+				.Padding(0.f, 0.f, 5.f, 0.f)
+				[
+					SNew(STextBlock)
+					.Text(LOCTEXT("MaterialDesignerInstanceBlendMode", "Blend"))
+				]
+			]
+			+ SHorizontalBox::Slot()
+			.FillWidth(1.0f)
+			.HAlign(HAlign_Fill)
+			.VAlign(VAlign_Center)
+			[
+				SNew(SBox)
+				.HeightOverride(22.f)
+				.VAlign(EVerticalAlignment::VAlign_Center)
+				[
+					SNew(SWidgetSwitcher)
+					.WidgetIndex(GetSourceBlendTypeSwitcherIndex())
+					+ SWidgetSwitcher::Slot()
+					[
+						SNew(STextBlock)
+						.Text(INVTEXT("-"))
+					]
+					+ SWidgetSwitcher::Slot()
+					[
+						SNew(STextBlock)
+						.Text(INVTEXT("-"))
+					]
+					+ SWidgetSwitcher::Slot()
+					[
+						SNew(SDMSourceBlendType)
+						.SelectedItem(this, &SDMSlot::GetSelectedSourceBlendType)
+						.OnSelectedItemChanged(this, &SDMSlot::OnSourceBlendTypedSelected)
+					]
+				]
+			]
+		]
+		+ SHorizontalBox::Slot()
+		.FillWidth(1.0f)
+		.HAlign(HAlign_Fill)
+		.VAlign(VAlign_Center)
+		.Padding(5.0f, 0.0f, 0.0f, 0.0f)
+		[
+			SNew(SHorizontalBox)
+			.ToolTipText(LOCTEXT("StageOpacityTooltip", "Change the opacity for the selected stage."))
+			+ SHorizontalBox::Slot()
+			.AutoWidth()
+			.HAlign(HAlign_Left)
+			.VAlign(VAlign_Center)
+			[
+				SAssignNew(OpacityLabel, SBox)
+				.Padding(0.f, 0.f, 5.f, 0.f)
+				[
+					SNew(STextBlock)
+					.Text(LOCTEXT("StageOpacity", "Opacity"))
+				]
+			]
+			+ SHorizontalBox::Slot()
+			.FillWidth(1.0f)
+			.HAlign(HAlign_Fill)
+			.VAlign(VAlign_Center)
+			[
+				SNew(SBox)
+				.HeightOverride(22.f)
+				.VAlign(EVerticalAlignment::VAlign_Center)
+				[
+					OpacityPropertyEditWidget.IsValid()
+						? StaticCastSharedRef<SWidget>(OpacityPropertyEditWidget.ToSharedRef())
+						: StaticCastSharedRef<SWidget>(SNew(STextBlock).Text(INVTEXT("-")))
+				]
+			]
+			+ SHorizontalBox::Slot()
+			.AutoWidth()
+			.HAlign(HAlign_Left)
+			.VAlign(VAlign_Center)
+			[
+				OpacityExtensionWidget
+			]
+		];
+
+	if (OpacityPropertyEditWidget.IsValid())
+	{
+		OpacityLabel->SetOnMouseButtonDown(FPointerEventHandler::CreateStatic(&SDMPropertyEdit::CreateRightClickDetailsMenu, OpacityPropertyEditWidget.ToWeakPtr()));
+	}
+
+	return Row;
+}
+
+END_SLATE_FUNCTION_BUILD_OPTIMIZATION
+
+void SDMSlot::InvalidateMainWidget()
+{
+	bInvalidateMainWidget = true;
+}
+
+void SDMSlot::InvalidateHeaderPropertyListWidget()
+{
+	bInvalidateHeaderWidget = true;
+}
+
+void SDMSlot::InvalidateSlotSettingsRowWidget()
+{
+	bInvalidateSettingsWidget = true;
+}
+
+void SDMSlot::InvalidateComponentEditWidget()
+{
+	bInvalidateComponentEditWidget = true;
+}
+
+void SDMSlot::RefreshHeaderPropertyListWidget()
+{
+	bInvalidateHeaderWidget = false;
+
+	if (HeaderPropertyListWidget.IsValid())
+	{
+		HeaderPropertyListWidget->SetContent(CreateHeaderPropertyListWidget());
+	}
+}
+
+void SDMSlot::RefreshSlotSettingsRowWidget()
+{
+	bInvalidateSettingsWidget = false;
+
+	if (SlotSettingsRowContainer.IsValid())
+	{
+		SlotSettingsRowContainer->SetContent(CreateSlotSettingsRow());
+	}
+}
+
+void SDMSlot::RefreshComponentEditWidget()
+{
+	bInvalidateComponentEditWidget = false;
+
+	if (ComponentEditContainer.IsValid())
+	{
+		ComponentEditContainer->ClearChildren();
+		ComponentEditContainer->AddSlot()
+			.AutoSize()
+			[
+				CreateComponentEditWidget()
+			];
+	}
+}
+
+void SDMSlot::OnSplitterResized() const
+{
+	UDynamicMaterialEditorSettings* Settings = UDynamicMaterialEditorSettings::Get();
+
+	if (LayerViewSplitterSlot)
+	{
+		Settings->SplitterLocation = LayerViewSplitterSlot->GetSizeValue();
+		Settings->SaveConfig();
+	}
+}
+
+void SDMSlot::StartStageOpacityTransaction()
+{
+	StartTransaction(LOCTEXT("StageOpacityTransDescription", "Material Designer Value Scrubbing (Stage Opacity)"));
+
+	if (UDMMaterialValueFloat1* FloatValue = SelectedOpacityInputValue.Get())
+	{
+		FloatValue->Modify();
+	}
+}
+
+void SDMSlot::SetStageOpacity(const float InNewValue)
+{
+	if (UDMMaterialValueFloat1* FloatValue = SelectedOpacityInputValue.Get())
+	{
+		FloatValue->SetValue(InNewValue);
+	}
+}
+
+void SDMSlot::OnStageOpacityChanged(const float InNewValue)
+{
+	if (ScrubbingTransaction.IsValid())
+	{
+		SetStageOpacity(InNewValue);
+	}
+}
+
+void SDMSlot::OnStageOpacityCommitted(const float InNewValue, ETextCommit::Type InCommitType)
+{
+	if (!ScrubbingTransaction.IsValid())
+	{
+		StartStageOpacityTransaction();
+		SetStageOpacity(InNewValue);
+		EndTransaction();
+	}
+}
+
+void SDMSlot::StartTransaction(const FText InDescription)
+{
+	if (ScrubbingTransaction.IsValid())
+	{
+		return;
+	}
+
+	ScrubbingTransaction = MakeShared<FScopedTransaction>(InDescription);
+}
+
+void SDMSlot::EndTransaction()
+{
+	ScrubbingTransaction.Reset();
+}
+
+bool SDMSlot::CheckValidity()
+{
+	UDMMaterialSlot* Slot = SlotWeak.Get();
+
+	if (!Slot || !Slot->IsComponentValid())
+	{
+		return false;
+	}
+
+	bool bHasInvalidStage = false;
+
+	for (const TSharedPtr<FDMMaterialLayerReference>& LayerItem : LayerView->GetLayerItems())
+	{
+		if (!LayerItem.IsValid())
+		{
+			bHasInvalidStage = true;
+			break;
+		}
+
+		const UDMMaterialLayerObject* Layer = LayerItem->GetLayer();
+
+		if (!Layer)
+		{
+			bHasInvalidStage = true;
+			break;
+		}
+
+		if (!IsValid(Layer->GetStage(EDMMaterialLayerStage::Base)) || !Layer->GetStage(EDMMaterialLayerStage::Base)->IsComponentValid() 
+			|| Layer->FindIndex() == INDEX_NONE)
+		{
+			bHasInvalidStage = true;
+			break;
+		}
+	}
+
+	if (bHasInvalidStage)
+	{
+		InvalidateSlotSettingsRowWidget();
+		InvalidateComponentEditWidget();
+	}
+
+	return true;
+}
+
+void SDMSlot::OnSlotLayersUpdated(UDMMaterialSlot* InSlot)
+{
+	if (InSlot != SlotWeak.Get())
+	{
+		return;
+	}
+
+	InvalidateComponentEditWidget();
+}
+
+void SDMSlot::OnSlotPropertiesUpdated(UDMMaterialSlot* InSlot)
+{
+	if (InSlot != SlotWeak.Get())
+	{
+		return;
+	}
+
+	InvalidateHeaderPropertyListWidget();
+}
+
+void SDMSlot::OnStageUpdated(UDMMaterialStage* InStage, EDMUpdateType InUpdateType)
+{
+}
+
+TSharedPtr<SDMStage> SDMSlot::FindStageWidget(UDMMaterialStage* const InStage) const
+{
+	if (!LayerView.IsValid())
+	{
+		return nullptr;
+	}
+
+	const TArray<TWeakPtr<SDMStage>>& StageWidgets = LayerView->GetSelectedStageWidgets();
+
+	for (const TWeakPtr<SDMStage>& StageWidgetWeak : StageWidgets)
+	{
+		if (TSharedPtr<SDMStage> StageWidget = StageWidgetWeak.Pin())
+		{
+			if (StageWidget->GetStage() == InStage)
+			{
+				return StageWidget;
+			}
+		}
+	}
+
+	return nullptr;
+}
+
+void SDMSlot::RemoveLayerByStage(UDMMaterialStage* const InStage, const bool bInSelectNextMaskIfMask)
+{
+	if (!IsValid(InStage))
+	{
+		return;
+	}
+
+	UDMMaterialLayerObject* Layer = InStage->GetLayer();
+
+	if (!Layer)
+	{
+		return;
+	}
+
+	RemoveLayerByIndex(Layer->FindIndex(), bInSelectNextMaskIfMask);
+}
+
+bool SDMSlot::CanRemoveLayerByIndex(const int32 InLayerIndex) const
+{
+	UDMMaterialSlot* Slot = SlotWeak.Get();
+
+	if (!Slot)
+	{
+		return false;
+	}
+
+	if (!Slot->GetLayers().IsValidIndex(InLayerIndex))
+	{
+		return false;
+	}
+
+	if (Slot->GetLayers().Num() > 1)
+	{
+		return true;
+	}
+
+	const UDMMaterialLayerObject* LayerToRemove = Slot->GetLayer(InLayerIndex);
+
+	if (!LayerToRemove)
+	{
+		return false;
+	}
+
+	if (UDynamicMaterialModelEditorOnlyData* ModelEditorOnlyData = Slot->GetMaterialModelEditorOnlyData())
+	{
+		const TArray<EDMMaterialPropertyType> Properties = ModelEditorOnlyData->GetMaterialPropertiesForSlot(Slot);
+
+		for (EDMMaterialPropertyType Property : Properties)
+		{
+			switch (Property)
+			{
+				case EDMMaterialPropertyType::BaseColor:
+				case EDMMaterialPropertyType::EmissiveColor:
+					return Slot->GetLayers().Num() > 1;
+
+				default:
+					// Nothing
+					break;
+			}
+		}
+	}
+
+	return true;
+}
+
+void SDMSlot::RemoveLayerByIndex(const int32 InLayerIndex, const bool bInSelectNextMaskIfMask)
+{
+	if (!CanRemoveLayerByIndex(InLayerIndex))
+	{
+		return;
+	}
+
+	UDMMaterialSlot* Slot = SlotWeak.Get();
+
+	if (!Slot)
+	{
+		return;
+	}
+
+	UDMMaterialLayerObject* LayerToRemove = Slot->GetLayer(InLayerIndex);
+
+	if (!LayerToRemove)
+	{
+		return;
+	}
+
+	FScopedTransaction Transaction(LOCTEXT("RemoveLayer", "Material Designer Remove Layer"));
+	Slot->Modify();
+	Slot->RemoveLayer(LayerToRemove);
+
+	InvalidateMainWidget();
+}
+
+int32 SDMSlot::GetLayerCount() const
+{
+	if (UDMMaterialSlot* Slot = SlotWeak.Get())
+	{
+		return Slot->GetLayers().Num();
+	}
+
+	return 0;
+}
+
+TArray<int32> SDMSlot::GetSelectedLayerIndices() const
+{
+	TArray<int32> OutIndices;
+	
+	TArray<TSharedPtr<FDMMaterialLayerReference>> SelectedLayerItems = LayerView->GetSelectedItems();
+	OutIndices.Reserve(SelectedLayerItems.Num());
+
+	for (const TSharedPtr<FDMMaterialLayerReference>& SelectedLayerItem : SelectedLayerItems)
+	{
+		const int32 LayerIndex = LayerView->GetLayerItemIndex(SelectedLayerItem);
+		OutIndices.Add(LayerIndex);
+	}
+	
+	return OutIndices;
+}
+
+void SDMSlot::ClearSelection()
+{
+	if (LayerView.IsValid())
+	{
+		LayerView->ClearSelection();
+	}
+}
+
+void SDMSlot::RemoveSlot()
+{
+	if (UDMMaterialSlot* Slot = SlotWeak.Get())
+	{
+		if (UDynamicMaterialModelEditorOnlyData* ModelEditorOnlyData = Slot->GetMaterialModelEditorOnlyData())
+		{
+			FScopedTransaction Transaction(LOCTEXT("RemoveSlot", "Material Designer Remove Slot"));
+			ModelEditorOnlyData->Modify();
+			ModelEditorOnlyData->RemoveSlot(Slot->GetIndex());
+
+			if (TSharedPtr<SDMEditor> Editor = EditorWidgetWeak.Pin())
+			{
+				Editor->SetActiveSlotIndex(0);
+				Editor->RefreshSlotPickerList();
+			}
+		}
+	}
+}
+
+void SDMSlot::AddPropertyToSlot(EDMMaterialPropertyType Property)
+{
+	UDMMaterialSlot* Slot = SlotWeak.Get();
+
+	if (!ensure(Slot))
+	{
+		return;
+	}
+
+	UDynamicMaterialModelEditorOnlyData* ModelEditorOnlyData = Slot->GetMaterialModelEditorOnlyData();
+	if (!ensure(IsValid(ModelEditorOnlyData)))
+	{
+		return;
+	}
+
+	FScopedTransaction Transaction(LOCTEXT("AssignPropertyToSlot", "Assign Property To Slot"));
+	ModelEditorOnlyData->Modify();
+	ModelEditorOnlyData->AssignMaterialPropertyToSlot(Property, Slot);
+}
+
+UDMMaterialLayerObject* SDMSlot::AddNewLayer(UDMMaterialStage* InNewBaseStage, UDMMaterialStage* InNewMaskStage)
+{
+	if (!IsValid(InNewBaseStage))
+	{
+		return nullptr;
+	}
+
+	UDMMaterialSlot* Slot = SlotWeak.Get();
+	
+	if (!ensure(Slot))
+	{
+		return nullptr;
+	}
+
+	UDynamicMaterialModelEditorOnlyData* ModelEditorOnlyData = Slot->GetMaterialModelEditorOnlyData();
+
+	if (!ensure(IsValid(ModelEditorOnlyData)))
+	{
+		return nullptr;
+	}
+
+	{
+		const FDMUpdateGuard Guard;
+
+		FScopedTransaction Transaction(LOCTEXT("AddLayer", "Material Designer Add Layer"));
+		Slot->Modify();
+
+		EDMMaterialPropertyType MaterialProperty = EDMMaterialPropertyType::None;
+
+		if (Slot->GetLayers().IsEmpty())
+		{
+			switch (ModelEditorOnlyData->GetShadingModel())
+			{
+				case EDMMaterialShadingModel::Unlit:
+					MaterialProperty = EDMMaterialPropertyType::EmissiveColor;
+					break;
+
+				case EDMMaterialShadingModel::DefaultLit:
+					MaterialProperty = EDMMaterialPropertyType::BaseColor;
+					break;
+
+				default:
+					checkNoEntry();
+					break;
+			}
+		}
+		else
+		{
+			MaterialProperty = Slot->GetLayers().Last()->GetMaterialProperty();
+		}
+
+		if (!InNewMaskStage)
+		{
+			Slot->AddLayer(MaterialProperty, InNewBaseStage);
+		}
+		else
+		{
+			Slot->AddLayerWithMask(MaterialProperty, InNewBaseStage, InNewMaskStage);
+		}
+	}
+
+	UDMMaterialStageSource* const Source = InNewBaseStage->GetSource();
+	if (IsValid(Source))
+	{
+		Source->Update(EDMUpdateType::Structure);
+	}
+
+	UDMMaterialLayerObject* Layer = Slot->FindLayer(InNewBaseStage);
+	if (!Layer)
+	{
+		return nullptr;
+	}
+
+	if (ensure(LayerView.IsValid()))
+	{
+		LayerView->AddLayerItem(MakeShared<FDMMaterialLayerReference>(Layer));
+	}
+
+	return Layer;
+}
+
+void SDMSlot::AddNewLayer_NewLocalValue(EDMValueType InType)
+{
+	UDMMaterialSlot* Slot = SlotWeak.Get();
+
+	if (!ensure(Slot))
+	{
+		return;
+	}
+
+	UDMMaterialStage* NewBase = UDMMaterialStageBlend::CreateStage(UDMMaterialStageBlendNormal::StaticClass());
+	AddNewLayer(NewBase);
+
+	NewBase->ChangeInput_NewLocalValue(UDMMaterialStageBlend::InputB, FDMMaterialStageConnectorChannel::WHOLE_CHANNEL, InType, FDMMaterialStageConnectorChannel::WHOLE_CHANNEL);
+}
+
+void SDMSlot::AddNewLayer_GlobalValue(UDMMaterialValue* InValue)
+{
+	UDMMaterialSlot* Slot = SlotWeak.Get();
+
+	if (!ensure(Slot))
+	{
+		return;
+	}
+
+	UDynamicMaterialModelEditorOnlyData* ModelEditorOnlyData = Slot->GetMaterialModelEditorOnlyData();
+	if (!ensure(ModelEditorOnlyData))
+	{
+		return;
+	}
+
+	if (!ensure(IsValid(InValue)) || !ensure(InValue->GetMaterialModel() == ModelEditorOnlyData->GetMaterialModel()))
+	{
+		return;
+	}
+
+	UDMMaterialStage* NewBase = UDMMaterialStageBlend::CreateStage(UDMMaterialStageBlendNormal::StaticClass());
+	AddNewLayer(NewBase);
+
+	NewBase->ChangeInput_Value(UDMMaterialStageBlend::InputB, FDMMaterialStageConnectorChannel::WHOLE_CHANNEL,
+		InValue, FDMMaterialStageConnectorChannel::WHOLE_CHANNEL);
+}
+
+void SDMSlot::AddNewLayer_NewGlobalValue(EDMValueType InType)
+{
+	UDMMaterialSlot* Slot = SlotWeak.Get();
+
+	if (!ensure(Slot))
+	{
+		return;
+	}
+
+	UDMMaterialStage* NewStage = UDMMaterialStageBlend::CreateStage(UDMMaterialStageBlendNormal::StaticClass());
+	AddNewLayer(NewStage);
+
+	NewStage->ChangeInput_NewValue(UDMMaterialStageBlend::InputB, FDMMaterialStageConnectorChannel::WHOLE_CHANNEL, 
+		InType, FDMMaterialStageConnectorChannel::WHOLE_CHANNEL);
+}
+
+void SDMSlot::AddNewLayer_Slot(UDMMaterialSlot* InSlot, EDMMaterialPropertyType InMaterialProperty)
+{
+	UDMMaterialSlot* Slot = SlotWeak.Get();
+
+	if (!ensure(Slot))
+	{
+		return;
+	}
+
+	if (!ensure(IsValid(InSlot)) || !ensure(Slot != InSlot)
+		|| !ensure(InSlot->GetMaterialModelEditorOnlyData() == Slot->GetMaterialModelEditorOnlyData()))
+	{
+		return;
+	}
+
+	UDMMaterialStage* NewStage = UDMMaterialStageBlend::CreateStage(UDMMaterialStageBlendNormal::StaticClass());
+	AddNewLayer(NewStage);
+
+	NewStage->ChangeInput_Slot(UDMMaterialStageBlend::InputB, FDMMaterialStageConnectorChannel::WHOLE_CHANNEL, 
+		InSlot, InMaterialProperty, 0, FDMMaterialStageConnectorChannel::WHOLE_CHANNEL);
+}
+
+void SDMSlot::AddNewLayer_Expression(TSubclassOf<UDMMaterialStageExpression> InExpressionClass, EDMMaterialLayerStage LayerEnabledMask)
+{
+	UDMMaterialSlot* Slot = SlotWeak.Get();
+
+	if (!ensure(Slot))
+	{
+		return;
+	}
+
+	if (!ensure(IsValid(InExpressionClass)) || !ensure(InExpressionClass != UDMMaterialStageExpression::StaticClass()))
+	{
+		return;
+	}
+
+	// Extra transaction here because there are changes to the slot afterwards.
+	FScopedTransaction Transaction(LOCTEXT("AddLayerExpression", "Material Designer Add Layer (Expression)"));
+	Slot->Modify();
+
+	UDMMaterialStage* NewStage = UDMMaterialStageBlend::CreateStage(UDMMaterialStageBlendNormal::StaticClass());
+	UDMMaterialLayerObject* NewLayer = AddNewLayer(NewStage);
+
+	NewStage->ChangeInput_Expression(UDMMaterialStageBlend::InputB, FDMMaterialStageConnectorChannel::WHOLE_CHANNEL, 
+		InExpressionClass, 0, FDMMaterialStageConnectorChannel::WHOLE_CHANNEL);
+
+	switch (LayerEnabledMask)
+	{
+		case EDMMaterialLayerStage::Base:
+			if (UDMMaterialStage* MaskStage = NewLayer->GetStage(EDMMaterialLayerStage::Mask))
+			{
+				MaskStage->SetEnabled(false);
+			}
+			break;
+
+		case EDMMaterialLayerStage::Mask:
+			if (UDMMaterialStage* BaseStage = NewLayer->GetStage(EDMMaterialLayerStage::Base))
+			{
+				BaseStage->SetEnabled(false);
+			}
+			break;
+
+		default:
+			// No nothing
+			break;
+	}
+}
+
+void SDMSlot::AddNewLayer_Blend(TSubclassOf<UDMMaterialStageBlend> InBlendClass)
+{
+	UDMMaterialSlot* Slot = SlotWeak.Get();
+
+	if (!ensure(Slot))
+	{
+		return;
+	}
+
+	if (!ensure(IsValid(InBlendClass)) || !ensure(InBlendClass != UDMMaterialStageBlend::StaticClass()))
+	{
+		return;
+	}
+
+	UDMMaterialStage* NewStage = UDMMaterialStageBlend::CreateStage(InBlendClass);
+	AddNewLayer(NewStage);
+
+	NewStage->ChangeInput_Expression(UDMMaterialStageBlend::InputB, FDMMaterialStageConnectorChannel::WHOLE_CHANNEL, 
+		UDMMaterialStageExpressionTextureSample::StaticClass(), 0, FDMMaterialStageConnectorChannel::WHOLE_CHANNEL);
+}
+
+void SDMSlot::AddNewLayer_Gradient(TSubclassOf<UDMMaterialStageGradient> InGradientClass)
+{
+	UDMMaterialSlot* Slot = SlotWeak.Get();
+
+	if (!ensure(Slot))
+	{
+		return;
+	}
+
+	if (!ensure(IsValid(InGradientClass)) || !ensure(InGradientClass != UDMMaterialStageGradient::StaticClass()))
+	{
+		return;
+	}
+
+	UDMMaterialStage* NewStage = UDMMaterialStageBlend::CreateStage(UDMMaterialStageBlendNormal::StaticClass());
+	AddNewLayer(NewStage);
+
+	NewStage->ChangeInput_Gradient(UDMMaterialStageBlend::InputB, FDMMaterialStageConnectorChannel::WHOLE_CHANNEL, InGradientClass, FDMMaterialStageConnectorChannel::WHOLE_CHANNEL);
+}
+
+void SDMSlot::AddNewLayer_UV()
+{
+	UDMMaterialSlot* Slot = SlotWeak.Get();
+
+	if (!ensure(Slot))
+	{
+		return;
+	}
+
+	UDynamicMaterialModelEditorOnlyData* EditorOnlyData = Slot->GetMaterialModelEditorOnlyData();
+
+	if (!ensure(EditorOnlyData))
+	{
+		return;
+	}
+
+	UDynamicMaterialModel* MaterialModel = EditorOnlyData->GetMaterialModel();
+
+	if (!ensure(MaterialModel))
+	{
+		return;
+	}
+
+	UDMMaterialStage* NewStage = UDMMaterialStageInputTextureUV::CreateStage(MaterialModel);
+	AddNewLayer(NewStage);
+}
+
+void SDMSlot::AddNewLayer_MaterialFunction()
+{
+	UDMMaterialSlot* Slot = SlotWeak.Get();
+
+	if (!ensure(Slot))
+	{
+		return;
+	}
+
+	UDMMaterialStage* NewStage = UDMMaterialStageBlend::CreateStage(UDMMaterialStageBlendNormal::StaticClass());
+	AddNewLayer(NewStage);
+
+	NewStage->ChangeInput_MaterialFunction(UDMMaterialStageBlend::InputB, FDMMaterialStageConnectorChannel::WHOLE_CHANNEL);
+}
+
+void SDMSlot::AddNewLayer_SceneTexture()
+{
+	UDMMaterialStage* NewStage = UDMMaterialStageBlend::CreateStage(UDMMaterialStageBlendNormal::StaticClass());
+	UDMMaterialLayerObject* NewLayer = AddNewLayer(NewStage);
+
+	NewStage->ChangeInput_Expression(UDMMaterialStageBlend::InputB, FDMMaterialStageConnectorChannel::WHOLE_CHANNEL,
+		UDMMaterialStageExpressionSceneTexture::StaticClass(), 0, FDMMaterialStageConnectorChannel::THREE_CHANNELS);
+
+	if (UDMMaterialStage* MaskStage = NewLayer->GetStage(EDMMaterialLayerStage::Mask))
+	{
+		MaskStage->ChangeInput_Expression(UDMMaterialStageThroughputLayerBlend::InputMaskSource, FDMMaterialStageConnectorChannel::WHOLE_CHANNEL,
+			UDMMaterialStageExpressionSceneTexture::StaticClass(), 0, FDMMaterialStageConnectorChannel::FOURTH_CHANNEL);
+	}
+}
+
+#undef LOCTEXT_NAMESPACE
diff --git a/Engine/Plugins/Experimental/DynamicMaterial/Source/DynamicMaterialEditor/Private/Slate/SDMSlot.h b/Engine/Plugins/Experimental/DynamicMaterial/Source/DynamicMaterialEditor/Private/Slate/SDMSlot.h
new file mode 100644
--- /dev/null
+++ b/Engine/Plugins/Experimental/DynamicMaterial/Source/DynamicMaterialEditor/Private/Slate/SDMSlot.h
@@ -0,0 +1,172 @@
+// Copyright Epic Games, Inc. All Rights Reserved.
+
+#pragma once
+
+#include "DMDefs.h"
+#include "Slate/Layers/SDMSlotLayerView.h"
+#include "Slate/SDMEditor.h"
+#include "Widgets/SCompoundWidget.h"
+
+class SBox;
+class SDMEditor;
+class SDMSlotLayerView;
+class SImage;
+class SScrollBox;
+class SSplitter;
+class SStackBox;
+class UDMMaterialLayerObject;
+class UDMMaterialSlot;
+class UDMMaterialStage;
+class UDMMaterialStageBlend;
+class UDMMaterialStageExpression;
+class UDMMaterialStageGradient;
+class UDMMaterialStageInputValue;
+class UDMMaterialValue;
+class UDMMaterialValueFloat1;
+enum class EDMMaterialLayerStage : uint8;
+struct FPropertyChangedEvent;
+
+class SDMSlot : public SCompoundWidget
+{
+	SLATE_BEGIN_ARGS(SDMSlot)
+		: _SlotPreviewSize(64)
+		, _LayerPreviewSize(48)
+	{}
+		SLATE_ATTRIBUTE(int32, SlotPreviewSize)
+		SLATE_ATTRIBUTE(int32, LayerPreviewSize)
+	SLATE_END_ARGS()
+
+	friend class UDMMaterialStageExpression;
+
+public:
+	virtual ~SDMSlot() override;
+
+	void Construct(const FArguments& InArgs, const TSharedRef<SDMEditor>& InEditor, UDMMaterialSlot* InSlot);
+
+	TSharedPtr<SDMEditor> GetEditorWidget() const { return EditorWidgetWeak.Pin(); }
+	UDMMaterialSlot* GetSlot() const { return SlotWeak.Get(); }
+
+	TSharedPtr<SDMStage> FindStageWidget(UDMMaterialStage* const InStage) const;
+	
+	void RemoveLayerByStage(UDMMaterialStage* const InStage, const bool bInSelectNextMaskIfMask = false);
+
+	int32 GetLayerCount() const;
+
+	TArray<int32> GetSelectedLayerIndices() const;
+	void ClearSelection();
+
+	void InvalidateMainWidget();
+	void InvalidateHeaderPropertyListWidget();
+	void InvalidateSlotSettingsRowWidget();
+	void InvalidateComponentEditWidget();
+
+	void AddNewLayer_NewLocalValue(EDMValueType InType);
+	void AddNewLayer_GlobalValue(UDMMaterialValue* InValue);
+	void AddNewLayer_NewGlobalValue(EDMValueType InType);
+	void AddNewLayer_Slot(UDMMaterialSlot* InSlot, EDMMaterialPropertyType InMaterialProperty);
+	void AddNewLayer_Expression(TSubclassOf<UDMMaterialStageExpression> InExpressionClass, EDMMaterialLayerStage LayerEnabledMask);
+	void AddNewLayer_Blend(TSubclassOf<UDMMaterialStageBlend> InBlendClass);
+	void AddNewLayer_Gradient(TSubclassOf<UDMMaterialStageGradient> InGradientClass);
+	void AddNewLayer_UV();
+	void AddNewLayer_MaterialFunction();
+	void AddNewLayer_SceneTexture();
+
+	void RemoveSlot();
+
+	bool CheckValidity();
+
+	bool GetLayerRowsButtonsCanRemove() const;
+	FReply OnLayerRowButtonsRemoveClicked();
+
+	UDMMaterialLayerObject* GetSelectedLayer() const;
+
+	void SetSelectedLayer(UDMMaterialLayerObject* InLayer) const;
+
+	UDMMaterialComponent* GetEditedComponent() const;
+
+	void SetEditedComponent(UDMMaterialComponent* InComponent);
+
+protected:
+	TWeakPtr<SDMEditor> EditorWidgetWeak;
+
+	TAttribute<int32> SlotPreviewSize;
+	TAttribute<int32> LayerPreviewSize;
+
+	TWeakObjectPtr<UDMMaterialSlot> SlotWeak;
+	TSharedPtr<SBox> HeaderPropertyListWidget;
+	TSharedPtr<SBox> SlotSettingsRowContainer;
+	TSharedPtr<SScrollBox> ComponentEditContainer;
+
+	TSharedPtr<SSplitter> SplitterContainer;
+	SSplitter::FSlot* LayerViewSplitterSlot = nullptr;
+	SSplitter::FSlot* ExtraSpaceSplitterSlot = nullptr;
+
+	TSharedPtr<SDMSlotLayerView> LayerView;
+
+	TSharedPtr<FScopedTransaction> ScrubbingTransaction;
+	float OriginalOpacity;
+
+	TWeakObjectPtr<UDMMaterialStageInputValue> SelectedOpacityStageInputValue;
+	TWeakObjectPtr<UDMMaterialValueFloat1> SelectedOpacityInputValue;
+
+	bool bInvalidateMainWidget;
+	bool bInvalidateHeaderWidget;
+	bool bInvalidateSettingsWidget;
+	bool bInvalidateComponentEditWidget;
+	FDelegateHandle OnEndFrameDelegateHandle;
+
+	TWeakObjectPtr<UDMMaterialComponent> EditedComponent;
+
+	TSharedRef<SWidget> CreateHeaderPropertyListWidget();
+	TSharedRef<SWidget> CreateLayerButtonsRowWidget();
+	TSharedRef<SWidget> CreateComponentEditWidget();
+
+	TSharedRef<SWidget> CreateSlotSettingsRow();
+
+	void OnSplitterResized() const;
+
+	UDMMaterialLayerObject* AddNewLayer(UDMMaterialStage* InNewBaseStage, UDMMaterialStage* InNewMaskStage = nullptr);
+
+	void AddPropertyToSlot(EDMMaterialPropertyType Property);
+
+	void OnSlotLayersUpdated(UDMMaterialSlot* InSlot);
+	void OnSlotPropertiesUpdated(UDMMaterialSlot* InSlot);
+	void OnStageUpdated(UDMMaterialStage* InStage, EDMUpdateType InUpdateType);
+
+	void StartStageOpacityTransaction();
+	void EndStageOpacityTransaction(const float InValue) { EndTransaction(); }
+	void SetStageOpacity(const float InNewValue);
+	void OnStageOpacityChanged(const float InNewValue);
+	void OnStageOpacityCommitted(const float InNewValue, const ETextCommit::Type InCommitType);
+
+	void StartTransaction(const FText InDescription);
+	void EndTransaction();
+
+	void OnLayerSelected(TSharedPtr<FDMMaterialLayerReference> InLayerItem, const int32 InLayerIndex);
+	void OnLayerStageSelected(const bool bInSelected, const TSharedRef<SDMStage>& InStageWidget);
+
+	void RefreshMainWidget();
+	void RefreshHeaderPropertyListWidget();
+	void RefreshSlotSettingsRowWidget();
+	void RefreshComponentEditWidget();
+
+	void HandleEndFrameRefresh();
+
+	FText GetLayerButtonsDescription() const;
+	TSharedRef<SWidget> GetLayerButtonsMenuContent();
+
+	bool GetLayerRowsButtonsCanDuplicate() const;
+	FReply OnLayerRowButtonsDuplicateClicked();
+
+	int32 GetSourceBlendTypeSwitcherIndex() const;
+	TSubclassOf<UDMMaterialStageBlend> GetSelectedSourceBlendType() const;
+	void OnSourceBlendTypedSelected(const TSubclassOf<UDMMaterialStageBlend> InNewItem);
+
+	bool IsOpacityPropertyEnabled() const;
+
+	bool CanRemoveLayerByIndex(const int32 InLayerIndex) const;
+	void RemoveLayerByIndex(const int32 InLayerIndex, const bool bInSelectNextMaskIfMask = false);
+
+	bool CanRemoveSlot() const;
+	FReply OnRemoveSlotClicked();
+};
diff --git a/Engine/Plugins/Experimental/DynamicMaterial/Source/DynamicMaterialEditor/Private/Slate/SDMStage.cpp b/Engine/Plugins/Experimental/DynamicMaterial/Source/DynamicMaterialEditor/Private/Slate/SDMStage.cpp
new file mode 100644
--- /dev/null
+++ b/Engine/Plugins/Experimental/DynamicMaterial/Source/DynamicMaterialEditor/Private/Slate/SDMStage.cpp
@@ -0,0 +1,633 @@
+// Copyright Epic Games, Inc. All Rights Reserved.
+
+#include "SDMStage.h"
+#include "Components/DMMaterialSlot.h"
+#include "Components/DMMaterialStage.h"
+#include "Components/DMMaterialStageBlend.h"
+#include "Components/DMMaterialStageThroughput.h"
+#include "Components/DMMaterialStageThroughputLayerBlend.h"
+#include "Components/DMMaterialSubStage.h"
+#include "Components/MaterialStageExpressions/DMMSETextureSample.h"
+#include "Components/MaterialStageInputs/DMMSIExpression.h"
+#include "Components/MaterialStageInputs/DMMSIValue.h"
+#include "Components/MaterialValues/DMMaterialValueTexture.h"
+#include "DragAndDrop/AssetDragDropOp.h"
+#include "DragDrop/DMStageDragDropOperation.h"
+#include "DynamicMaterialEditorSettings.h"
+#include "DynamicMaterialEditorStyle.h"
+#include "Engine/Texture.h"
+#include "ScopedTransaction.h"
+#include "SDMEditor.h"
+#include "Slate/Previews/SDMStagePreview.h"
+#include "Slate/SMaterialToolTip.h"
+#include "SlateOptMacros.h"
+#include "Components/DMMaterialLayer.h"
+#include "Styling/CoreStyle.h"
+#include "Widgets/Colors/SColorBlock.h"
+#include "Widgets/Layout/SBox.h"
+#include "Widgets/Layout/SScaleBox.h"
+#include "Widgets/Layout/SScissorRectBox.h"
+#include "Widgets/SOverlay.h"
+#include "Widgets/SToolTip.h"
+#include "Widgets/Text/STextBlock.h"
+#include "Widgets/Views/STableViewBase.h"
+
+#define LOCTEXT_NAMESPACE "SDMStage"
+
+BEGIN_SLATE_FUNCTION_BUILD_OPTIMIZATION
+
+SDMStage::~SDMStage()
+{
+	SDMEditor::ClearPropertyHandles(this);
+}
+
+void SDMStage::Construct(const FArguments& InArgs, UDMMaterialStage* InStage)
+{
+	ensure(IsValid(InStage));
+	StageWeak = InStage;
+
+	bStageEnabled = InArgs._StageEnabled;
+	bStageSelected = InArgs._StageSelected;
+	bIsMask = InArgs._StageIsMask;
+	bInteractable = InArgs._Interactable;
+	bShowTextOverlays = InArgs._ShowTextOverlays;
+	DesiredSize = InArgs._DesiredSize;
+	MinUniformSize = InArgs._MinUniformSize;
+	MaxUniformSize = InArgs._MaxUniformSize;
+	OnUniformSizeChanged = InArgs._OnUniformSizeChanged;
+	OnClick = InArgs._OnClick;
+
+	OnCanAcceptDropEvent = InArgs._OnCanAcceptDrop;
+	OnAcceptDropEvent = InArgs._OnAcceptDrop;
+	OnPaintDropIndicatorEvent = InArgs._OnPaintDropIndicator;
+	
+	FText ToolTipText;
+	if (bInteractable.Get())
+	{
+		SetCursor(EMouseCursor::Hand);
+		ToolTipText = LOCTEXT("MaterialStagePreviewTooltip", "Material Stage\n\nRight click for options.\nControl+LMB to remove.\nAlt+LMB to toggle.");
+	}
+	else
+	{
+		ToolTipText = LOCTEXT("MaterialStagePreviewTooltipNoInteraction", "Material Stage Preview");
+	}
+
+	SetToolTip(
+		SNew(SToolTip)
+		.IsInteractive(false)
+		.BorderImage(FCoreStyle::Get().GetBrush("ToolTip.Background"))
+		[
+			SNew(SMaterialToolTip)
+			.Material(StageWeak->GetPreviewMaterial())
+			.Text(ToolTipText)
+			.MaterialSize_Lambda([]() { return FVector2D(UDynamicMaterialEditorSettings::Get()->TooltipTextureSize); })
+			.ShowMaterial(this, &SDMStage::IsToolTipInteractable)
+		]
+	);
+
+	const float BorderWidth = 1.0f;
+
+	ChildSlot
+	[
+		SNew(SBorder)
+		.BorderBackgroundColor(FLinearColor(1, 1, 1, 0.5f))
+		.Padding(2.0f)
+		.BorderImage(FAppStyle::Get().GetBrush("ContentBrowser.AssetTileItem.DropShadow"))
+		[
+			SNew(SOverlay)
+			+ SOverlay::Slot()
+			.HAlign(HAlign_Center)
+			.VAlign(VAlign_Center)
+			.Padding(BorderWidth)
+			[
+				SNew(SBorder)
+				.Clipping(EWidgetClipping::ClipToBounds)
+				.BorderBackgroundColor(FLinearColor::Transparent)
+				[
+					SNew(SBox)
+					.Clipping(EWidgetClipping::ClipToBounds)
+					.WidthOverride(this, &SDMStage::GetStagePreviewDesiredWidth)
+					.HeightOverride(this, &SDMStage::GetStagePreviewDesiredHeight)
+					[
+						SNew(SDMStagePreview, StageWeak.Get())
+						.PreviewSize(DesiredSize)
+					]
+				]
+			]
+			+ SOverlay::Slot()
+			.HAlign(HAlign_Fill)
+			.VAlign(VAlign_Bottom)
+			[
+				CreateTextBlockBackground()
+			]
+			+ SOverlay::Slot()
+			.HAlign(HAlign_Fill)
+			.VAlign(VAlign_Fill)
+			.Padding(1.0f)
+			[
+				SNew(SBorder)
+				.BorderBackgroundColor(this, &SDMStage::GetBorderBrushBackgroundColor)
+				.BorderImage(this, &SDMStage::GetBorderBrush)
+			]
+			+ SOverlay::Slot()
+			.HAlign(HAlign_Fill)
+			.VAlign(VAlign_Fill)
+			[
+				SNew(SBorder)
+				.Visibility(this, &SDMStage::GetOverlayBorderBrushVisibility)
+				.BorderImage(this, &SDMStage::GetOverlayBorderBrush)
+			]
+		]
+	];
+}
+
+TSharedRef<SWidget> SDMStage::CreateTextBlockBackground()
+{
+	const float DesiredHeight = FMath::Min(DesiredSize.Get().Y / 4.0f, 20.0f);
+	const int32 DefaultLabelFontSize = 12;
+
+	return 
+		SNew(SBox)
+		.HAlign(HAlign_Fill)
+		.VAlign(VAlign_Bottom)
+		.HeightOverride(DesiredHeight)
+		.Visibility(this, &SDMStage::GetTextBlockBackgroundVisibility)
+		[
+			SNew(SBorder)
+			.BorderImage(FDynamicMaterialEditorStyle::GetBrush("Stage.TextBlock"))
+			[
+				SNew(SScaleBox)
+				.StretchDirection(EStretchDirection::Both)
+				.Stretch(EStretch::ScaleToFitX)
+				[
+					SNew(STextBlock)
+					.Margin(3.0f)
+					.ColorAndOpacity(FSlateColor(FLinearColor::White))
+					.Text(this, &SDMStage::GetStageName)
+					.Font(FCoreStyle::GetDefaultFontStyle("Regular", DefaultLabelFontSize))
+				]
+			]
+		];
+}
+
+END_SLATE_FUNCTION_BUILD_OPTIMIZATION
+
+EVisibility SDMStage::GetBorderVisibility() const
+{
+	if (bInteractable.Get())
+	{
+		return EVisibility::Visible;
+	}
+	return EVisibility::Hidden;
+}
+
+const FSlateBrush* SDMStage::GetBorderBrush() const
+{
+	FString BrushName = "Stage.Inactive";
+
+	if (!bInteractable.Get() || bStageSelected.Get())
+	{
+		return FDynamicMaterialEditorStyle::GetBrush(*BrushName);
+	}
+
+	if (!bStageEnabled.Get())
+	{
+		BrushName = "Stage.Disabled";
+	}
+
+	if (IsHovered())
+	{
+		BrushName.Append(".Hover");
+	}
+
+	return FDynamicMaterialEditorStyle::GetBrush(*BrushName);
+}
+
+const FSlateBrush* SDMStage::GetOverlayBorderBrush() const
+{
+	FString BrushName = "Stage.Inactive";
+
+	if (!bInteractable.Get())
+	{
+		return FDynamicMaterialEditorStyle::GetBrush(*BrushName);
+	}
+
+	if (!bStageEnabled.Get())
+	{
+		BrushName = "Stage.Disabled";
+	}
+
+	if (bStageSelected.Get())
+	{
+		BrushName.Append(".Select");
+	}
+
+	if (IsHovered())
+	{
+		BrushName.Append(".Hover");
+	}
+
+	return FDynamicMaterialEditorStyle::GetBrush(*BrushName);
+}
+
+FText SDMStage::GetStageName() const
+{
+	return StageWeak.IsValid() ? StageWeak->GetComponentDescription() : FText::GetEmpty();
+}
+
+FCursorReply SDMStage::OnCursorQuery(const FGeometry& InMyGeometry, const FPointerEvent& InCursorEvent) const
+{
+	if (HasMouseCapture())
+	{
+		return FCursorReply::Cursor(EMouseCursor::None);
+	}
+
+	return SCompoundWidget::OnCursorQuery(InMyGeometry, InCursorEvent);
+}
+
+FReply SDMStage::OnMouseMove(const FGeometry& InMyGeometry, const FPointerEvent& InMouseEvent)
+{
+	if (OnUniformSizeChanged.IsBound() && InMouseEvent.IsMouseButtonDown(EKeys::MiddleMouseButton))
+	{
+		FReply Reply = FReply::Handled();
+
+		if (!HasMouseCapture())
+		{
+			CachedCursorLocation = InMouseEvent.GetScreenSpacePosition();
+			CachedPreviewSize = UDynamicMaterialEditorSettings::Get()->LayerPreviewSize;
+
+			Reply.CaptureMouse(AsShared()).UseHighPrecisionMouseMovement(AsShared());
+		}
+
+		FTableViewDimensions CursorDeltaDimensions(Orient_Vertical, InMouseEvent.GetCursorDelta());
+		const float SizeChange = CursorDeltaDimensions.ScrollAxis / InMyGeometry.Scale;
+
+		CachedPreviewSize += SizeChange;
+		CachedPreviewSize = FMath::Clamp(CachedPreviewSize, MinUniformSize, MaxUniformSize);
+
+		OnUniformSizeChanged.Execute(CachedPreviewSize, SharedThis(this));
+
+		return Reply;
+	}
+
+	return SCompoundWidget::OnMouseMove(InMyGeometry, InMouseEvent);
+}
+
+FReply SDMStage::OnMouseButtonUp(const FGeometry& InMyGeometry, const FPointerEvent& InMouseEvent)
+{
+	if (HasMouseCapture())
+	{
+		FReply Reply = FReply::Handled().ReleaseMouseCapture();
+
+		if (OnUniformSizeChanged.IsBound() && InMouseEvent.IsMouseButtonDown(EKeys::MiddleMouseButton))
+		{
+			Reply.SetMousePos(FIntPoint(CachedCursorLocation.X, CachedCursorLocation.Y));
+		}
+
+		return Reply;
+	}
+	
+	return SCompoundWidget::OnMouseButtonUp(InMyGeometry, InMouseEvent);
+}
+
+FReply SDMStage::OnMouseButtonDown(const FGeometry& MyGeometry, const FPointerEvent& MouseEvent)
+{
+	if (!bInteractable.Get())
+	{
+		return FReply::Unhandled();
+	}
+
+	const FKey MouseButton = MouseEvent.GetEffectingButton();
+
+	OnClick.ExecuteIfBound(MouseEvent, SharedThis(this));
+	
+	if (MouseEvent.IsMouseButtonDown(EKeys::LeftMouseButton))
+	{
+		return FReply::Handled().DetectDrag(SharedThis(this), EKeys::LeftMouseButton);
+	}
+
+	return SCompoundWidget::OnMouseButtonDown(MyGeometry, MouseEvent);
+}
+
+FReply SDMStage::OnDragDetected(const FGeometry& InMyGeometry, const FPointerEvent& InMouseEvent)
+{
+	TSharedRef<FDMStageDragDropOperation> Operation = MakeShared<FDMStageDragDropOperation>(SharedThis(this));
+	return FReply::Handled().BeginDragDrop(Operation);
+}
+
+void SDMStage::OnDragEnter(FGeometry const& InMyGeometry, FDragDropEvent const& InDragDropEvent)
+{
+	SCompoundWidget::OnDragEnter(InMyGeometry, InDragDropEvent);
+
+	if (!bInteractable.Get())
+	{
+		return;
+	}
+
+	if (!StageWeak.IsValid())
+	{
+		return;
+	}
+
+	if (const TSharedPtr<FDMStageDragDropOperation> StageDragDropOp = InDragDropEvent.GetOperationAs<FDMStageDragDropOperation>())
+	{
+		HandleStageDragDropOperation(*StageDragDropOp);
+	}
+	else if (const TSharedPtr<FAssetDragDropOp> AssetDragDropOp = InDragDropEvent.GetOperationAs<FAssetDragDropOp>())
+	{
+		if (UDMMaterialStage* DraggedOverStage = StageWeak.Get())
+		{
+			if (UDMMaterialLayerObject* DraggedOverStageLayer = DraggedOverStage->GetLayer())
+			{
+				UDMMaterialStage* LayerBase = DraggedOverStageLayer->GetStage(EDMMaterialLayerStage::Base);
+				UDMMaterialStage* LayerMask = DraggedOverStageLayer->GetStage(EDMMaterialLayerStage::Mask);
+
+				if ((LayerBase && LayerBase == DraggedOverStage && LayerBase->IsEnabled())
+					|| (LayerMask && LayerMask == DraggedOverStage && LayerMask->IsEnabled()))
+				{
+					// @TODO: some visual cue for drag over
+					//HighlightColor = FStyleColors::AccentBlue.GetSpecifiedColor();
+					//HighlightColor.A = 0.3;
+				}
+			}
+		}
+	}
+}
+
+void SDMStage::OnDragLeave(FDragDropEvent const& InDragDropEvent)
+{
+	SCompoundWidget::OnDragLeave(InDragDropEvent);
+
+	if (!bInteractable.Get())
+	{
+		return;
+	}
+
+	const TSharedPtr<FDMStageDragDropOperation> StageDragDropOp = InDragDropEvent.GetOperationAs<FDMStageDragDropOperation>();
+
+	if (StageDragDropOp.IsValid())
+	{
+		StageDragDropOp->SetToValidDropLocation();
+	}
+}
+
+FReply SDMStage::OnDragOver(const FGeometry& InMyGeometry, const FDragDropEvent& InDragDropEvent)
+{
+	if (OnCanAcceptDropEvent.IsBound())
+	{
+		OnCanAcceptDropEvent.Execute(InDragDropEvent, SharedThis(this));
+		return FReply::Handled();
+	}
+	return FReply::Unhandled();
+}
+
+FReply SDMStage::OnDrop(const FGeometry& InMyGeometry, const FDragDropEvent& InDragDropEvent)
+{
+	if (!bInteractable.Get())
+	{
+		return FReply::Unhandled();
+	}
+
+	SCompoundWidget::OnDrop(InMyGeometry, InDragDropEvent);
+
+	if (!StageWeak.IsValid())
+	{
+		return FReply::Unhandled();
+	}
+
+	FReply Reply = FReply::Unhandled();
+
+	if (const TSharedPtr<FAssetDragDropOp> AssetDragDropOp = InDragDropEvent.GetOperationAs<FAssetDragDropOp>())
+	{
+		HandleAssetDragDropOperation(*AssetDragDropOp);
+		Reply = FReply::Handled();
+	}
+
+	if (OnAcceptDropEvent.IsBound())
+	{
+		return OnAcceptDropEvent.Execute(InDragDropEvent, SharedThis(this));
+	}
+
+	return Reply;
+}
+
+void SDMStage::HandleStageDragDropOperation(FDMStageDragDropOperation& StageDragDropOperation)
+{
+	UDMMaterialStage* DraggedOverStage = StageWeak.Get();
+	
+	UDMMaterialStage* DraggedStage = StageDragDropOperation.GetStage();
+	if (!IsValid(DraggedStage))
+	{
+		return;
+	}
+
+	if (DraggedStage == DraggedOverStage)
+	{
+		StageDragDropOperation.SetToValidDropLocation();
+		return;
+	}
+
+	UDMMaterialLayerObject* DraggedStageLayer = DraggedStage->GetLayer();
+	UDMMaterialLayerObject* DraggedOverStageLayer = DraggedOverStage->GetLayer();
+
+	if (!DraggedStageLayer || !DraggedOverStageLayer)
+	{
+		StageDragDropOperation.SetToInvalidDropLocation();
+		return;
+	}
+
+	if (DraggedStageLayer == DraggedOverStageLayer || DraggedStageLayer->GetSlot() != DraggedOverStageLayer->GetSlot())
+	{
+		StageDragDropOperation.SetToInvalidDropLocation();
+		return;
+	}
+
+	const int32 DraggedIndex = DraggedStageLayer->FindIndex();
+	const int32 DraggedOverIndex = DraggedOverStageLayer->FindIndex();
+	if (DraggedIndex == INDEX_NONE || DraggedOverIndex == INDEX_NONE)
+	{
+		return;
+	}
+
+	bool bReject = false;
+
+	if (DraggedOverIndex == 0 && !DraggedStageLayer->GetStage(EDMMaterialLayerStage::Base))
+	{
+		return;
+	}
+
+	if (!bReject)
+	{
+		const UDMMaterialLayerObject* PreviousLayer = DraggedOverStageLayer->GetPreviousLayer(EDMMaterialPropertyType::Any, EDMMaterialLayerStage::None);
+		const UDMMaterialLayerObject* NextLayer = DraggedOverStageLayer->GetNextLayer(EDMMaterialPropertyType::Any, EDMMaterialLayerStage::None);
+
+		if (PreviousLayer || NextLayer)
+		{
+			if (DraggedOverIndex < DraggedIndex)
+			{
+				if (PreviousLayer &&
+					(!DraggedStage->IsCompatibleWithPreviousStage(PreviousLayer->GetStage(EDMMaterialLayerStage::Base))
+						|| !DraggedStage->IsCompatibleWithNextStage(DraggedOverStage)))
+				{
+					bReject = true;
+				}
+			}
+			else
+			{
+				if (NextLayer &&
+					(!DraggedStage->IsCompatibleWithPreviousStage(DraggedOverStage)
+						|| !DraggedStage->IsCompatibleWithNextStage(NextLayer->GetStage(EDMMaterialLayerStage::Base))))
+				{
+					bReject = true;
+				}
+			}
+
+			if (!bReject)
+			{
+				PreviousLayer = DraggedStageLayer->GetPreviousLayer(EDMMaterialPropertyType::Any, EDMMaterialLayerStage::None);
+				NextLayer = DraggedStageLayer->GetNextLayer(EDMMaterialPropertyType::Any, EDMMaterialLayerStage::None);
+
+				if (PreviousLayer && NextLayer)
+				{
+					if (!PreviousLayer->GetStage(EDMMaterialLayerStage::Base)->IsCompatibleWithNextStage(NextLayer->GetStage(EDMMaterialLayerStage::Base)))
+					{
+						bReject = true;
+					}
+
+					if (!NextLayer->GetStage(EDMMaterialLayerStage::Base)->IsCompatibleWithPreviousStage(PreviousLayer->GetStage(EDMMaterialLayerStage::Base)))
+					{
+						bReject = true;
+					}
+				}
+			}
+		}
+	}
+
+	if (bReject)
+	{
+		StageDragDropOperation.SetToInvalidDropLocation();
+		return;
+	}
+
+	StageDragDropOperation.SetToValidDropLocation();
+}
+
+void SDMStage::HandleAssetDragDropOperation(FAssetDragDropOp& AssetDragDropOperation)
+{
+	UDMMaterialStage* DraggedOverStage = StageWeak.Get();
+	if (!IsValid(DraggedOverStage))
+	{
+		return;
+	}
+
+	const UDMMaterialLayerObject* StageLayer = DraggedOverStage->GetLayer();
+	if (!StageLayer)
+	{
+		return;
+	}
+
+	if ((StageLayer->GetStage(EDMMaterialLayerStage::Base) == DraggedOverStage && !StageLayer->IsStageEnabled(EDMMaterialLayerStage::Base))
+		|| (StageLayer->GetStage(EDMMaterialLayerStage::Mask) == DraggedOverStage && !StageLayer->IsStageEnabled(EDMMaterialLayerStage::Mask)))
+	{
+		return;
+	}
+
+	if (!DraggedOverStage->GetSource())
+	{
+		return;
+	}
+
+	UTexture* Texture = nullptr;
+
+	for (const FAssetData& AssetData : AssetDragDropOperation.GetAssets())
+	{
+		UObject* Asset = AssetData.GetAsset();
+		if (IsValid(Asset))
+		{
+			if (UTexture* TextureIter = Cast<UTexture>(Asset))
+			{
+				Texture = TextureIter;
+				break;
+			}
+		}
+	}
+
+	if (!Texture)
+	{
+		return;
+	}
+
+	FScopedTransaction Transaction(LOCTEXT("DragTextureOntoStage", "Material Designer Drag Texture onto Stage"), !FDMInitializationGuard::IsInitializing());
+	UDMMaterialValueTexture* TextureValue = nullptr;
+
+	DraggedOverStage->Modify();
+
+	UDMMaterialStageSource* DraggedOverStageSource = DraggedOverStage->GetSource();
+
+	if (DraggedOverStageSource->IsA<UDMMaterialStageBlend>())
+	{
+		UDMMaterialStageInputExpression* NewInput = DraggedOverStage->ChangeInput_Expression(UDMMaterialStageBlend::InputB, FDMMaterialStageConnectorChannel::WHOLE_CHANNEL,
+			UDMMaterialStageExpressionTextureSample::StaticClass(), 0, FDMMaterialStageConnectorChannel::WHOLE_CHANNEL);
+
+		UDMMaterialSubStage* SubStage = NewInput->GetSubStage();
+		UDMMaterialStageInputValue* NewInputValue = SubStage->ChangeInput_NewLocalValue(0, FDMMaterialStageConnectorChannel::WHOLE_CHANNEL, EDMValueType::VT_Texture, FDMMaterialStageConnectorChannel::WHOLE_CHANNEL);
+
+		TextureValue = Cast<UDMMaterialValueTexture>(NewInputValue->GetValue());
+	}
+	else if (DraggedOverStageSource->IsA<UDMMaterialStageThroughputLayerBlend>())
+	{
+		bool bHasAlpha = UE::DynamicMaterial::Private::HasAlpha(Texture);
+
+		UDMMaterialStageInputExpression* NewInput = DraggedOverStage->ChangeInput_Expression(2, FDMMaterialStageConnectorChannel::WHOLE_CHANNEL,
+			UDMMaterialStageExpressionTextureSample::StaticClass(), bHasAlpha ? 1 : 0, FDMMaterialStageConnectorChannel::WHOLE_CHANNEL);
+
+		UDMMaterialSubStage* SubStage = NewInput->GetSubStage();
+		UDMMaterialStageInputValue* NewInputValue = SubStage->ChangeInput_NewLocalValue(0, FDMMaterialStageConnectorChannel::WHOLE_CHANNEL, EDMValueType::VT_Texture, FDMMaterialStageConnectorChannel::WHOLE_CHANNEL);
+
+		TextureValue = Cast<UDMMaterialValueTexture>(NewInputValue->GetValue());
+	}
+	else
+	{
+		UDMMaterialStageExpression* NewExpression = DraggedOverStage->ChangeSource_Expression(UDMMaterialStageExpressionTextureSample::StaticClass());
+		UDMMaterialStageInputValue* NewInputValue = DraggedOverStage->ChangeInput_NewLocalValue(0, FDMMaterialStageConnectorChannel::WHOLE_CHANNEL, EDMValueType::VT_Texture, FDMMaterialStageConnectorChannel::WHOLE_CHANNEL);
+
+		TextureValue = Cast<UDMMaterialValueTexture>(NewInputValue->GetValue());
+	}
+
+	if (TextureValue)
+	{
+		TextureValue->SetValue(Texture);
+	}
+}
+
+bool SDMStage::IsToolTipInteractable() const
+{
+	return bInteractable.Get() && UDynamicMaterialEditorSettings::Get()->bShowTooltipPreview;
+}
+
+FOptionalSize SDMStage::GetStagePreviewDesiredWidth() const
+{
+	return DesiredSize.Get().X;
+}
+
+FOptionalSize SDMStage::GetStagePreviewDesiredHeight() const
+{
+	return DesiredSize.Get().Y;
+}
+
+FSlateColor SDMStage::GetBorderBrushBackgroundColor() const
+{
+	return bInteractable.Get() ? FLinearColor::White : FLinearColor::Transparent;
+}
+
+EVisibility SDMStage::GetOverlayBorderBrushVisibility() const
+{
+	return bStageSelected.Get() ? EVisibility::SelfHitTestInvisible : EVisibility::Hidden;
+}
+
+EVisibility SDMStage::GetTextBlockBackgroundVisibility() const
+{
+	return bShowTextOverlays.Get() ? EVisibility::SelfHitTestInvisible : EVisibility::Collapsed;
+}
+
+#undef LOCTEXT_NAMESPACE
diff --git a/Engine/Plugins/Experimental/DynamicMaterial/Source/DynamicMaterialEditor/Private/Slate/SDMStage.h b/Engine/Plugins/Experimental/DynamicMaterial/Source/DynamicMaterialEditor/Private/Slate/SDMStage.h
new file mode 100644
--- /dev/null
+++ b/Engine/Plugins/Experimental/DynamicMaterial/Source/DynamicMaterialEditor/Private/Slate/SDMStage.h
@@ -0,0 +1,139 @@
+// Copyright Epic Games, Inc. All Rights Reserved.
+
+#pragma once
+
+#include "Widgets/SCompoundWidget.h"
+
+class FAssetDragDropOp;
+class FDMStageDragDropOperation;
+class UDMMaterialStage;
+class SDMStage;
+struct FOptionalSize;
+
+DECLARE_DELEGATE_TwoParams(FDMOnStageClick, const FPointerEvent&, const TSharedRef<SDMStage>& /** StageWidget */);
+DECLARE_DELEGATE_TwoParams(FDMOnStageUniformSizeChanged, const float /** NewSize */, const TSharedRef<SDMStage>& /** StageWidget */);
+
+DECLARE_DELEGATE_OneParam(FDMOnStageDragEnter, const FDragDropEvent&);
+DECLARE_DELEGATE_OneParam(FDMOnStageDragLeave, const FDragDropEvent&);
+DECLARE_DELEGATE_RetVal_OneParam(FReply, FDMOnStageDragDetected, const FDragDropEvent&);
+DECLARE_DELEGATE_RetVal_OneParam(FReply, FDMOnStageDrop, const FDragDropEvent&);
+
+/** Delegate signature for querying whether this FDragDropEvent will be handled by the drop target. */
+DECLARE_DELEGATE_RetVal_TwoParams(bool, FDMOnStageCanAcceptDrop, const FDragDropEvent&, const TSharedRef<SDMStage>& /** StageWidget */);
+/** Delegate signature for handling the drop of FDragDropEvent onto target. */
+DECLARE_DELEGATE_RetVal_TwoParams(FReply, FDMOnStageAcceptDrop, const FDragDropEvent&, const TSharedRef<SDMStage>& /** StageWidget */);
+/** Delegate signature for painting drop indicators. */
+DECLARE_DELEGATE_RetVal_SevenParams(int32, FDMOnStagePaintDropIndicator, const FPaintArgs&, const FGeometry&, const FSlateRect&, FSlateWindowElementList&, int32, const FWidgetStyle&, bool);
+
+class SDMStage : public SCompoundWidget
+{
+public:
+	SLATE_BEGIN_ARGS(SDMStage)
+		: _StageIsMask(false)
+		, _Interactable(false)
+		, _ShowTextOverlays(false)
+		, _DesiredSize(FVector2D(48.0f))
+		, _MinUniformSize(32)
+		, _MaxUniformSize(128)
+	{}
+		SLATE_ATTRIBUTE(bool, StageEnabled)
+		SLATE_ATTRIBUTE(bool, StageSelected)
+		SLATE_ARGUMENT(bool, StageIsMask)
+		SLATE_ATTRIBUTE(bool, Interactable)
+		SLATE_ATTRIBUTE(bool, ShowTextOverlays)
+		SLATE_ATTRIBUTE(FVector2D, DesiredSize)
+		SLATE_ARGUMENT(float, MinUniformSize)
+		SLATE_ARGUMENT(float, MaxUniformSize)
+		SLATE_EVENT(FDMOnStageUniformSizeChanged, OnUniformSizeChanged)
+		SLATE_EVENT(FDMOnStageClick, OnClick)
+
+		/** High Level DragAndDrop */
+		/**
+		 * Handle this event to determine whether a drag and drop operation can be executed on top of the target row widget.
+		 * Most commonly, this is used for previewing re-ordering and re-organization operations in lists or trees.
+		 * e.g. A user is dragging one item into a different spot in the list or tree.
+		 *      This delegate will be called to figure out if we should give visual feedback on whether an item will
+		 *      successfully drop into the list.
+		 */
+		SLATE_EVENT(FDMOnStageCanAcceptDrop, OnCanAcceptDrop)
+		/**
+		 * Perform a drop operation onto the target row widget
+		 * Most commonly used for executing a re-ordering and re-organization operations in lists or trees.
+		 * e.g. A user was dragging one item into a different spot in the list; they just dropped it.
+		 *      This is our chance to handle the drop by reordering items and calling for a list refresh.
+		 */
+		SLATE_EVENT(FDMOnStageAcceptDrop, OnAcceptDrop)
+		/**
+		 * Used for painting drop indicators
+		 */
+		SLATE_EVENT(FDMOnStagePaintDropIndicator, OnPaintDropIndicator)
+	SLATE_END_ARGS()
+
+	virtual ~SDMStage() override;
+
+	void Construct(const FArguments& InArgs, UDMMaterialStage* InStage);
+
+	FORCEINLINE UDMMaterialStage* GetStage() const { return StageWeak.Get(); }
+
+	FORCEINLINE bool IsStageSelected() const { return bStageSelected.Get(); }
+	FORCEINLINE bool IsStageEnabled() const { return bStageEnabled.Get(); }
+	FORCEINLINE bool IsStageMask() const { return bIsMask; }
+
+	void SetStageSelected(const bool bInStageSelected) { bStageSelected = bInStageSelected; }
+
+	FVector2D GetPreviewSize() const { return DesiredSize.Get(); }
+
+	//~ Begin SWidget
+	virtual FCursorReply OnCursorQuery(const FGeometry& InMyGeometry, const FPointerEvent& InCursorEvent) const override;
+	virtual FReply OnMouseMove(const FGeometry& InMyGeometry, const FPointerEvent& InMouseEvent) override;
+	virtual FReply OnMouseButtonUp(const FGeometry& InMyGeometry, const FPointerEvent& InMouseEvent) override;
+	virtual FReply OnMouseButtonDown(const FGeometry& MyGeometry, const FPointerEvent& MouseEvent) override;
+	virtual FReply OnDragDetected(const FGeometry& InMyGeometry, const FPointerEvent& InMouseEvent) override;
+	virtual void OnDragEnter(FGeometry const& InMyGeometry, FDragDropEvent const& InDragDropEvent) override;
+	virtual void OnDragLeave(FDragDropEvent const& InDragDropEvent) override;
+	virtual FReply OnDragOver(const FGeometry& InMyGeometry, const FDragDropEvent& InDragDropEvent) override;
+	virtual FReply OnDrop(const FGeometry& InMyGeometry, const FDragDropEvent& InDragDropEvent) override;
+	//~ End SWidget
+
+protected:
+	TWeakObjectPtr<UDMMaterialStage> StageWeak;
+
+	TAttribute<bool> bStageEnabled;
+	TAttribute<bool> bStageSelected;
+	bool bIsMask;
+	TAttribute<bool> bInteractable;
+	TAttribute<bool> bShowTextOverlays;
+	TAttribute<FVector2D> DesiredSize;
+	float MinUniformSize;
+	float MaxUniformSize;
+	FDMOnStageUniformSizeChanged OnUniformSizeChanged;
+	FDMOnStageClick OnClick;
+
+	FDMOnStageCanAcceptDrop OnCanAcceptDropEvent;
+	FDMOnStageAcceptDrop OnAcceptDropEvent;
+	FDMOnStagePaintDropIndicator OnPaintDropIndicatorEvent;
+
+	FVector2D CachedCursorLocation;
+	float CachedPreviewSize;
+
+	TSharedRef<SWidget> CreateTextBlockBackground();
+
+	EVisibility GetBorderVisibility() const;
+	const FSlateBrush* GetBorderBrush() const;
+	const FSlateBrush* GetOverlayBorderBrush() const;
+
+	FText GetStageName() const;
+
+	void HandleStageDragDropOperation(FDMStageDragDropOperation& StageDragDropOperation);
+	void HandleAssetDragDropOperation(FAssetDragDropOp& AssetDragDropOperation);
+
+	bool IsToolTipInteractable() const;
+
+	FOptionalSize GetStagePreviewDesiredWidth() const;
+	FOptionalSize GetStagePreviewDesiredHeight() const;
+
+	FSlateColor GetBorderBrushBackgroundColor() const;
+	EVisibility GetOverlayBorderBrushVisibility() const;
+
+	EVisibility GetTextBlockBackgroundVisibility() const;
+};
diff --git a/Engine/Plugins/Experimental/DynamicMaterial/Source/DynamicMaterialEditor/Private/Slate/SDMStatusBar.cpp b/Engine/Plugins/Experimental/DynamicMaterial/Source/DynamicMaterialEditor/Private/Slate/SDMStatusBar.cpp
new file mode 100644
--- /dev/null
+++ b/Engine/Plugins/Experimental/DynamicMaterial/Source/DynamicMaterialEditor/Private/Slate/SDMStatusBar.cpp
@@ -0,0 +1,148 @@
+// Copyright Epic Games, Inc. All Rights Reserved.
+
+#include "Slate/SDMStatusBar.h"
+#include "Components/DMMaterialSlot.h"
+#include "DetailLayoutBuilder.h"
+#include "DynamicMaterialEditorStyle.h"
+#include "Model/DynamicMaterialModel.h"
+#include "Slate/SDMEditor.h"
+#include "Slate/SDMSlot.h"
+#include "SlateOptMacros.h"
+#include "Widgets/SBoxPanel.h"
+#include "Widgets/Text/STextBlock.h"
+
+#define LOCTEXT_NAMESPACE "SDMStatusBar"
+
+BEGIN_SLATE_FUNCTION_BUILD_OPTIMIZATION
+
+void SDMStatusBar::Construct(const FArguments& InArgs, const TWeakObjectPtr<UDynamicMaterialModel>& InEditorModel, const TSharedRef<SDMEditor>& InEditorWidget)
+{
+	SetCanTick(true);
+
+	ensure(InEditorModel.IsValid());
+	EditorModelWeak = InEditorModel;
+	EditorWidgetWeak = InEditorWidget;
+
+	ChildSlot
+	.HAlign(HAlign_Fill)
+	.VAlign(VAlign_Center)
+	[
+		SNew(SBorder)
+		.HAlign(HAlign_Fill)
+		.VAlign(VAlign_Center)
+		.BorderImage(FDynamicMaterialEditorStyle::GetBrush("Border.Top"))
+		.BorderBackgroundColor(FLinearColor(1, 1, 1, 0.05f))
+		[
+			SNew(SWrapBox)
+			.HAlign(HAlign_Left)
+			.InnerSlotPadding(FVector2D(5.0f))
+			.UseAllottedSize(true)
+			+ CreateStatsWrapBoxEntry(
+				TAttribute<FText>::CreateSP(this, &SDMStatusBar::GetNumPixelShaderInstructionsText), 
+				LOCTEXT("NumPixelShaderInstructions_ToolTip", "Pixel shader instruction count")
+			)
+			+ CreateStatsWrapBoxEntry(
+				TAttribute<FText>::CreateSP(this, &SDMStatusBar::GetNumVertexShaderInstructionsText), 
+				LOCTEXT("NumVertexShaderInstructions_ToolTip", "Vertex shader instruction count")
+			)
+			+ CreateStatsWrapBoxEntry(
+				TAttribute<FText>::CreateSP(this, &SDMStatusBar::GetNumSamplersText), 
+				LOCTEXT("NumSamplers_ToolTip", "Sampler count")
+			)
+		]
+	];
+}
+
+SWrapBox::FSlot::FSlotArguments SDMStatusBar::CreateStatsWrapBoxEntry(TAttribute<FText> InText, const FText& InTooltipText)
+{
+	const FLinearColor SeparatorColor = FLinearColor(1, 1, 1, 0.1f);
+
+	return 
+		MoveTemp(
+			SWrapBox::Slot()
+			.Padding(0.0f, 2.0f, 0.0f, 2.0f)
+			[
+				SNew(SHorizontalBox)
+				.ToolTipText(InTooltipText)
+				+ SHorizontalBox::Slot()
+				.AutoWidth()
+				.HAlign(HAlign_Center)
+				.VAlign(VAlign_Fill)
+				[
+					SNew(SBorder)
+					.BorderBackgroundColor(SeparatorColor)
+					.BorderImage(FDynamicMaterialEditorStyle::GetBrush("Border.Left"))
+					
+				]
+				+ SHorizontalBox::Slot()
+				.AutoWidth()
+				.HAlign(HAlign_Center)
+				.VAlign(VAlign_Center)
+				.Padding(5.0f, 0.0f, 5.0f, 0.0f)
+				[
+					SNew(STextBlock)
+					.TextStyle(FAppStyle::Get(), "SmallText")
+					.Text(InText)
+				]
+			]
+		);
+}
+
+END_SLATE_FUNCTION_BUILD_OPTIMIZATION
+
+FText SDMStatusBar::GetNumMaterialSlotsText() const
+{
+	return FText::Format(LOCTEXT("NumMaterialSlots_Text", "Material Slots {0}"), CachedSlotCount);
+}
+
+FText SDMStatusBar::GetNumTotalLayersText() const
+{
+	return FText::Format(LOCTEXT("NumTotalLayers_Text", "Total Layers {0}"), CachedCurrentLayerCount);
+}
+
+FText SDMStatusBar::GetNumCurrentSlotLayersText() const
+{
+	return FText::Format(LOCTEXT("NumCurrentSlotLayers_Text", "Layers {0}"), CachedTotalLayerCount);
+}
+
+FText SDMStatusBar::GetNumPixelShaderInstructionsText() const
+{
+	return FText::Format(LOCTEXT("NumPixelShaderInstructions_Text", "PS Instructions {0}"), CachedMaterialStats.NumPixelShaderInstructions);
+}
+
+FText SDMStatusBar::GetNumVertexShaderInstructionsText() const
+{
+	return FText::Format(LOCTEXT("NumVertexShaderInstructions_Text", "VS Instructions {0}"), CachedMaterialStats.NumVertexShaderInstructions);
+}
+
+FText SDMStatusBar::GetNumSamplersText() const
+{
+	return FText::Format(LOCTEXT("NumSamplers_Text", "Samplers {0}"), CachedMaterialStats.NumSamplers);
+}
+
+FText SDMStatusBar::GetNumPixelTextureSamplesText() const
+{
+	return FText::Format(LOCTEXT("NumPixelTextureSamples_Text", "PT Samples {0}"), CachedMaterialStats.NumPixelTextureSamples);
+}
+
+FText SDMStatusBar::GetNumVertexTextureSamplesText() const
+{
+	return FText::Format(LOCTEXT("NumVertexTextureSamples_Text", "VT Samples {0}"), CachedMaterialStats.NumVertexTextureSamples);
+}
+
+FText SDMStatusBar::GetNumVirtualTextureSamplesText() const
+{
+	return FText::Format(LOCTEXT("NumVirtualTextureSamples_Text", "Virtual Texture Samples {0}"), CachedMaterialStats.NumVirtualTextureSamples);
+}
+
+FText SDMStatusBar::GetNumUVScalarsText() const
+{
+	return FText::Format(LOCTEXT("NumUVScalars_Text", "UV Scalars {0}"), CachedMaterialStats.NumUVScalars);
+}
+
+FText SDMStatusBar::GetNumInterpolatorScalarsText() const
+{
+	return FText::Format(LOCTEXT("NumInterpolatorScalars_Text", "Interpolator Scalars {0}"), CachedMaterialStats.NumInterpolatorScalars);
+}
+
+#undef LOCTEXT_NAMESPACE
diff --git a/Engine/Plugins/Experimental/DynamicMaterial/Source/DynamicMaterialEditor/Private/Slate/SDMStatusBar.h b/Engine/Plugins/Experimental/DynamicMaterial/Source/DynamicMaterialEditor/Private/Slate/SDMStatusBar.h
new file mode 100644
--- /dev/null
+++ b/Engine/Plugins/Experimental/DynamicMaterial/Source/DynamicMaterialEditor/Private/Slate/SDMStatusBar.h
@@ -0,0 +1,46 @@
+// Copyright Epic Games, Inc. All Rights Reserved.
+
+#pragma once
+
+#include "MaterialEditingLibrary.h"
+#include "Widgets/Layout/SWrapBox.h"
+#include "Widgets/SCompoundWidget.h"
+
+class SDMEditor;
+class UDynamicMaterialModel;
+
+class SDMStatusBar : public SCompoundWidget
+{
+	SLATE_BEGIN_ARGS(SDMStatusBar) {}
+	SLATE_END_ARGS()
+
+public:
+	virtual ~SDMStatusBar() {}
+
+	void Construct(const FArguments& InArgs, const TWeakObjectPtr<UDynamicMaterialModel>& InEditorModel, const TSharedRef<SDMEditor>& InEditorWidget);
+
+protected:
+	TWeakObjectPtr<UDynamicMaterialModel> EditorModelWeak;
+	TWeakPtr<SDMEditor> EditorWidgetWeak;
+
+	int32 CachedSlotCount = 0;
+	int32 CachedCurrentLayerCount = 0;
+	int32 CachedTotalLayerCount = 0;
+
+	FMaterialStatistics CachedMaterialStats;
+
+	SWrapBox::FSlot::FSlotArguments CreateStatsWrapBoxEntry(TAttribute<FText> InText, const FText& InTooltipText);
+
+	FText GetNumMaterialSlotsText() const;
+	FText GetNumTotalLayersText() const;
+	FText GetNumCurrentSlotLayersText() const;
+
+	FText GetNumPixelShaderInstructionsText() const;
+	FText GetNumVertexShaderInstructionsText() const;
+	FText GetNumSamplersText() const;
+	FText GetNumPixelTextureSamplesText() const;
+	FText GetNumVertexTextureSamplesText() const;
+	FText GetNumVirtualTextureSamplesText() const;
+	FText GetNumUVScalarsText() const;
+	FText GetNumInterpolatorScalarsText() const;
+};
diff --git a/Engine/Plugins/Experimental/DynamicMaterial/Source/DynamicMaterialEditor/Private/Slate/SDMToolBar.cpp b/Engine/Plugins/Experimental/DynamicMaterial/Source/DynamicMaterialEditor/Private/Slate/SDMToolBar.cpp
new file mode 100644
--- /dev/null
+++ b/Engine/Plugins/Experimental/DynamicMaterial/Source/DynamicMaterialEditor/Private/Slate/SDMToolBar.cpp
@@ -0,0 +1,382 @@
+// Copyright Epic Games, Inc. All Rights Reserved.
+
+#include "Slate/SDMToolBar.h"
+#include "AssetToolsModule.h"
+#include "DMBlueprintFunctionLibrary.h"
+#include "DMEDefs.h"
+#include "DynamicMaterialEditorModule.h"
+#include "DynamicMaterialEditorSettings.h"
+#include "DynamicMaterialEditorStyle.h"
+#include "DynamicMaterialModule.h"
+#include "GameFramework/Actor.h"
+#include "IAssetTools.h"
+#include "Model/DynamicMaterialModel.h"
+#include "SlateOptMacros.h"
+#include "Styling/AppStyle.h"
+#include "Styling/StyleColors.h"
+#include "Widgets/Colors/SColorBlock.h"
+#include "Widgets/Images/SImage.h"
+#include "Widgets/Input/SButton.h"
+#include "Widgets/Input/SComboBox.h"
+#include "Widgets/Input/SComboButton.h"
+#include "Widgets/Layout/SBorder.h"
+#include "Widgets/Layout/SWrapBox.h"
+#include "Widgets/SBoxPanel.h"
+#include "Widgets/SNullWidget.h"
+#include "Widgets/Text/STextBlock.h"
+
+#define LOCTEXT_NAMESPACE "SDMToolBar"
+
+BEGIN_SLATE_FUNCTION_BUILD_OPTIMIZATION
+
+void SDMToolBar::Construct(const FArguments& InArgs)
+{
+	SetCanTick(true);
+
+	MaterialActorWeak = InArgs._MaterialActor;
+	MaterialModelWeak = InArgs._MaterialModel;
+	OnSlotChanged = InArgs._OnSlotChanged;
+	OnGetSettingsMenu = InArgs._OnGetSettingsMenu;
+	
+	if (MaterialActorWeak.IsValid())
+	{
+		SetMaterialActor(MaterialActorWeak.Get());
+	}
+	else
+	{
+		SetMaterialModel(MaterialModelWeak.Get());
+	}
+	
+	ChildSlot
+	.HAlign(HAlign_Fill)
+	.VAlign(VAlign_Center)
+	[
+		SNew(SBorder)
+		.HAlign(HAlign_Fill)
+		.VAlign(VAlign_Center)
+		.BorderImage(FDynamicMaterialEditorStyle::GetBrush("Border.Bottom"))
+		.BorderBackgroundColor(FLinearColor(1, 1, 1, 0.05f))
+		[
+			CreateToolBarEntries()
+		]
+	];
+}
+
+TSharedRef<SWidget> SDMToolBar::CreateToolBarEntries()
+{
+	return 
+		SNew(SHorizontalBox)
+		
+		+ SHorizontalBox::Slot()
+		.FillWidth(1.0f)
+		.HAlign(HAlign_Fill)
+		.VAlign(VAlign_Center)
+		[
+			SNew(SWrapBox)
+			.Orientation(Orient_Horizontal)
+			.UseAllottedSize(true)
+			.HAlign(HAlign_Left)
+			.InnerSlotPadding(FVector2D(20.0f))
+			+ SWrapBox::Slot()
+			.HAlign(HAlign_Left)
+			.VAlign(VAlign_Center)
+			.Padding(5.0f, 0.0f, 0.0f, 0.0f)
+			[
+				SNew(SHorizontalBox)
+				+ SHorizontalBox::Slot()
+				.AutoWidth()
+				.HAlign(HAlign_Left)
+				.VAlign(VAlign_Center)
+				.Padding(0.0f, 0.0f, 10.0f, 0.0f)
+				[
+					SNew(STextBlock)
+					.Text(LOCTEXT("MaterialDesignerInstanceActorLabel", "Actor"))
+				]
+				+ SHorizontalBox::Slot()
+				.AutoWidth()
+				.HAlign(HAlign_Left)
+				.VAlign(VAlign_Center)
+				[
+					SNew(STextBlock)
+					.TextStyle(FDynamicMaterialEditorStyle::Get(), "ActorName")
+					.Text(this, &SDMToolBar::GetSlotActorDisplayName)
+				]
+			]
+			+ SWrapBox::Slot()
+			.FillEmptySpace(true)
+			.HAlign(HAlign_Fill)
+			.VAlign(VAlign_Center)
+			[
+				SNew(SHorizontalBox)
+				.Visibility(this, &SDMToolBar::GetSlotsComboBoxWidgetVisibiltiy)
+				+ SHorizontalBox::Slot()
+				.AutoWidth()
+				.HAlign(HAlign_Left)
+				.VAlign(VAlign_Center)
+				.Padding(0.0f, 0.0f, 10.0f, 0.0f)
+				[
+					SNew(STextBlock)
+					.Text(LOCTEXT("MaterialDesignerInstanceActorSlotLabel", "Property"))
+				]
+				+ SHorizontalBox::Slot()
+				.AutoWidth()
+				.HAlign(HAlign_Left)
+				.VAlign(VAlign_Center)
+				[
+					CreateSlotsComboBoxWidget()
+				]
+			]
+		]		
+		
+		+ SHorizontalBox::Slot()
+		.AutoWidth()
+		.HAlign(HAlign_Right)
+		.VAlign(VAlign_Top)
+		.Padding(5.0f, 0.0f, 0.0f, 0.0f)
+		[
+			SNew(SButton)
+			.ContentPadding(GetDefaultToolBarButtonContentPadding())
+			.ButtonStyle(FDynamicMaterialEditorStyle::Get(), "HoverHintOnly")
+			.ToolTipText(LOCTEXT("MaterialDesignerFollowSelectionTooltip", "Toggles whether the Material Designer display will change selecting new objects and actors."))
+			.OnClicked(this, &SDMToolBar::OnFollowSelectionButtonClicked)
+			[
+				SNew(SImage)
+				.Image(this, &SDMToolBar::GetFollowSelectionBrush)
+				.DesiredSizeOverride(GetDefaultToolBarButtonSize())
+				.ColorAndOpacity(this, &SDMToolBar::GetFollowSelectionColor)
+			]
+		]
+		
+		+ SHorizontalBox::Slot()
+		.AutoWidth()
+		.HAlign(HAlign_Right)
+		.VAlign(VAlign_Top)
+		.Padding(5.0f, 0.0f, 0.0f, 0.0f)
+		[
+			SNew(SComboButton)
+			.HasDownArrow(false)
+			.IsFocusable(true)
+			.ContentPadding(GetDefaultToolBarButtonContentPadding())
+			.ButtonStyle(FDynamicMaterialEditorStyle::Get(), "HoverHintOnly")
+			.ToolTipText(LOCTEXT("MaterialDesignerSettingsTooltip", "Material Designer Settings"))
+			.OnGetMenuContent(OnGetSettingsMenu)
+			.ButtonContent()
+			[
+				SNew(SImage)
+				.Image(FDynamicMaterialEditorStyle::GetBrush("Icons.Menu.Dropdown"))
+				.DesiredSizeOverride(GetDefaultToolBarButtonSize())
+			]
+		];
+}
+
+TSharedRef<SWidget> SDMToolBar::CreateToolBarButton(TAttribute<const FSlateBrush*> InImageBrush, const TAttribute<FText>& InTooltipText, FOnClicked InOnClicked)
+{
+	return 
+		SNew(SButton)
+		.ContentPadding(GetDefaultToolBarButtonContentPadding())
+		.ButtonStyle(FDynamicMaterialEditorStyle::Get(), "HoverHintOnly")
+		.ToolTipText(InTooltipText)
+		.OnClicked(InOnClicked)
+		[
+			SNew(SImage)
+			.Image(InImageBrush)
+			.DesiredSizeOverride(GetDefaultToolBarButtonSize())
+		];
+}
+
+TSharedRef<SWidget> SDMToolBar::CreateSlotsComboBoxWidget()
+{
+	if (!MaterialActorWeak.IsValid() || !MaterialModelWeak.IsValid())
+	{
+		return SNullWidget::NullWidget;
+	}
+
+	const TSharedPtr<FDMObjectMaterialProperty> InitiallySelectedItem = 
+		ActorMaterialProperties.IsValidIndex(SelectedMaterialSlotIndex) ? ActorMaterialProperties[SelectedMaterialSlotIndex] : nullptr;
+
+	return 
+		SNew(SComboBox<TSharedPtr<FDMObjectMaterialProperty>>)
+		.InitiallySelectedItem(InitiallySelectedItem)
+		.OptionsSource(&ActorMaterialProperties)
+		.OnGenerateWidget(this, &SDMToolBar::GenerateSelectedMaterialSlotRow)
+		.OnSelectionChanged(this, &SDMToolBar::OnMaterialSlotChanged)
+		[
+			SNew(STextBlock)
+			.MinDesiredWidth(100.0f)
+			.Text(this, &SDMToolBar::GetSelectedMaterialSlotName)
+		];
+}
+
+TSharedRef<SWidget> SDMToolBar::GenerateSelectedMaterialSlotRow(TSharedPtr<FDMObjectMaterialProperty> InSelectedSlot) const
+{
+	if (InSelectedSlot.IsValid())
+	{
+		return SNew(STextBlock)
+			.MinDesiredWidth(100.f)
+			.Text(this, &SDMToolBar::GetSlotDisplayName, InSelectedSlot);
+	}
+	return SNullWidget::NullWidget;
+}
+
+END_SLATE_FUNCTION_BUILD_OPTIMIZATION
+
+FText SDMToolBar::GetSlotDisplayName(TSharedPtr<FDMObjectMaterialProperty> InSlot) const
+{
+	return InSlot->GetPropertyName(false);
+}
+
+FText SDMToolBar::GetSelectedMaterialSlotName() const
+{
+	if (ActorMaterialProperties.IsValidIndex(SelectedMaterialSlotIndex) && ActorMaterialProperties[SelectedMaterialSlotIndex].IsValid())
+	{
+		return GetSlotDisplayName(ActorMaterialProperties[SelectedMaterialSlotIndex]);
+	}
+	return FText::GetEmpty();
+}
+
+void SDMToolBar::OnMaterialSlotChanged(TSharedPtr<FDMObjectMaterialProperty> InSelectedSlot, ESelectInfo::Type InSelectInfoType)
+{
+	if (!InSelectedSlot.IsValid())
+	{
+		return;
+	}
+
+	UDynamicMaterialModel* SelectedMaterialModel = InSelectedSlot->GetMaterialModel();
+	if (IsValid(SelectedMaterialModel))
+	{
+		SetMaterialModel(SelectedMaterialModel);
+	}
+	else if (InSelectedSlot->OuterWeak.IsValid())
+	{
+		UDMBlueprintFunctionLibrary::CreateDynamicMaterialInObject(*InSelectedSlot.Get());
+	}
+
+	OnSlotChanged.ExecuteIfBound(InSelectedSlot);
+}
+
+EVisibility SDMToolBar::GetSlotsComboBoxWidgetVisibiltiy() const
+{
+	return ActorMaterialProperties.Num() > 1 ? EVisibility::Visible : EVisibility::Collapsed;
+}
+
+FText SDMToolBar::GetSlotActorDisplayName() const
+{
+	const AActor* const SlotActor = GetMaterialActor();
+	return IsValid(SlotActor) ? FText::FromString(SlotActor->GetActorLabel()) : FText();
+}
+
+void SDMToolBar::SetMaterialProperties(const TArray<TSharedPtr<FDMObjectMaterialProperty>>& InActorMaterialProperties)
+{
+	ActorMaterialProperties = InActorMaterialProperties;
+	SelectedMaterialSlotIndex = 0;
+}
+
+void SDMToolBar::SetMaterialModel(UDynamicMaterialModel* InModel)
+{
+	if (!IsValid(InModel))
+	{
+		return;
+	}
+
+	MaterialModelWeak = InModel;
+
+	if (!MaterialActorWeak.IsValid())
+	{
+		MaterialActorWeak = MaterialModelWeak->GetTypedOuter<AActor>();
+	}
+
+	ActorMaterialProperties.Empty(0);
+	SelectedMaterialSlotIndex = INDEX_NONE;
+
+	if (!MaterialActorWeak.IsValid())
+	{
+		return;
+	}
+
+	const TArray<FDMObjectMaterialProperty> ActorDynamicMaterialSlots = UDMBlueprintFunctionLibrary::GetActorMaterialProperties(MaterialActorWeak.Get());
+	ActorMaterialProperties.Reserve(ActorDynamicMaterialSlots.Num());
+
+	for (const FDMObjectMaterialProperty& ActorDynamicMaterialSlot : ActorDynamicMaterialSlots)
+	{
+		ActorMaterialProperties.Add(MakeShared<FDMObjectMaterialProperty>(ActorDynamicMaterialSlot));
+	}
+
+	for (int32 SlotIndex = 0; SlotIndex < ActorMaterialProperties.Num(); ++SlotIndex)
+	{
+		FDMObjectMaterialProperty* ActorMaterialSlot = ActorMaterialProperties[SlotIndex].Get();
+		if (ActorMaterialSlot && ActorMaterialSlot->GetMaterialModel() == MaterialModelWeak)
+		{
+			SelectedMaterialSlotIndex = SlotIndex;
+			break;
+		}
+	}
+}
+
+void SDMToolBar::SetMaterialActor(AActor* InActor, const int32 InActiveSlotIndex)
+{
+	MaterialActorWeak = InActor;
+	ActorMaterialProperties.Empty();
+	SelectedMaterialSlotIndex = 0;
+
+	TArray<FDMObjectMaterialProperty> ActorProperties = UDMBlueprintFunctionLibrary::GetActorMaterialProperties(InActor);
+	UDynamicMaterialModel* MaterialModel = MaterialModelWeak.Get();
+
+	for (int32 MaterialPropertyIdx = 0; MaterialPropertyIdx < ActorProperties.Num(); ++MaterialPropertyIdx)
+	{
+		const FDMObjectMaterialProperty& MaterialProperty = ActorProperties[MaterialPropertyIdx];
+
+		ActorMaterialProperties.Add(MakeShared<FDMObjectMaterialProperty>(MaterialProperty));
+
+		if (MaterialProperty.GetMaterialModel() == MaterialModel)
+		{
+			SelectedMaterialSlotIndex = MaterialPropertyIdx;
+		}
+	}
+}
+
+const FSlateBrush* SDMToolBar::GetFollowSelectionBrush() const
+{
+	static const FSlateBrush* Unlocked = FAppStyle::Get().GetBrush("Icons.Unlock");
+	static const FSlateBrush* Locked = FAppStyle::Get().GetBrush("Icons.Lock");
+
+	if (UDynamicMaterialEditorSettings* Settings = UDynamicMaterialEditorSettings::Get())
+	{
+		if (!Settings->bFollowSelection)
+		{
+			return Locked;
+		}
+	}
+
+	return Unlocked;
+}
+
+FSlateColor SDMToolBar::GetFollowSelectionColor() const
+{
+	// We want the icon to stand out when it's locked.
+	static FSlateColor EnabledColor = FSlateColor(EStyleColor::AccentGray);
+	static FSlateColor DisabledColor = FSlateColor(EStyleColor::Primary);
+
+	bool bFollowingSelection = true;
+
+	if (UDynamicMaterialEditorSettings* Settings = UDynamicMaterialEditorSettings::Get())
+	{
+		bFollowingSelection = Settings->bFollowSelection;
+	}
+
+	return bFollowingSelection
+		? EnabledColor
+		: DisabledColor;
+}
+
+FReply SDMToolBar::OnFollowSelectionButtonClicked()
+{
+	if (UDynamicMaterialEditorSettings* Settings = UDynamicMaterialEditorSettings::Get())
+	{
+		Settings->bFollowSelection = !Settings->bFollowSelection;
+		Settings->SaveConfig();
+	}
+
+	return FReply::Handled();
+}
+
+#undef LOCTEXT_NAMESPACE
diff --git a/Engine/Plugins/Experimental/DynamicMaterial/Source/DynamicMaterialEditor/Private/Slate/SDMToolBar.h b/Engine/Plugins/Experimental/DynamicMaterial/Source/DynamicMaterialEditor/Private/Slate/SDMToolBar.h
new file mode 100644
--- /dev/null
+++ b/Engine/Plugins/Experimental/DynamicMaterial/Source/DynamicMaterialEditor/Private/Slate/SDMToolBar.h
@@ -0,0 +1,74 @@
+// Copyright Epic Games, Inc. All Rights Reserved.
+
+#pragma once
+
+#include "Framework/SlateDelegates.h"
+#include "Widgets/SCompoundWidget.h"
+
+class AActor;
+class SWidget;
+class UDMMaterialStageExpression;
+class UDynamicMaterialModel;
+enum class EDMExpressionMenu : uint8;
+struct FDMObjectMaterialProperty;
+
+DECLARE_DELEGATE_OneParam(FDMOnActorMaterailSlotChanged, TSharedPtr<FDMObjectMaterialProperty> /** NewSelectedSlot */)
+
+/**
+ * Material Designer ToolBar
+ * 
+ * Displays the selected actor that the Material Designer is editing and allows for switching between slots for that actor.
+ */
+class SDMToolBar : public SCompoundWidget
+{
+public:
+	SLATE_BEGIN_ARGS(SDMToolBar) {}
+		SLATE_ARGUMENT(TWeakObjectPtr<AActor>, MaterialActor)
+		SLATE_ARGUMENT(TWeakObjectPtr<UDynamicMaterialModel>, MaterialModel)
+		SLATE_EVENT(FDMOnActorMaterailSlotChanged, OnSlotChanged)
+		SLATE_EVENT(FOnGetContent, OnGetSettingsMenu)
+	SLATE_END_ARGS()
+
+	virtual ~SDMToolBar() {}
+
+	void Construct(const FArguments& InArgs);
+
+	const TArray<TSharedPtr<FDMObjectMaterialProperty>>& GetMaterialProperties() const { return ActorMaterialProperties; }
+	void SetMaterialProperties(const TArray<TSharedPtr<FDMObjectMaterialProperty>>& InActorMaterialProperties);
+
+	const AActor* GetMaterialActor() const { return MaterialActorWeak.Get(); }
+	void SetMaterialActor(AActor* InActor, const int32 InActiveSlotIndex = 0);
+
+	UDynamicMaterialModel* GetMaterialModel() const { return MaterialModelWeak.Get(); }
+	void SetMaterialModel(UDynamicMaterialModel* InModel);
+
+	FText GetSlotActorDisplayName() const;
+
+protected:
+	TWeakObjectPtr<AActor> MaterialActorWeak;
+	TWeakObjectPtr<UDynamicMaterialModel> MaterialModelWeak;
+	FDMOnActorMaterailSlotChanged OnSlotChanged;
+	FOnGetContent OnGetSettingsMenu;
+
+	TArray<TSharedPtr<FDMObjectMaterialProperty>> ActorMaterialProperties;
+	int32 SelectedMaterialSlotIndex;
+
+	TSharedRef<SWidget> CreateToolBarEntries();
+
+	TSharedRef<SWidget> CreateToolBarButton(TAttribute<const FSlateBrush*> InImageBrush, const TAttribute<FText>& InTooltipText, FOnClicked InOnClicked);
+	TSharedRef<SWidget> CreateSlotsComboBoxWidget();
+
+	TSharedRef<SWidget> GenerateSelectedMaterialSlotRow(TSharedPtr<FDMObjectMaterialProperty> InSelectedSlot) const;
+	FText GetSlotDisplayName(TSharedPtr<FDMObjectMaterialProperty> InSlot) const;
+	FText GetSelectedMaterialSlotName() const;
+	void OnMaterialSlotChanged(TSharedPtr<FDMObjectMaterialProperty> InSelectedSlot, ESelectInfo::Type InSelectInfoType);
+
+	const FMargin GetDefaultToolBarButtonContentPadding() const { return FMargin(2.0f); }
+	const FVector2D GetDefaultToolBarButtonSize() const { return FVector2D(20.0f); }
+
+	EVisibility GetSlotsComboBoxWidgetVisibiltiy() const;
+
+	const FSlateBrush* GetFollowSelectionBrush() const;
+	FSlateColor GetFollowSelectionColor() const;
+	FReply OnFollowSelectionButtonClicked();
+};
diff --git a/Engine/Plugins/Experimental/DynamicMaterial/Source/DynamicMaterialEditor/Private/Slate/SMaterialToolTip.cpp b/Engine/Plugins/Experimental/DynamicMaterial/Source/DynamicMaterialEditor/Private/Slate/SMaterialToolTip.cpp
new file mode 100644
--- /dev/null
+++ b/Engine/Plugins/Experimental/DynamicMaterial/Source/DynamicMaterialEditor/Private/Slate/SMaterialToolTip.cpp
@@ -0,0 +1,111 @@
+// Copyright Epic Games, Inc. All Rights Reserved.
+
+#include "SMaterialToolTip.h"
+#include "SlateOptMacros.h"
+#include "Styling/SlateBrush.h"
+#include "Widgets/Images/SImage.h"
+#include "Widgets/Layout/SBorder.h"
+#include "Widgets/Layout/SBox.h"
+#include "Widgets/Layout/SScaleBox.h"
+#include "Widgets/SToolTip.h"
+#include "Widgets/Text/STextBlock.h"
+
+#define LOCTEXT_NAMESPACE "SMaterialToolTip"
+
+BEGIN_SLATE_FUNCTION_BUILD_OPTIMIZATION
+
+SMaterialToolTip::SMaterialToolTip()
+	: MaterialBrush(FSlateMaterialBrush(FVector2D(1.0f, 1.0f)))
+{
+}
+
+void SMaterialToolTip::Construct(const FArguments& InArgs)
+{
+	Material = InArgs._Material;
+	Text = InArgs._Text;
+	MaterialSize = InArgs._MaterialSize;
+	ShowMaterial = InArgs._ShowMaterial;
+	ShowDefault = InArgs._ShowDefault;
+
+	ChildSlot
+	[
+		CreateToolTipWidget()
+	];
+}
+
+TSharedRef<SWidget> SMaterialToolTip::CreateToolTipWidget()
+{
+	if (ShowDefault)
+	{
+		return SNullWidget::NullWidget;
+	}
+
+	MaterialBrush.SetImageSize(MaterialSize.Get());
+	MaterialBrush.SetMaterial(Material.Get());
+
+	return 
+		SNew(SVerticalBox)
+		+ SVerticalBox::Slot()
+		.AutoHeight()
+		.HAlign(HAlign_Left)
+		.VAlign(VAlign_Top)
+		[
+			SNew(STextBlock)
+			.Text(Text)
+			.Font(FCoreStyle::Get().GetFontStyle("ToolTip.Font"))
+			.ColorAndOpacity(FSlateColor::UseForeground())
+			.WrapTextAt_Static(&SToolTip::GetToolTipWrapWidth)
+		]
+		+ SVerticalBox::Slot()
+		.AutoHeight()
+		.HAlign(HAlign_Center)
+		.VAlign(VAlign_Top)
+		[
+			SNew(SBox)
+			.Padding(0.0f, 10.0f, 4.0f, 4.0f)
+			.Visibility(this, &SMaterialToolTip::GetShowMaterialVisibility)
+			[
+				SNew(SBorder)
+				.Padding(FMargin(0.0f, 0.0f, 5.0f, 5.0f))
+				.BorderImage(FAppStyle::Get().GetBrush("ContentBrowser.AssetTileItem.DropShadow"))
+				[
+					SNew(SBox)
+					.WidthOverride(this, &SMaterialToolTip::GetMaterialSizeX)
+					.HeightOverride(this, &SMaterialToolTip::GetMaterialSizeY)
+					[
+						SNew(SScaleBox)
+						.Stretch(EStretch::ScaleToFill)
+						.StretchDirection(EStretchDirection::Both)
+						[
+							SNew(SBorder)
+							.Padding(0.0f)
+							.BorderImage(FAppStyle::Get().GetBrush("ContentBrowser.AssetTileItem.ThumbnailAreaBackground"))
+							[
+								SNew(SImage)
+								.Image(&MaterialBrush)
+							]
+						]
+					]
+				]
+			]
+		];
+}
+
+EVisibility SMaterialToolTip::GetShowMaterialVisibility() const
+{
+	return ShowMaterial.Get() ? EVisibility::SelfHitTestInvisible : EVisibility::Collapsed;
+}
+
+FOptionalSize SMaterialToolTip::GetMaterialSizeX() const
+{
+	return MaterialSize.Get().X;
+}
+
+FOptionalSize SMaterialToolTip::GetMaterialSizeY() const
+{
+	return MaterialSize.Get().Y;
+}
+
+END_SLATE_FUNCTION_BUILD_OPTIMIZATION
+
+#undef LOCTEXT_NAMESPACE
diff --git a/Engine/Plugins/Experimental/DynamicMaterial/Source/DynamicMaterialEditor/Private/Slate/SMaterialToolTip.h b/Engine/Plugins/Experimental/DynamicMaterial/Source/DynamicMaterialEditor/Private/Slate/SMaterialToolTip.h
new file mode 100644
--- /dev/null
+++ b/Engine/Plugins/Experimental/DynamicMaterial/Source/DynamicMaterialEditor/Private/Slate/SMaterialToolTip.h
@@ -0,0 +1,50 @@
+// Copyright Epic Games, Inc. All Rights Reserved.
+
+#pragma once
+
+#include "SlateMaterialBrush.h"
+#include "Widgets/SCompoundWidget.h"
+
+class UMaterialInterface;
+
+class SMaterialToolTip : public SCompoundWidget
+{
+	SLATE_BEGIN_ARGS(SMaterialToolTip)
+		: _Text(FText())
+		, _MaterialSize(FVector2D(512.0f))
+		, _ShowMaterial(true)
+		, _ShowDefault(false)
+	{}
+		/** The material to display in the tooltip. */
+		SLATE_ARGUMENT(UMaterialInterface*, Material)
+		/** The tooltip text to display. */
+		SLATE_ATTRIBUTE(FText, Text)
+		/** The size of the space for the texture. This is not the texture size itself. */
+		SLATE_ATTRIBUTE(FVector2D, MaterialSize)
+		/** If true, shows the material in the tooltip along with the tooltip text. If false, shows only the tooltip text. */
+		SLATE_ATTRIBUTE(bool, ShowMaterial)
+		/** If true, shows the default text widget instead of the customized texture tooltip. */
+		SLATE_ARGUMENT(bool, ShowDefault)
+	SLATE_END_ARGS()
+
+public:
+	SMaterialToolTip();
+	virtual ~SMaterialToolTip() {}
+
+	void Construct(const FArguments& InArgs);
+
+protected:
+	TWeakObjectPtr<UMaterialInterface> Material;
+	TAttribute<FText> Text;
+	TAttribute<FVector2D> MaterialSize;
+	TAttribute<bool> ShowMaterial;
+	bool ShowDefault;
+
+	FSlateMaterialBrush MaterialBrush;
+	
+	TSharedRef<SWidget> CreateToolTipWidget();
+
+	EVisibility GetShowMaterialVisibility() const;
+	FOptionalSize GetMaterialSizeX() const;
+	FOptionalSize GetMaterialSizeY() const;
+};
