[MaterialDesigner]

diff --git a/Engine/Plugins/Experimental/DynamicMaterial/Source/DynamicMaterialEditor/Private/Slate/SOrientBox.cpp b/Engine/Plugins/Experimental/DynamicMaterial/Source/DynamicMaterialEditor/Private/Slate/SOrientBox.cpp
new file mode 100644
--- /dev/null
+++ b/Engine/Plugins/Experimental/DynamicMaterial/Source/DynamicMaterialEditor/Private/Slate/SOrientBox.cpp
@@ -0,0 +1,337 @@
+// Copyright Epic Games, Inc. All Rights Reserved.
+
+#include "SOrientBox.h"
+#include "Layout/LayoutUtils.h"
+#include "Widgets/SBoxPanel.h"
+
+#define LOCTEXT_NAMESPACE "SOrientBox"
+
+SOrientBox::FSlot::FSlotArguments SOrientBox::Slot()
+{
+	return FSlot::FSlotArguments(MakeUnique<FSlot>());
+}
+
+EOrientation SOrientBox::InvertOrientation(const EOrientation InOrietation)
+{
+	return InOrietation == Orient_Horizontal ? Orient_Vertical : Orient_Horizontal;
+}
+
+EHorizontalAlignment SOrientBox::InvertHAlign(const EHorizontalAlignment InAlignment)
+{
+	switch (InAlignment)
+	{
+		case HAlign_Left: return HAlign_Right;
+		case HAlign_Right: return HAlign_Left;
+	}
+	return InAlignment;
+}
+
+EVerticalAlignment SOrientBox::InvertVAlign(const EVerticalAlignment InAlignment)
+{
+	switch (InAlignment)
+	{
+		case VAlign_Top: return VAlign_Bottom;
+		case VAlign_Bottom: return VAlign_Top;
+	}
+	return InAlignment;
+}
+
+EVerticalAlignment SOrientBox::HAlignToVAlign(const EHorizontalAlignment InAlignment)
+{
+	switch (InAlignment)
+	{
+		case HAlign_Fill: return VAlign_Fill;
+		case HAlign_Left: return VAlign_Top;
+		case HAlign_Center: return VAlign_Center;
+		case HAlign_Right: return VAlign_Bottom;
+	}
+	return VAlign_Center;
+}
+
+EHorizontalAlignment SOrientBox::VAlignToHAlign(const EVerticalAlignment InAlignment)
+{
+	switch (InAlignment)
+	{
+		case VAlign_Fill: return HAlign_Fill;
+		case VAlign_Top: return HAlign_Left;
+		case VAlign_Center: return HAlign_Center;
+		case VAlign_Bottom: return HAlign_Right;
+	}
+	return HAlign_Center;
+}
+
+FMargin SOrientBox::InvertPadding(const FMargin& InPadding)
+{
+	return FMargin(InPadding.Right, InPadding.Bottom, InPadding.Left, InPadding.Top);
+}
+
+FMargin SOrientBox::InvertHorizontalPadding(const FMargin& InPadding)
+{
+	return FMargin(InPadding.Right, InPadding.Top, InPadding.Left, InPadding.Bottom);
+}
+
+FMargin SOrientBox::InvertVerticalPadding(const FMargin& InPadding)
+{
+	return FMargin(InPadding.Left, InPadding.Bottom, InPadding.Right, InPadding.Top);
+}
+
+SOrientBox::SOrientBox()
+	: Children(this)
+{
+	SetCanTick(true);
+	bCanSupportFocus = false;
+}
+
+void SOrientBox::Construct(const FArguments& InArgs)
+{
+	Orientation = InArgs._Orientation;
+	Reverse = InArgs._Reverse;
+
+	Children.AddSlots(MoveTemp(const_cast<TArray<FSlot::FSlotArguments>&>(InArgs._Slots)));
+	
+	SetOrientation(Orientation.Get());
+	SetReversed(Reverse.Get());
+}
+
+void SOrientBox::Tick(const FGeometry& InAllottedGeometry, const double InCurrentTime, const float InDeltaTime)
+{
+	SPanel::Tick(InAllottedGeometry, InCurrentTime, InDeltaTime);
+
+	if (Orientation.Get() != CachedOrientation)
+	{
+		SetOrientation(Orientation.Get());
+	}
+}
+
+FVector2D SOrientBox::ComputeDesiredSize(float) const
+{
+	const int32 ChildCount = Children.Num();
+	const bool bIsReversed = Reverse.Get();
+
+	FVector2D ThisDesiredSize = FVector2D::ZeroVector;
+	for (int32 SlotIndex = 0; SlotIndex < ChildCount; ++SlotIndex)
+	{
+		const int32 ThisSlotIndex = bIsReversed
+			? (ChildCount - 1) - SlotIndex
+			: SlotIndex;
+
+		const SOrientBox::FSlot& ThisSlot = Children[ThisSlotIndex];
+
+		if (ThisSlot.GetWidget()->GetVisibility() != EVisibility::Collapsed)
+		{
+			const FVector2D ChildDesiredSize = ThisSlot.GetWidget()->GetDesiredSize();
+			FMargin ThisPadding = Orientation.Get() == EOrientation::Orient_Horizontal ? ThisSlot.GetPadding() : InvertPadding(ThisSlot.GetPadding());
+
+			if (Orientation.Get() == Orient_Horizontal)
+			{
+				ThisPadding = bIsReversed ? InvertHorizontalPadding(ThisPadding) : ThisPadding;
+				ThisDesiredSize.X += ChildDesiredSize.X + ThisPadding.GetTotalSpaceAlong<Orient_Horizontal>();
+				ThisDesiredSize.Y = FMath::Max(ChildDesiredSize.Y, ThisDesiredSize.Y);
+			}
+			else
+			{
+				ThisPadding = bIsReversed ? InvertVerticalPadding(ThisPadding) : ThisPadding;
+				ThisDesiredSize.X = FMath::Max(ChildDesiredSize.X, ThisDesiredSize.X);
+				ThisDesiredSize.Y += ChildDesiredSize.Y + ThisPadding.GetTotalSpaceAlong<Orient_Vertical>();
+			}
+		}
+	}
+
+	return ThisDesiredSize;
+}
+
+void SOrientBox::OnArrangeChildren(const FGeometry& InAllottedGeometry, FArrangedChildren& ArrangedChildren) const
+{
+	const int32 ChildCount = Children.Num();
+	const bool bIsReversed = Reverse.Get();
+
+	float CurrentChildOffset = 0.0f;
+
+	for (int32 SlotIndex = 0; SlotIndex < Children.Num(); ++SlotIndex)
+	{
+		const int32 ThisSlotIndex = bIsReversed
+			? (ChildCount - 1) - SlotIndex
+			: SlotIndex;
+
+		const SOrientBox::FSlot& ThisSlot = Children[ThisSlotIndex];
+		const EVisibility ChildVisibility = ThisSlot.GetWidget()->GetVisibility();
+
+		if (ChildVisibility != EVisibility::Collapsed)
+		{
+			const EFlowDirection FlowDirection = EFlowDirection::LeftToRight;
+			const FMargin& ThisPadding = Orientation.Get() == EOrientation::Orient_Horizontal ? ThisSlot.GetPadding() : InvertPadding(ThisSlot.GetPadding());
+			const FVector2D& WidgetDesiredSize = ThisSlot.GetWidget()->GetDesiredSize();
+
+			if (Orientation.Get() == EOrientation::Orient_Horizontal)
+			{
+				const AlignmentArrangeResult YAlignmentResult = AlignChild<Orient_Vertical>(FlowDirection, InAllottedGeometry.GetLocalSize().Y, ThisSlot, ThisPadding);
+				ArrangedChildren.AddWidget(InAllottedGeometry.MakeChild(ThisSlot.GetWidget(), FVector2D(CurrentChildOffset + ThisPadding.Left, YAlignmentResult.Offset), FVector2D(WidgetDesiredSize.X, YAlignmentResult.Size)));
+				
+				CurrentChildOffset += WidgetDesiredSize.X + ThisPadding.GetTotalSpaceAlong<Orient_Horizontal>();
+			}
+			else
+			{
+				const AlignmentArrangeResult XAlignmentResult = AlignChild<Orient_Horizontal>(FlowDirection, InAllottedGeometry.GetLocalSize().X, ThisSlot, ThisPadding);
+				ArrangedChildren.AddWidget(InAllottedGeometry.MakeChild(ThisSlot.GetWidget(), FVector2D(XAlignmentResult.Offset, CurrentChildOffset + ThisPadding.Top), FVector2D(XAlignmentResult.Size, WidgetDesiredSize.Y)));
+
+				CurrentChildOffset += WidgetDesiredSize.Y + ThisPadding.GetTotalSpaceAlong<Orient_Vertical>();
+			}
+		}
+	}
+}
+
+AlignmentArrangeResult SOrientBox::AlignChildSlot(const EOrientation InOrientation, const EFlowDirection InLayoutFlow, const FVector2D InAllottedSize, const FSlot& InChildToArrange, const FMargin& InSlotPadding, const float& InContentScale, bool bInClampToParent)
+{
+	const int32 Alignment = GetChildAlignment(InOrientation, InLayoutFlow, InChildToArrange);
+	float TotalMargin = 0.0f;
+	float MarginPre = 0.0f;
+	float MarginPost = 0.0f;
+	float AllottedSize = 0.0f;
+
+	if (InOrientation == Orient_Vertical)
+	{
+		TotalMargin = InSlotPadding.GetTotalSpaceAlong<Orient_Horizontal>();
+		MarginPre = InSlotPadding.Left;
+		MarginPost = InSlotPadding.Right;
+		AllottedSize = InAllottedSize.Y;
+	}
+	else
+	{
+		TotalMargin = InSlotPadding.GetTotalSpaceAlong<Orient_Vertical>();
+		MarginPre = InSlotPadding.Top;
+		MarginPost = InSlotPadding.Bottom;
+		AllottedSize = InAllottedSize.X;
+	}
+
+	switch (Alignment)
+	{
+	case HAlign_Fill:
+		return AlignmentArrangeResult(MarginPre, FMath::Max((AllottedSize - TotalMargin) * InContentScale, 0.f));
+	}
+
+	const float ChildDesiredSize = (InOrientation == Orient_Horizontal)
+		? (InChildToArrange.GetWidget()->GetDesiredSize().X * InContentScale)
+		: (InChildToArrange.GetWidget()->GetDesiredSize().Y * InContentScale);
+
+	const float ChildSize = FMath::Max((bInClampToParent ? FMath::Min(ChildDesiredSize, AllottedSize - TotalMargin) : ChildDesiredSize), 0.f);
+
+	switch (Alignment)
+	{
+	case HAlign_Left: // same as Align_Top
+		return AlignmentArrangeResult(MarginPre, ChildSize);
+	case HAlign_Center:
+		return AlignmentArrangeResult((AllottedSize - ChildSize) / 2.0f + MarginPre - MarginPost, ChildSize);
+	case HAlign_Right: // same as Align_Bottom		
+		return AlignmentArrangeResult(AllottedSize - ChildSize - MarginPost, ChildSize);
+	}
+
+	// Same as Fill
+	return AlignmentArrangeResult(MarginPre, FMath::Max((AllottedSize - TotalMargin) * InContentScale, 0.f));
+}
+
+int32 SOrientBox::GetChildAlignment(const EOrientation InOrientation, const EFlowDirection InFlowDirection, const FSlot& InChildToArrange)
+{
+	EHorizontalAlignment HAlign = InChildToArrange.GetHorizontalAlignment();
+	EVerticalAlignment VAlign = InChildToArrange.GetVerticalAlignment();
+
+	// We flip the horizontal and vertical alignments when orientation is not default (horizontal).
+	if (InOrientation == Orient_Vertical)
+	{
+		HAlign = VAlignToHAlign(VAlign);
+		VAlign = HAlignToVAlign(HAlign);
+	}
+
+	switch (InOrientation)
+	{
+	case Orient_Horizontal:
+	{
+		switch (InFlowDirection)
+		{
+		default:
+		case EFlowDirection::LeftToRight:
+			return StaticCast<int32>(HAlign);
+		case EFlowDirection::RightToLeft:
+			switch (HAlign)
+			{
+			case HAlign_Left: return StaticCast<int32>(HAlign_Right);
+			case HAlign_Right: return StaticCast<int32>(HAlign_Left);
+			default: return StaticCast<int32>(HAlign);
+			}
+		}
+		break;
+	}
+	case Orient_Vertical:
+	{
+		switch (InFlowDirection)
+		{
+		default:
+		case EFlowDirection::LeftToRight:
+			return StaticCast<int32>(VAlign);
+		case EFlowDirection::RightToLeft:
+			switch (VAlign)
+			{
+			case VAlign_Top: return StaticCast<int32>(VAlign_Bottom);
+			case VAlign_Bottom: return StaticCast<int32>(VAlign_Top);
+			default: return StaticCast<int32>(VAlign);
+			}
+		}
+		break;
+	}
+	}
+	return 0;
+}
+
+SOrientBox::FScopedWidgetSlotArguments SOrientBox::AddSlot(const int32 SlotIndex)
+{
+	if (!Children.IsValidIndex(SlotIndex))
+	{
+		// Insert at the end.
+		return FScopedWidgetSlotArguments { MakeUnique<FSlot>(), Children, INDEX_NONE };
+	}
+	else
+	{
+		TWeakPtr<SOrientBox> OrientBoxWeak = SharedThis(this);
+		return FScopedWidgetSlotArguments{ MakeUnique<FSlot>(), Children, SlotIndex, [OrientBoxWeak](const FSlot*, int32 SlotIndex) {} };
+	}
+}
+
+void SOrientBox::RemoveSlot(const TSharedRef<SWidget>& WidgetToRemove)
+{
+	Children.Remove(WidgetToRemove);
+}
+
+void SOrientBox::ClearChildren()
+{
+	Children.Empty();
+}
+
+void SOrientBox::SetOrientation(const EOrientation InOrientation)
+{
+	CachedOrientation = InOrientation;
+
+	Invalidate(EInvalidateWidget::Layout);
+}
+
+void SOrientBox::SetReversed(const bool bInReversed)
+{
+	CachedReverse = bInReversed;
+
+	Invalidate(EInvalidateWidget::Layout);
+}
+
+bool SOrientBox::IsContentVertical() const
+{
+	return CachedOrientation == Orient_Vertical;
+}
+
+bool SOrientBox::IsContentHorizontal() const
+{
+	return CachedOrientation == Orient_Horizontal;
+}
+
+bool SOrientBox::IsContentReversed() const
+{
+	return Reverse.Get();
+}
+
+#undef LOCTEXT_NAMESPACE
diff --git a/Engine/Plugins/Experimental/DynamicMaterial/Source/DynamicMaterialEditor/Private/Slate/SOrientBox.h b/Engine/Plugins/Experimental/DynamicMaterial/Source/DynamicMaterialEditor/Private/Slate/SOrientBox.h
new file mode 100644
--- /dev/null
+++ b/Engine/Plugins/Experimental/DynamicMaterial/Source/DynamicMaterialEditor/Private/Slate/SOrientBox.h
@@ -0,0 +1,106 @@
+// Copyright Epic Games, Inc. All Rights Reserved.
+
+#pragma once
+
+#include "Layout/BasicLayoutWidgetSlot.h"
+#include "Layout/Children.h"
+#include "Layout/WidgetSlotWithAttributeSupport.h"
+#include "Widgets/DeclarativeSyntaxSupport.h"
+#include "Widgets/SPanel.h"
+#include "Widgets/SWidget.h"
+
+struct AlignmentArrangeResult;
+
+/**
+ * Orient Box
+ * 
+ * Manages vertical and horizontal layout of child slots.
+ * Child slots can also have their order reversed.
+ * Automatically flips HAlign, VAlign, and Padding based on orientation and reverse state.
+ * 
+ * Slot alignment and padding should be set for the default horizontal orientation. For example,
+ * if you want bottom padding for a vertical oriented box, you should add padding to the right margin
+ * and set orientation to vertical.
+ */
+class SOrientBox : public SPanel
+{
+public:
+	/** A Slot that provides layout options for the contents of an oriented box. */
+	using FSlot = FBasicLayoutWidgetSlot;
+
+	static EOrientation InvertOrientation(const EOrientation InOrietation);
+
+	static EHorizontalAlignment InvertHAlign(const EHorizontalAlignment InAlignment);
+	static EVerticalAlignment InvertVAlign(const EVerticalAlignment InAlignment);
+
+	static EVerticalAlignment HAlignToVAlign(const EHorizontalAlignment InAlignment);
+	static EHorizontalAlignment VAlignToHAlign(const EVerticalAlignment InAlignment);
+
+	static FMargin InvertPadding(const FMargin& InPadding);
+	static FMargin InvertHorizontalPadding(const FMargin& InPadding);
+	static FMargin InvertVerticalPadding(const FMargin& InPadding);
+
+	static AlignmentArrangeResult AlignChildSlot(const EOrientation InOrientation, const EFlowDirection InLayoutFlow, const FVector2D InAllottedSize, const FSlot& InChildToArrange, const FMargin& InSlotPadding, const float& InContentScale = 1.0f, bool bInClampToParent = true);
+
+	/**
+	 * Creates a new widget slot.
+	 *
+	 * @return A new slot.
+	 */
+	static FSlot::FSlotArguments Slot();
+
+	SLATE_BEGIN_ARGS(SOrientBox)
+		: _Orientation(EOrientation::Orient_Horizontal)
+		, _Reverse(false)
+		{}
+		SLATE_SLOT_ARGUMENT(FSlot, Slots)
+		SLATE_ATTRIBUTE(EOrientation, Orientation)
+		SLATE_ATTRIBUTE(bool, Reverse)
+	SLATE_END_ARGS()
+
+	using FScopedWidgetSlotArguments = TPanelChildren<FSlot>::FScopedWidgetSlotArguments;
+	/**
+		* Adds a slot to the widget switcher at the specified location.
+		*
+		* @param SlotIndex The index at which to insert the slot, or INDEX_NONE to append.
+		*/
+	FScopedWidgetSlotArguments AddSlot(const int32 SlotIndex = INDEX_NONE);
+
+	/** Removes a slot at the specified location. */
+	void RemoveSlot(const TSharedRef<SWidget>& WidgetToRemove);
+
+	/** Removes all children from the box. */
+	void ClearChildren();
+
+	TPanelChildren<FSlot> Children;
+
+	SOrientBox();
+
+	void Construct(const FArguments& InArgs);
+
+	//~ Begin SWidget
+	virtual FChildren* GetChildren() override { return &Children; }
+	virtual void Tick(const FGeometry& InAllottedGeometry, const double InCurrentTime, const float InDeltaTime) override;
+	virtual FVector2D ComputeDesiredSize(float) const override;
+	virtual void OnArrangeChildren(const FGeometry& InAllottedGeometry, FArrangedChildren& ArrangedChildren) const override;
+	//~ End SWidget
+
+	EOrientation GetOrientation() const { return CachedOrientation; }
+	void SetOrientation(const EOrientation InOrientation);
+
+	bool IsReversed() const { return CachedReverse; }
+	void SetReversed(const bool bInReversed);
+
+	bool IsContentVertical() const;
+	bool IsContentHorizontal() const;
+	bool IsContentReversed() const;
+
+protected:
+	TAttribute<EOrientation> Orientation;
+	TAttribute<bool> Reverse;
+
+	EOrientation CachedOrientation;
+	bool CachedReverse;
+
+	static int32 GetChildAlignment(const EOrientation InOrientation, const EFlowDirection InFlowDirection, const FSlot& InChildToArrange);
+};
diff --git a/Engine/Plugins/Experimental/DynamicMaterial/Source/DynamicMaterialEditor/Public/Components/DMMaterialEffect.h b/Engine/Plugins/Experimental/DynamicMaterial/Source/DynamicMaterialEditor/Public/Components/DMMaterialEffect.h
new file mode 100644
--- /dev/null
+++ b/Engine/Plugins/Experimental/DynamicMaterial/Source/DynamicMaterialEditor/Public/Components/DMMaterialEffect.h
@@ -0,0 +1,89 @@
+// Copyright Epic Games, Inc. All Rights Reserved.
+
+#pragma once
+
+#include "Components/DMMaterialComponent.h"
+#include "Components/DMMaterialStage.h"
+#include "DMDefs.h"
+#include "DMMaterialEffect.generated.h"
+
+class UDMMaterialEffect;
+class UDMMaterialEffectStack;
+class UDMMaterialLayerObject;
+class UDMMaterialValue;
+class UMaterialExpression;
+struct FDMMaterialBuildState;
+
+UENUM(BlueprintType)
+enum class EDMMaterialEffectTarget : uint8
+{
+	None      = 0,
+	BaseStage = 1 << 0,
+	MaskStage = 1 << 1,
+	TextureUV = 1 << 2,
+	Slot      = 1 << 3
+};
+
+UCLASS(Abstract, BlueprintType, ClassGroup = "Material Designer", Meta = (DisplayName = "Material Designer Effect"))
+class DYNAMICMATERIALEDITOR_API UDMMaterialEffect : public UDMMaterialComponent
+{
+	GENERATED_BODY()
+
+public:
+	UFUNCTION(BlueprintPure, Category = "Material Designer")
+	static EDMMaterialEffectTarget StageTypeToEffectType(EDMMaterialLayerStage InStageType);
+
+	UFUNCTION(BlueprintCallable, Category = "Material Designer")
+	static UDMMaterialEffect* CreateEffect(UDMMaterialEffectStack* InEffectStack, TSubclassOf<UDMMaterialEffect> InEffectClass);
+
+	template<typename InEffectClass>
+	static InEffectClass* CreateEffect(UDMMaterialEffectStack* InEffectStack)
+	{
+		return Cast<InEffectClass>(CreateEffect(InEffectStack, InEffectClass::StaticClass()));
+	}
+
+	UDMMaterialEffect();
+
+	UFUNCTION(BlueprintPure, Category = "Material Designer")
+	UDMMaterialEffectStack* GetEffectStack() const;
+
+	UFUNCTION(BlueprintCallable, Category = "Material Designer")
+	int32 FindIndex() const;
+
+	UFUNCTION(BlueprintPure, Category = "Material Designer")
+	bool IsEnabled() const;
+
+	UFUNCTION(BlueprintCallable, Category = "Material Designer")
+	bool SetEnabled(bool bInIsEnabled);
+
+	UFUNCTION(BlueprintPure, Category = "Material Designer")
+	EDMMaterialEffectTarget GetEffectTarget() const;
+
+	UFUNCTION(BlueprintPure, Category = "Material Designer")
+	virtual FText GetEffectName() const;
+
+	UFUNCTION(BlueprintPure, Category = "Material Designer")
+	virtual FText GetEffectDescription() const PURE_VIRTUAL(UDMMaterialEffect::ApplyTo, return FText::GetEmpty();)
+
+	virtual void ApplyTo(const TSharedRef<FDMMaterialBuildState>& InBuildState, TArray<UMaterialExpression*>& InOutExpressions, 
+		int32& InOutLastExpressionOutputChannel, int32& InLastExpressionOutputIndex) const PURE_VIRTUAL(UDMMaterialEffect::ApplyTo)
+
+	//~ Begin UDMMaterialComponent
+	virtual UDMMaterialComponent* GetParentComponent() const override;
+	virtual FString GetComponentPathComponent() const override;
+	virtual FText GetComponentDescription() const override;
+	virtual void Update(EDMUpdateType InUpdateType) override;
+	virtual void PostEditorDuplicate(UDynamicMaterialModel* InMaterialModel, UDMMaterialComponent* InParent) override;
+	//~ End UDMMaterialComponent
+
+	//~ Begin UObject
+	virtual void PostEditUndo() override;
+	//~ End UObject
+
+protected:
+	UPROPERTY(VisibleInstanceOnly, BlueprintReadOnly, Category = "Material Designer")
+	EDMMaterialEffectTarget EffectTarget;
+
+	UPROPERTY(VisibleInstanceOnly, BlueprintReadOnly, Category = "Material Designer")
+	bool bEnabled;
+};
diff --git a/Engine/Plugins/Experimental/DynamicMaterial/Source/DynamicMaterialEditor/Public/Components/DMMaterialEffectFunction.h b/Engine/Plugins/Experimental/DynamicMaterial/Source/DynamicMaterialEditor/Public/Components/DMMaterialEffectFunction.h
new file mode 100644
--- /dev/null
+++ b/Engine/Plugins/Experimental/DynamicMaterial/Source/DynamicMaterialEditor/Public/Components/DMMaterialEffectFunction.h
@@ -0,0 +1,71 @@
+// Copyright Epic Games, Inc. All Rights Reserved.
+
+#pragma once
+
+#include "Components/DMMaterialEffect.h"
+#include "DMMaterialEffectFunction.generated.h"
+
+class UMaterialFunctionInterface;
+
+UCLASS(BlueprintType, ClassGroup = "Material Designer", Meta = (DisplayName = "Material Designer Effect Function"))
+class DYNAMICMATERIALEDITOR_API UDMMaterialEffectFunction : public UDMMaterialEffect
+{
+	GENERATED_BODY()
+
+public:
+	static const FString InputsPathToken;
+
+	UDMMaterialEffectFunction();
+
+	UFUNCTION(BlueprintPure, Category = "Material Designer")
+	UMaterialFunctionInterface* GetMaterialFunction() const;
+
+	UFUNCTION(BlueprintCallable, Category = "Material Designer")
+	bool SetMaterialFunction(UMaterialFunctionInterface* InFunction);
+
+	UFUNCTION(BlueprintPure, Category = "Material Designer")
+	UDMMaterialValue* GetInputValue(int32 InIndex) const;
+
+	UFUNCTION(BlueprintCallable, Category = "Material Designer", Meta = (DisplayName = "Get Input Values"))
+	TArray<UDMMaterialValue*> BP_GetInputValues() const;
+
+	const TArray<TObjectPtr<UDMMaterialValue>>& GetInputValues() const;
+
+	//~ Begin UDMMaterialEffect
+	virtual FText GetEffectName() const override;
+	virtual FText GetEffectDescription() const override;
+	virtual void ApplyTo(const TSharedRef<FDMMaterialBuildState>& InBuildState, TArray<UMaterialExpression*>& InOutExpressions,
+		int32& InOutLastExpressionOutputChannel, int32& InOutLastExpressionOutputIndex) const override;
+	//~ End UDMMaterialEffect
+
+	//~ Begin UDMMaterialComponent
+	virtual FText GetComponentDescription() const override;
+	virtual void PostEditorDuplicate(UDynamicMaterialModel* InMaterialModel, UDMMaterialComponent* InParent) override;
+	//~ End UDMMaterialComponent
+
+	//~ Begin UObject
+	virtual bool Modify(bool bInAlwaysMarkDirty = true) override;
+	virtual void PostLoad() override;
+	//~ End UObject
+
+protected:
+	UPROPERTY(VisibleInstanceOnly, BlueprintReadOnly, Instanced, Category = "Material Designer")
+	TObjectPtr<UMaterialFunctionInterface> MaterialFunctionPtr;
+
+	UPROPERTY(VisibleInstanceOnly, BlueprintReadOnly, Instanced, Category = "Material Designer")
+	TArray<TObjectPtr<UDMMaterialValue>> InputValues;
+
+	void OnMaterialFunctionChanged();
+
+	void InitFunction();
+
+	void DeinitFunction();
+
+	bool NeedsFunctionInit();
+
+	//~ Begin UDMMaterialComponent
+	virtual UDMMaterialComponent* GetSubComponentByPath(FDMComponentPath& InPath, const FDMComponentPathSegment& InPathSegment) const override;
+	virtual void OnComponentAdded() override;
+	virtual void OnComponentRemoved() override;
+	//~ End UDMMaterialComponent
+};
diff --git a/Engine/Plugins/Experimental/DynamicMaterial/Source/DynamicMaterialEditor/Public/Components/DMMaterialEffectStack.h b/Engine/Plugins/Experimental/DynamicMaterial/Source/DynamicMaterialEditor/Public/Components/DMMaterialEffectStack.h
new file mode 100644
--- /dev/null
+++ b/Engine/Plugins/Experimental/DynamicMaterial/Source/DynamicMaterialEditor/Public/Components/DMMaterialEffectStack.h
@@ -0,0 +1,136 @@
+// Copyright Epic Games, Inc. All Rights Reserved.
+
+#pragma once
+
+#include "Components/DMMaterialComponent.h"
+#include "DMDefs.h"
+#include "DMMaterialEffectStack.generated.h"
+
+class UDMMaterialEffect;
+class UDMMaterialLayerObject;
+class UDMMaterialSlot;
+class UDMMaterialStage;
+class UMaterialExpression;
+enum class EDMMaterialEffectTarget : uint8;
+struct FDMMaterialBuildState;
+
+/**
+ * Container for effects. Effects can be applied to either layers (on a per stage basis) or to slots.
+ */
+UCLASS(BlueprintType, ClassGroup = "Material Designer", Meta = (DisplayName = "Material Designer Effect Stack"))
+class DYNAMICMATERIALEDITOR_API UDMMaterialEffectStack : public UDMMaterialComponent
+{
+	GENERATED_BODY()
+
+public:
+	static const FString EffectsPathToken;
+
+	using FEffectCallbackFunc = TFunctionRef<void(UDMMaterialEffect*)>;
+
+	UFUNCTION(BlueprintCallable, Category = "Material Designer", Meta = (DisplayName = "Create Effect Stack (For Slot)"))
+	static UDMMaterialEffectStack* CreateEffectStackForSlot(UDMMaterialSlot* InSlot)
+	{
+		return CreateEffectStack(InSlot);
+	}
+
+	UFUNCTION(BlueprintCallable, Category = "Material Designer", Meta = (DisplayName = "Create Effect Stack (For Layer)"))
+	static UDMMaterialEffectStack* CreateEffectStackForLayer(UDMMaterialLayerObject* InLayer)
+	{
+		return CreateEffectStack(InLayer);
+	}
+
+	static UDMMaterialEffectStack* CreateEffectStack(UDMMaterialSlot* InSlot);
+	static UDMMaterialEffectStack* CreateEffectStack(UDMMaterialLayerObject* InLayer);
+
+	UDMMaterialEffectStack();
+
+	UFUNCTION(BlueprintPure, Category = "Material Designer")
+	UDMMaterialSlot* GetSlot() const;
+
+	UFUNCTION(BlueprintPure, Category = "Material Designer")
+	UDMMaterialLayerObject* GetLayer() const;
+
+	UFUNCTION(BlueprintCallable, Category = "Material Designer")
+	bool IsEnabled() const;
+
+	UFUNCTION(BlueprintCallable, Category = "Material Designer")
+	bool SetEnabled(bool bInIsEnabled);
+
+	UFUNCTION(BlueprintPure, Category = "Material Designer")
+	UDMMaterialEffect* GetEffect(int32 InIndex) const;
+
+	UFUNCTION(BlueprintCallable, Category = "Material Designer", Meta = (DisplayName = "Get Effects"))
+	TArray<UDMMaterialEffect*> BP_GetEffects() const;
+
+	const TArray<TObjectPtr<UDMMaterialEffect>>& GetEffects() const;
+
+	UFUNCTION(BlueprintCallable, Category = "Material Designer")
+	bool HasEffect(const UDMMaterialEffect* InEffect) const;
+
+	UFUNCTION(BlueprintCallable, Category = "Material Designer")
+	bool AddEffect(UDMMaterialEffect* InEffect);
+
+	UFUNCTION(BlueprintCallable, Category = "Material Designer")
+	bool SetEffect(int32 InIndex, UDMMaterialEffect* InEffect);
+
+	UFUNCTION(BlueprintCallable, Category = "Material Designer", Meta = (DisplayName = "Move Effect (By Index)"))
+	bool BP_MoveEffectByIndex(int32 InIndex, int32 InNewIndex)
+	{
+		return MoveEffect(InIndex, InNewIndex);
+	}
+
+	bool MoveEffect(int32 InIndex, int32 InNewIndex);
+
+	UFUNCTION(BlueprintCallable, Category = "Material Designer", Meta = (DisplayName = "Move Effect (By Value)"))
+	bool BP_MoveEffectByValue(UDMMaterialEffect* InEffect, int32 InNewIndex)
+	{
+		return MoveEffect(InEffect, InNewIndex);
+	}
+
+	bool MoveEffect(UDMMaterialEffect* InEffect, int32 InNewIndex);
+
+	UFUNCTION(BlueprintCallable, Category = "Material Designer", Meta = (DisplayName = "Remove Effect (By Index)"))
+	bool BP_RemoveEffectByIndex(int32 InIndex)
+	{
+		return RemoveEffect(InIndex);
+	}
+
+	bool RemoveEffect(int32 InIndex);
+
+	UFUNCTION(BlueprintCallable, Category = "Material Designer", Meta = (DisplayName = "Remove Effect (By Value)"))
+	bool BP_RemoveEffectByValue(UDMMaterialEffect* InEffect)
+	{
+		return RemoveEffect(InEffect);
+	}
+
+	bool RemoveEffect(UDMMaterialEffect* InEffect);
+
+	bool ApplyEffects(const TSharedRef<FDMMaterialBuildState>& InBuildState, EDMMaterialEffectTarget InEffectTarget,
+		TArray<UMaterialExpression*>& InOutStageExpressions, int32& InOutLastExpressionOutputChannel, int32& InOutLastExpressionOutputIndex) const;
+
+	//~ Begin UDMMaterialComponent
+	virtual UDMMaterialComponent* GetParentComponent() const override;
+	virtual FString GetComponentPathComponent() const override;
+	virtual FText GetComponentDescription() const override;
+	virtual void Update(EDMUpdateType InUpdateType) override;
+	virtual void PostEditorDuplicate(UDynamicMaterialModel* InMaterialModel, UDMMaterialComponent* InParent) override;
+	//~ End UDMMaterialComponent
+
+	//~ Begin UObject
+	virtual bool Modify(bool bInAlwaysMarkDirty = true) override;
+	virtual void PostEditUndo() override;
+	//~ End UObject
+
+protected:
+	UPROPERTY(VisibleInstanceOnly, BlueprintReadOnly, Category = "Material Designer")
+	bool bEnabled;
+
+	UPROPERTY(VisibleInstanceOnly, BlueprintReadOnly, Instanced, Category = "Material Designer")
+	TArray<TObjectPtr<UDMMaterialEffect>> Effects;
+
+	//~ Begin UDMMaterialComponent
+	virtual UDMMaterialComponent* GetSubComponentByPath(FDMComponentPath& InPath, const FDMComponentPathSegment& InPathSegment) const override;
+	virtual void OnComponentAdded() override;
+	virtual void OnComponentRemoved() override;
+	//~ End UDMMaterialComponent
+};
diff --git a/Engine/Plugins/Experimental/DynamicMaterial/Source/DynamicMaterialEditor/Public/Components/DMMaterialLayer.h b/Engine/Plugins/Experimental/DynamicMaterial/Source/DynamicMaterialEditor/Public/Components/DMMaterialLayer.h
new file mode 100644
--- /dev/null
+++ b/Engine/Plugins/Experimental/DynamicMaterial/Source/DynamicMaterialEditor/Public/Components/DMMaterialLayer.h
@@ -0,0 +1,182 @@
+// Copyright Epic Games, Inc. All Rights Reserved.
+
+#pragma once
+
+#include "Components/DMMaterialComponent.h"
+#include "DMDefs.h"
+#include "Components/DMMaterialStage.h"
+#include "DMMaterialLayer.generated.h"
+
+class UDMMaterialEffectStack;
+class UDMMaterialSlot;
+class UDMMaterialStage;
+class UMaterialExpression;
+
+UCLASS(BlueprintType, ClassGroup = "Material Designer", Meta = (DisplayName = "Material Designer Layer"))
+class DYNAMICMATERIALEDITOR_API UDMMaterialLayerObject : public UDMMaterialComponent
+{
+	GENERATED_BODY()
+
+public:
+	static const FString StagesPathToken;
+	static const FString BasePathToken;
+	static const FString MaskPathToken;
+	static const FString EffectStackPathToken;
+
+	using FStageCallbackFunc = TFunctionRef<void(UDMMaterialStage*)>;
+
+	UFUNCTION(BlueprintCallable, Category = "Material Designer")
+	static UDMMaterialLayerObject* CreateLayer(UDMMaterialSlot* InSlot, EDMMaterialPropertyType InMaterialProperty, 
+		const TArray<UDMMaterialStage*>& InStages);
+	
+	static UDMMaterialLayerObject* DeserializeFromString(UDMMaterialSlot* InOuter, const FString& InSerializedString);
+
+	UDMMaterialLayerObject();
+
+	UFUNCTION(BlueprintPure, Category = "Material Designer")
+	UDMMaterialSlot* GetSlot() const;
+
+	UFUNCTION(BlueprintCallable, Category = "Material Designer")
+	int32 FindIndex() const;
+
+	UFUNCTION(BlueprintCallable, Category = "Material Designer")
+	const FText& GetLayerName() const { return LayerName; }
+
+	UFUNCTION(BlueprintCallable, Category = "Material Designer")
+	void SetLayerName(const FText& InName);
+
+	UFUNCTION(BlueprintCallable, Category = "Material Designer")
+	bool IsEnabled() const;
+
+	UFUNCTION(BlueprintCallable, Category = "Material Designer")
+	bool SetEnabled(bool bInIsEnabled);
+
+	UFUNCTION(BlueprintPure, Category = "Material Designer")
+	EDMMaterialPropertyType GetMaterialProperty() const;
+
+	UFUNCTION(BlueprintCallable, Category = "Material Designer")
+	bool SetMaterialProperty(EDMMaterialPropertyType InMaterialProperty);
+
+	UFUNCTION(BlueprintPure, Category = "Material Designer")
+	bool IsTextureUVLinkEnabled() const;
+
+	UFUNCTION(BlueprintCallable, Category = "Material Designer")
+	bool SetTextureUVLinkEnabled(bool bInValue);
+
+	UFUNCTION(BlueprintCallable, Category = "Material Designer")
+	bool ToggleTextureUVLinkEnabled();
+
+	UFUNCTION(BlueprintCallable, Category = "Material Designer")
+	UDMMaterialLayerObject* GetPreviousLayer(EDMMaterialPropertyType InUsingProperty, EDMMaterialLayerStage InSearchFor) const;
+
+	UFUNCTION(BlueprintCallable, Category = "Material Designer")
+	UDMMaterialLayerObject* GetNextLayer(EDMMaterialPropertyType InUsingProperty, EDMMaterialLayerStage InSearchFor) const;
+
+	UFUNCTION(BlueprintCallable, Category = "Material Designer")
+	bool CanMoveLayerAbove(UDMMaterialLayerObject* InLayer) const;
+
+	UFUNCTION(BlueprintCallable, Category = "Material Designer")
+	bool CanMoveLayerBelow(UDMMaterialLayerObject* InLayer) const;
+
+	UFUNCTION(BlueprintCallable, Category = "Material Designer")
+	UDMMaterialStage* GetStage(EDMMaterialLayerStage InStageType = EDMMaterialLayerStage::All, bool bInCheckEnabled = false) const;
+
+	UFUNCTION(BlueprintCallable, Category = "Material Designer")
+	TArray<UDMMaterialStage*> GetStages(EDMMaterialLayerStage InStageType = EDMMaterialLayerStage::All, bool bInCheckEnabled = false) const;
+
+	const TArray<TObjectPtr<UDMMaterialStage>>& GetAllStages() const;
+
+	UFUNCTION(BlueprintCallable, Category = "Material Designer")
+	EDMMaterialLayerStage GetStageType(const UDMMaterialStage* InStage) const;
+
+	UFUNCTION(BlueprintCallable, Category = "Material Designer")
+	UDMMaterialStage* GetFirstValidStage(EDMMaterialLayerStage InStageScope) const;
+
+	UFUNCTION(BlueprintCallable, Category = "Material Designer")
+	UDMMaterialStage* GetLastValidStage(EDMMaterialLayerStage InStageScope) const;
+
+	UFUNCTION(BlueprintCallable, Category = "Material Designer")
+	bool HasValidStage(const UDMMaterialStage* InStage) const;
+
+	UFUNCTION(BlueprintCallable, Category = "Material Designer")
+	bool HasValidStageOfType(EDMMaterialLayerStage InStageScope = EDMMaterialLayerStage::All) const;
+
+	UFUNCTION(BlueprintPure, Category = "Material Designer")
+	bool IsStageEnabled(EDMMaterialLayerStage InStageScope = EDMMaterialLayerStage::All) const;
+
+	/** Checks for the first enabled and valid stage. */
+	UFUNCTION(BlueprintCallable, Category = "Material Designer")
+	UDMMaterialStage* GetFirstEnabledStage(EDMMaterialLayerStage InStageScope) const;
+
+	/** Checks for the last enabled and valid stage. */
+	UFUNCTION(BlueprintCallable, Category = "Material Designer")
+	UDMMaterialStage* GetLastEnabledStage(EDMMaterialLayerStage InStageScope) const;
+
+	UFUNCTION(BlueprintCallable, Category = "Material Designer")
+	bool SetStage(EDMMaterialLayerStage InStageType, UDMMaterialStage* InStage);
+
+	UFUNCTION(BlueprintCallable, Category = "Material Designer")
+	bool AreAllStagesValid(EDMMaterialLayerStage InStageScope) const;
+
+	/** Checks if both stages are enabled and valid */
+	UFUNCTION(BlueprintCallable, Category = "Material Designer")
+	bool AreAllStagesEnabled(EDMMaterialLayerStage InStageScope) const;
+
+	UFUNCTION(BlueprintPure, Category = "Material Designer")
+	bool IsStageBeingEdited(EDMMaterialLayerStage InStageScope = EDMMaterialLayerStage::All) const;
+
+	UFUNCTION(BlueprintCallable, Category = "Material Designer")
+	UDMMaterialStage* GetFirstStageBeingEdited(EDMMaterialLayerStage InStageScope) const;
+
+	void ForEachValidStage(EDMMaterialLayerStage InStageScope, FStageCallbackFunc InCallback) const;
+
+	void ForEachEnabledStage(EDMMaterialLayerStage InStageScope, FStageCallbackFunc InCallback) const;
+
+	UFUNCTION(BlueprintPure, Category = "Material Designer")
+	UDMMaterialEffectStack* GetEffectStack() const;
+
+	FString SerializeToString() const;
+
+	void GenerateExpressions(const TSharedRef<FDMMaterialBuildState>& InBuildState) const;
+
+	bool ApplyEffects(const TSharedRef<FDMMaterialBuildState>& InBuildState, const UDMMaterialStage* InStage,
+		TArray<UMaterialExpression*>& InOutStageExpressions, int32& InOutLastExpressionOutputChannel, int32& InOutLastExpressionOutputIndex) const;
+
+	//~ Begin UDMMaterialComponent
+	virtual UDMMaterialComponent* GetParentComponent() const override;
+	virtual FString GetComponentPathComponent() const override;
+	virtual FText GetComponentDescription() const override;
+	virtual void Update(EDMUpdateType InUpdateType) override;
+	virtual void PostEditorDuplicate(UDynamicMaterialModel* InMaterialModel, UDMMaterialComponent* InParent) override;
+	//~ End UDMMaterialComponent
+
+	//~ Begin UObject
+	virtual bool Modify(bool bInAlwaysMarkDirty = true) override;
+	virtual void PostEditUndo() override;
+	//~ End UObject
+
+protected:
+	UPROPERTY(VisibleInstanceOnly, BlueprintReadOnly, Category = "Material Designer")
+	EDMMaterialPropertyType MaterialProperty = EDMMaterialPropertyType::None;
+
+	UPROPERTY(VisibleInstanceOnly, BlueprintReadOnly, Category = "Material Designer")
+	FText LayerName;
+
+	UPROPERTY(VisibleInstanceOnly, BlueprintReadOnly, Category = "Material Designer")
+	bool bEnabled;
+
+	UPROPERTY(VisibleInstanceOnly, BlueprintReadOnly, Instanced, Category = "Material Designer")
+	TArray<TObjectPtr<UDMMaterialStage>> Stages;
+
+	UPROPERTY(VisibleInstanceOnly, BlueprintReadOnly, Instanced, Category = "Material Designer")
+	TObjectPtr<UDMMaterialEffectStack> EffectStack;
+
+	UPROPERTY(VisibleInstanceOnly, BlueprintReadOnly, Category = "Material Designer")
+	bool bLinkedUVs;
+
+	//~ Begin UDMMaterialComponent
+	virtual UDMMaterialComponent* GetSubComponentByPath(FDMComponentPath& InPath, const FDMComponentPathSegment& InPathSegment) const override;
+	virtual void OnComponentAdded() override;
+	virtual void OnComponentRemoved() override;
+	//~ End UDMMaterialComponent
+};
diff --git a/Engine/Plugins/Experimental/DynamicMaterial/Source/DynamicMaterialEditor/Public/Components/DMMaterialProperty.h b/Engine/Plugins/Experimental/DynamicMaterial/Source/DynamicMaterialEditor/Public/Components/DMMaterialProperty.h
new file mode 100644
--- /dev/null
+++ b/Engine/Plugins/Experimental/DynamicMaterial/Source/DynamicMaterialEditor/Public/Components/DMMaterialProperty.h
@@ -0,0 +1,97 @@
+// Copyright Epic Games, Inc. All Rights Reserved.
+
+#pragma once
+
+#include "DMEDefs.h"
+#include "Components/DMMaterialComponent.h"
+#include "Math/MathFwd.h"
+#include "Templates/SharedPointerFwd.h"
+#include "DMMaterialProperty.generated.h"
+
+class FText;
+class UDMMaterialSlot;
+class UDynamicMaterialModel;
+class UDynamicMaterialModelEditorOnlyData;
+class UMaterialExpression;
+class UMaterialFunctionInterface;
+struct FDMMaterialBuildState;
+
+/**
+ * Base Color, Specular, Opacity, etc
+ */
+UCLASS(BlueprintType, ClassGroup = "Material Designer", Meta = (DisplayName = "Material Designer Property"))
+class DYNAMICMATERIALEDITOR_API UDMMaterialProperty : public UDMMaterialComponent
+{
+	GENERATED_BODY()
+
+public:
+	static UDMMaterialProperty* CreateCustomMaterialPropertyDefaultSubobject(UDynamicMaterialModelEditorOnlyData* InModelEditorOnlyData, 
+		EDMMaterialPropertyType InMaterialProperty, const FName& InSubObjName);
+
+	UFUNCTION(BlueprintPure, Category = "Material Designer")
+	UDynamicMaterialModelEditorOnlyData* GetMaterialModelEditorOnlyData() const;
+
+	UFUNCTION(BlueprintPure, Category = "Material Designer")
+	EDMMaterialPropertyType GetMaterialProperty() const { return MaterialProperty; }
+
+	UFUNCTION(BlueprintPure, Category = "Material Designer")
+	FText GetDescription() const;
+
+	UFUNCTION(BlueprintPure, Category = "Material Designer")
+	EDMValueType GetInputConnectorType() const { return InputConnectorType; }
+
+	UFUNCTION(BlueprintPure, Category = "Material Designer")
+	const FDMMaterialStageConnection& GetInputConnectionMap() const { return InputConnectionMap; }
+
+	FDMMaterialStageConnection& GetInputConnectionMap() { return InputConnectionMap; }
+
+	UFUNCTION(BlueprintPure, Category = "Material Designer")
+	UMaterialFunctionInterface* GetOutputProcessor() const { return OutputProcessor; }
+
+	UFUNCTION(BlueprintCallable, Category = "Material Designer")
+	void SetOutputProcessor(UMaterialFunctionInterface* InFunction);
+
+	virtual bool IsValidForModel(UDynamicMaterialModelEditorOnlyData& InMaterialModel) const { return true; }
+
+	UFUNCTION(BlueprintCallable, Category = "Material Designer")
+	virtual void ResetInputConnectionMap();
+
+	virtual UMaterialExpression* GetDefaultInput(const TSharedRef<FDMMaterialBuildState>& InBuildState) const;
+
+	//~ Begin UDMMaterialComponent
+	virtual void Update(EDMUpdateType InUpdateType) override;
+	virtual bool IsRootComponent() const override { return true; }
+	virtual FString GetComponentPathComponent() const override;
+	virtual void PostEditorDuplicate(UDynamicMaterialModel* InMaterialModel, UDMMaterialComponent* InParent) override;
+	virtual void PreEditChange(FProperty* InPropertyAboutToChange) override;
+	virtual void PostEditChangeProperty(FPropertyChangedEvent& InPropertyChangedEvent) override;
+	//~ End UDMMaterialComponent
+
+	void LoadDeprecatedModelData(UDMMaterialProperty* InOldProperty);
+
+protected:
+	static UMaterialExpression* CreateConstant(const TSharedRef<FDMMaterialBuildState>& InBuildState, float InDefaultValue);
+	static UMaterialExpression* CreateConstant(const TSharedRef<FDMMaterialBuildState>& InBuildState, const FVector2d& InDefaultValue);
+	static UMaterialExpression* CreateConstant(const TSharedRef<FDMMaterialBuildState>& InBuildState, const FVector3d& InDefaultValue);
+	static UMaterialExpression* CreateConstant(const TSharedRef<FDMMaterialBuildState>& InBuildState, const FVector4d& InDefaultValue);
+
+	UPROPERTY(VisibleInstanceOnly, BlueprintReadOnly, Category = "Material Designer")
+	EDMMaterialPropertyType MaterialProperty;
+
+	UPROPERTY(VisibleInstanceOnly, BlueprintReadOnly, Category = "Material Designer")
+	EDMValueType InputConnectorType;
+
+	UPROPERTY(VisibleInstanceOnly, BlueprintReadOnly, Category = "Material Designer")
+	FDMMaterialStageConnection InputConnectionMap;
+
+	UPROPERTY(VisibleInstanceOnly, BlueprintReadOnly, Category = "Material Designer")
+	TObjectPtr<UMaterialFunctionInterface> OutputProcessor;
+
+	UPROPERTY(Transient, DuplicateTransient, TextExportTransient)
+	TObjectPtr<UMaterialFunctionInterface> OutputProcessor_PreUpdate;
+
+	UDMMaterialProperty();
+	UDMMaterialProperty(EDMMaterialPropertyType InMaterialProperty, EDMValueType InInputConnectorType);
+
+	void OnOutputProcessorUpdated();
+};
diff --git a/Engine/Plugins/Experimental/DynamicMaterial/Source/DynamicMaterialEditor/Public/Components/DMMaterialSlot.h b/Engine/Plugins/Experimental/DynamicMaterial/Source/DynamicMaterialEditor/Public/Components/DMMaterialSlot.h
new file mode 100644
--- /dev/null
+++ b/Engine/Plugins/Experimental/DynamicMaterial/Source/DynamicMaterialEditor/Public/Components/DMMaterialSlot.h
@@ -0,0 +1,184 @@
+// Copyright Epic Games, Inc. All Rights Reserved.
+
+#pragma once
+
+#include "Components/DMMaterialComponent.h"
+#include "DMEDefs.h"
+#include "DMMaterialSlot.generated.h"
+
+class FScopedTransaction;
+class UDMMaterialLayerObject;
+class UDMMaterialSlot;
+class UDMMaterialStage;
+class UDynamicMaterialModelEditorOnlyData;
+class UMaterial;
+struct FDMMaterialBuildState;
+struct FDMMaterialLayer;
+
+DECLARE_MULTICAST_DELEGATE_OneParam(FDMOnMaterialSlotConnectorsUpdated, UDMMaterialSlot*);
+DECLARE_MULTICAST_DELEGATE_OneParam(FDMOnMaterialSlotPropertiesUpdated, UDMMaterialSlot*);
+DECLARE_MULTICAST_DELEGATE_OneParam(FDMOnMaterialSlotLayersUpdated, UDMMaterialSlot*);
+
+/**
+ * A list of operations/inputs daisy chained together to produce an output.
+ */
+UCLASS(BlueprintType, ClassGroup = "Material Designer", Meta = (DisplayName = "Material Designer Slot"))
+class DYNAMICMATERIALEDITOR_API UDMMaterialSlot : public UDMMaterialComponent
+{
+	friend class SDMMaterialSlot;
+	friend class SDMSlot;
+
+	GENERATED_BODY()
+
+public:
+	static const FString LayersPathToken;
+
+	UDMMaterialSlot();
+
+	UFUNCTION(BlueprintPure, Category = "Material Designer")
+	UDynamicMaterialModelEditorOnlyData* GetMaterialModelEditorOnlyData() const;
+
+	UFUNCTION(BlueprintPure, Category = "Material Designer")
+	int32 GetIndex() const { return Index; }
+
+	void SetIndex(int32 InNewIndex) { Index = InNewIndex; }
+
+	UFUNCTION(BlueprintPure, Category = "Material Designer")
+	FText GetDescription() const;
+
+	UFUNCTION(BlueprintPure, Category = "Material Designer")
+	const TArray<EDMValueType>& GetOutputConnectorTypesForMaterialProperty(EDMMaterialPropertyType InMaterialProperty) const;
+
+	UFUNCTION(BlueprintPure, Category = "Material Designer")
+	TSet<EDMValueType> GetAllOutputConnectorTypes() const;
+
+	UFUNCTION(BlueprintPure, Category = "Material Designer", meta = (DisplayName = "Get Layer"))
+	UDMMaterialLayerObject* GetLayer(int32 InLayerIndex) const;
+
+	UFUNCTION(BlueprintCallable, Category = "Material Designer")
+	UDMMaterialLayerObject* FindLayer(const UDMMaterialStage* InBaseOrMask) const;
+
+	UFUNCTION(BlueprintPure, Category = "Material Designer", meta = (DisplayName = "Get Layers"))
+	TArray<UDMMaterialLayerObject*> BP_GetLayers() const;
+
+	const TArray<TObjectPtr<UDMMaterialLayerObject>>& GetLayers() const { return LayerObjects; }
+
+	UFUNCTION(BlueprintCallable, Category = "Material Designer")
+	UDMMaterialLayerObject* AddLayer(EDMMaterialPropertyType InMaterialProperty, UDMMaterialStage* InNewBase);
+
+	UFUNCTION(BlueprintCallable, Category = "Material Designer")
+	UDMMaterialLayerObject* AddLayerWithMask(EDMMaterialPropertyType InMaterialProperty, UDMMaterialStage* InNewBase, UDMMaterialStage* InNewMask);
+
+	// Adds the specified layer to the end of the layer list
+	UFUNCTION(BlueprintCallable, Category = "Material Designer")
+	bool PasteLayer(UDMMaterialLayerObject* InLayer);
+
+	UFUNCTION(BlueprintCallable, Category = "Material Designer")
+	bool RemoveLayer(UDMMaterialLayerObject* InLayer);
+
+	UFUNCTION(BlueprintCallable, Category = "Material Designer")
+	bool MoveLayer(UDMMaterialLayerObject* InLayer, int32 InNewIndex);
+
+	UFUNCTION(BlueprintCallable, Category = "Material Designer")
+	bool MoveLayerBefore(UDMMaterialLayerObject* InLayer, UDMMaterialLayerObject* InBeforeLayer = nullptr);
+
+	UFUNCTION(BlueprintCallable, Category = "Material Designer")
+	bool MoveLayerAfter(UDMMaterialLayerObject* InLayer, UDMMaterialLayerObject* InAfterLayer = nullptr);
+
+	UFUNCTION(BlueprintCallable, Category = "Material Designer")
+	UDMMaterialLayerObject* GetLastLayerForMaterialProperty(EDMMaterialPropertyType InMaterialProperty) const;
+
+	void UpdateOutputConnectorTypes();
+	void UpdateMaterialProperties();
+
+	FDMOnMaterialSlotConnectorsUpdated& GetOnConnectorsUpdateDelegate() { return OnConnectorsUpdateDelegate; }
+	FDMOnMaterialSlotPropertiesUpdated& GetOnPropertiesUpdateDelegate() { return OnPropertiesUpdateDelegate; }
+	FDMOnMaterialSlotLayersUpdated& GetOnLayersUpdateDelegate() { return OnLayersUpdateDelegate; }
+
+	void GenerateExpressions(const TSharedRef<FDMMaterialBuildState>& InBuildState) const;
+
+	const TMap<TWeakObjectPtr<UDMMaterialSlot>, int32>& GetSlotsReferencedBy() const { return SlotsReferencedBy; }
+
+	UFUNCTION(BlueprintCallable, Category = "Material Designer", Meta = (DisplayName="Get Slots Referenced By"))
+	TArray<UDMMaterialSlot*> K2_GetSlotsReferencedBy() const;
+
+	/** Returns true if a new associated is created */
+	bool ReferencedBySlot(UDMMaterialSlot* InOtherSlot);
+
+	/** Returns true if all associations have been removed */
+	bool UnreferencedBySlot(UDMMaterialSlot* InOtherSlot);
+
+	bool IsEditingLayers() const { return bIsEditingLayers; }
+	void SetEditingLayers(bool bInIsEditing) { bIsEditingLayers = bInIsEditing; }
+
+	UFUNCTION(BlueprintPure, Category = "Material Designer")
+	UMaterial* GetPreviewMaterial(EDMMaterialLayerStage InLayerStage);
+
+	UFUNCTION(BlueprintCallable, Category = "Material Designer")
+	bool SetLayerMaterialPropertyAndReplaceOthers(UDMMaterialLayerObject* InLayer, EDMMaterialPropertyType InMaterialProperty, 
+		EDMMaterialPropertyType InReplaceWithProperty);
+
+	//~ Begin UDMMaterialComponent
+	virtual void Update(EDMUpdateType InUpdateType) override;
+	virtual void DoClean() override;
+	virtual bool IsRootComponent() const override { return true; }
+	virtual FString GetComponentPathComponent() const override;
+	virtual UDMMaterialComponent* GetParentComponent() const override { return nullptr; }
+	virtual void PostEditorDuplicate(UDynamicMaterialModel* InMaterialModel, UDMMaterialComponent* InParent) override;
+	//~ End UDMMaterialComponent
+
+	//~ Begin UObject
+	virtual bool Modify(bool bInAlwaysMarkDirty = true) override;
+	virtual void PostEditUndo() override;
+	virtual void PostLoad() override;
+	//~ End UObject
+
+protected:
+	UPROPERTY(VisibleInstanceOnly, BlueprintReadOnly, Category = "Material Designer")
+	int32 Index;
+
+	UPROPERTY(VisibleInstanceOnly, BlueprintReadOnly, Category = "Material Designer")
+	TArray<TObjectPtr<UDMMaterialLayerObject>> LayerObjects;
+
+	UPROPERTY(VisibleInstanceOnly, BlueprintReadOnly, Category = "Material Designer")
+	TMap<EDMMaterialPropertyType, FDMMaterialSlotOutputConnectorTypes> OutputConnectorTypes;
+
+	UPROPERTY()
+	TMap<TWeakObjectPtr<UDMMaterialSlot>, int32> SlotsReferencedBy;
+
+	UPROPERTY(VisibleInstanceOnly, BlueprintReadOnly, Transient, DuplicateTransient, TextExportTransient, Category = "Material Designer")
+	bool bIsEditingLayers;
+
+	UPROPERTY(VisibleInstanceOnly, BlueprintReadOnly, Transient, DuplicateTransient, TextExportTransient, Category = "Material Designer")
+	TObjectPtr<UMaterial> BasePreviewMaterial;
+
+	UPROPERTY(VisibleInstanceOnly, BlueprintReadOnly, Transient, DuplicateTransient, TextExportTransient, Category = "Material Designer")
+	TObjectPtr<UMaterial> MaskPreviewMaterial;
+
+	FDMOnMaterialSlotConnectorsUpdated OnConnectorsUpdateDelegate;
+	FDMOnMaterialSlotPropertiesUpdated OnPropertiesUpdateDelegate;
+	FDMOnMaterialSlotLayersUpdated OnLayersUpdateDelegate;
+
+	void CreatePreviewMaterial(EDMMaterialLayerStage InLayerStage);
+	void UpdatePreviewMaterial(EDMMaterialLayerStage InLayerStage);
+	void UpdateBasePreviewMaterial(const TSharedRef<FDMMaterialBuildState>& InBuildState);
+	void UpdateBasePreviewMaterialProperty(const TSharedRef<FDMMaterialBuildState>& InBuildState, EDMMaterialPropertyType InBaseProperty);
+	void UpdateBasePreviewMaterialFull(const TSharedRef<FDMMaterialBuildState>& InBuildState);
+	void UpdateMaskPreviewMaterial(const TSharedRef<FDMMaterialBuildState>& InBuildState);
+	void UpdateMaskPreviewMaterialProperty(const TSharedRef<FDMMaterialBuildState>& InBuildState, EDMMaterialPropertyType InMaskProperty);
+	void UpdateMaskPreviewMaterialMaskCombination(const TSharedRef<FDMMaterialBuildState>& InBuildState, EDMMaterialPropertyType InMaskProperty);
+	void UpdatePreviewMaterialProperty(const TSharedRef<FDMMaterialBuildState>& InBuildState, UMaterial* InPreviewMaterial, EDMMaterialPropertyType InProperty);
+
+	//~ Begin UDMMaterialComponent
+	virtual void OnComponentAdded() override;
+	virtual void OnComponentRemoved() override;
+	virtual UDMMaterialComponent* GetSubComponentByPath(FDMComponentPath& InPath, const FDMComponentPathSegment& InPathSegment) const override;
+	//~ End UDMMaterialComponent
+
+private:
+	UE_DEPRECATED(5.4, "Promoted to full UObjects.")
+	UPROPERTY()
+	TArray<FDMMaterialLayer> Layers;
+
+	void ConvertDeprecatedLayers(TArray<FDMMaterialLayer>& InLayers);
+};
diff --git a/Engine/Plugins/Experimental/DynamicMaterial/Source/DynamicMaterialEditor/Public/Components/DMMaterialStage.h b/Engine/Plugins/Experimental/DynamicMaterial/Source/DynamicMaterialEditor/Public/Components/DMMaterialStage.h
new file mode 100644
--- /dev/null
+++ b/Engine/Plugins/Experimental/DynamicMaterial/Source/DynamicMaterialEditor/Public/Components/DMMaterialStage.h
@@ -0,0 +1,269 @@
+// Copyright Epic Games, Inc. All Rights Reserved.
+
+#pragma once
+
+#include "Components/DMMaterialComponent.h"
+#include "DMEDefs.h"
+#include "DMMaterialStage.generated.h"
+
+class FAssetThumbnailPool;
+class FMenuBuilder;
+class SWidget;
+class UDMMaterialLayerObject;
+class UDMMaterialSlot;
+class UDMMaterialStage;
+class UDMMaterialStageBlend;
+class UDMMaterialStageExpression;
+class UDMMaterialStageFunction;
+class UDMMaterialStageGradient;
+class UDMMaterialStageInput;
+class UDMMaterialStageInputExpression;
+class UDMMaterialStageInputFunction;
+class UDMMaterialStageInputGradient;
+class UDMMaterialStageInputSlot;
+class UDMMaterialStageInputTextureUV;
+class UDMMaterialStageInputValue;
+class UDMMaterialStageSource;
+class UDMMaterialStageThroughput;
+class UDMMaterialValue;
+class UDMTextureUV;
+class UMaterial;
+class UMaterialInstanceDynamic;
+class UMaterialInterface;
+enum class EDMMaterialLayerStage : uint8;
+struct FDMMaterialBuildState;
+struct FDMTextureUV;
+
+/**
+ * A node which handles a specific operation and manages its inputs and outputs.
+ */
+UCLASS(BlueprintType, ClassGroup = "Material Designer", meta = (DisplayName = "Material Designer Stage"))
+class DYNAMICMATERIALEDITOR_API UDMMaterialStage : public UDMMaterialComponent
+{
+	GENERATED_BODY()
+
+	friend class FDMThroughputPropertyRowGenerator;
+
+public:
+	static const FString SourcePathToken;
+	static const FString InputsPathToken;
+
+	UFUNCTION(BlueprintCallable, Category = "Material Designer")
+	static UDMMaterialStage* CreateMaterialStage(UDMMaterialLayerObject* InLayer = nullptr);
+
+	UFUNCTION(BlueprintPure, Category = "Material Designer")
+	UDMMaterialLayerObject* GetLayer() const;
+
+	UFUNCTION(BlueprintPure, Category = "Material Designer")
+	UDMMaterialStageSource* GetSource() const { return Source; }
+
+	UFUNCTION(BlueprintPure, Category = "Material Designer")
+	bool IsEnabled() const { return bEnabled; }
+
+	UFUNCTION(BlueprintCallable, Category = "Material Designer")
+	bool SetEnabled(bool bInEnabled);
+
+	UFUNCTION(BlueprintPure, Category = "Material Designer")
+	bool CanChangeSource() const { return bCanChangeSource; }
+
+	UFUNCTION(BlueprintCallable, Category = "Material Designer")
+	void SetCanChangeSource(bool bInCanChangeSource) { bCanChangeSource = bInCanChangeSource; }
+
+	UFUNCTION(BlueprintCallable, Category = "Material Designer")
+	void SetSource(UDMMaterialStageSource* InSource);
+
+	//~ Begin UDMMaterialComponent
+	virtual FText GetComponentDescription() const override;
+	//~ End UDMMaterialComponent
+
+	UFUNCTION(BlueprintPure, Category = "Material Designer")
+	const TArray<UDMMaterialStageInput*>& GetInputs() const { return Inputs; }
+
+	/** Determines what connects to what on this stage's Source. */
+	UFUNCTION(BlueprintPure, Category = "Material Designer")
+	const TArray<FDMMaterialStageConnection>& GetInputConnectionMap() const { return InputConnectionMap; }
+
+	TArray<FDMMaterialStageConnection>& GetInputConnectionMap() { return InputConnectionMap; }
+
+	UFUNCTION(BlueprintCallable, Category = "Material Designer")
+	EDMValueType GetSourceType(const FDMMaterialStageConnectorChannel& InChannel) const;
+
+	UFUNCTION(BlueprintCallable, Category = "Material Designer")
+	bool IsInputMapped(int32 InputIndex) const;
+
+	UFUNCTION(BlueprintPure, Category = "Material Designer")
+	virtual bool IsCompatibleWithPreviousStage(const UDMMaterialStage* InPreviousStage) const;
+
+	UFUNCTION(BlueprintPure, Category = "Material Designer")
+	virtual bool IsCompatibleWithNextStage(const UDMMaterialStage* InNextStage) const;
+
+	UFUNCTION(BlueprintCallable, Category = "Material Designer")
+	void AddInput(UDMMaterialStageInput* InNewInput);
+
+	UFUNCTION(BlueprintCallable, Category = "Material Designer")
+	void RemoveInput(UDMMaterialStageInput* InInput);
+
+	virtual void InputUpdated(UDMMaterialStageInput* InInput, EDMUpdateType InUpdateType);
+
+	UFUNCTION(BlueprintCallable, Category = "Material Designer")
+	virtual void ResetInputConnectionMap();
+
+	void GenerateExpressions(const TSharedRef<FDMMaterialBuildState>& InBuildState) const;
+
+	UFUNCTION(BlueprintPure, Category = "Material Designer")
+	bool IsBeingEdited() const { return bIsBeingEdited; }
+
+	UFUNCTION(BlueprintCallable, Category = "Material Designer")
+	bool SetBeingEdited(bool bInBeingEdited);
+
+	TMap<EDMMaterialPropertyType, UDMMaterialLayerObject*> GetPreviousStagesPropertyMap();
+	TMap<EDMMaterialPropertyType, UDMMaterialLayerObject*> GetPropertyMap();
+
+	UFUNCTION(BlueprintCallable, Category = "Material Designer")
+	UDMMaterialStageInputValue* ChangeSource_NewLocalValue(EDMValueType InType);
+
+	UFUNCTION(BlueprintCallable, Category = "Material Designer")
+	UDMMaterialStageInputValue* ChangeSource_Value(UDMMaterialValue* InValue);
+
+	UFUNCTION(BlueprintCallable, Category = "Material Designer")
+	UDMMaterialStageInputValue* ChangeSource_NewValue(EDMValueType InType);
+
+	UFUNCTION(BlueprintCallable, Category = "Material Designer")
+	UDMMaterialStageInputSlot* ChangeSource_Slot(UDMMaterialSlot* InSlot, EDMMaterialPropertyType InProperty);
+
+	UFUNCTION(BlueprintCallable, Category = "Material Designer")
+	UDMMaterialStageExpression* ChangeSource_Expression(TSubclassOf<UDMMaterialStageExpression> InExpressionClass);
+
+	UFUNCTION(BlueprintCallable, Category = "Material Designer")
+	UDMMaterialStageBlend* ChangeSource_Blend(TSubclassOf<UDMMaterialStageBlend> InBlendClass);
+
+	UFUNCTION(BlueprintCallable, Category = "Material Designer")
+	UDMMaterialStageGradient* ChangeSource_Gradient(TSubclassOf<UDMMaterialStageGradient> InGradientClass);
+
+	UFUNCTION(BlueprintCallable, Category = "Material Designer")
+	UDMMaterialStageInputTextureUV* ChangeSource_UV(bool bInDoUpdate);
+
+	UFUNCTION(BlueprintCallable, Category = "Material Designer")
+	UDMMaterialStageFunction* ChangeSource_MaterialFunction();
+
+	UFUNCTION(BlueprintCallable, Category = "Material Designer")
+	UDMMaterialStageSource* ChangeInput_PreviousStage(int32 InInputIdx, int32 InInputChannel, EDMMaterialPropertyType InPreviousStageProperty, 
+		int32 InOutputIdx, int32 InOutputChannel);
+
+	UFUNCTION(BlueprintCallable, Category = "Material Designer")
+	UDMMaterialStageInputValue* ChangeInput_NewLocalValue(int32 InInputIdx, int32 InInputChannel, EDMValueType InType, int32 InOutputChannel);
+
+	UFUNCTION(BlueprintCallable, Category = "Material Designer")
+	UDMMaterialStageInputValue* ChangeInput_Value(int32 InInputIdx, int32 InInputChannel, UDMMaterialValue* InValue, int32 InOutputChannel);
+
+	UFUNCTION(BlueprintCallable, Category = "Material Designer")
+	UDMMaterialStageInputValue* ChangeInput_NewValue(int32 InInputIdx, int32 InInputChannel, EDMValueType InType, int32 InOutputChannel);
+
+	UFUNCTION(BlueprintCallable, Category = "Material Designer")
+	UDMMaterialStageInputSlot* ChangeInput_Slot(int32 InInputIdx, int32 InInputChannel, UDMMaterialSlot* InSlot, EDMMaterialPropertyType InProperty, 
+		int32 InOutputIdx, int32 InOutputChannel);
+
+	UFUNCTION(BlueprintCallable, Category = "Material Designer")
+	UDMMaterialStageInputExpression* ChangeInput_Expression(int32 InInputIdx, int32 InInputChannel, TSubclassOf<UDMMaterialStageExpression> InExpressionClass, 
+		int32 InOutputIdx, int32 InOutputChannel);
+
+	UFUNCTION(BlueprintCallable, Category = "Material Designer")
+	UDMMaterialStageInputTextureUV* ChangeInput_UV(int32 InInputIdx, int32 InInputChannel, int32 InOutputChannel);
+
+	UFUNCTION(BlueprintCallable, Category = "Material Designer")
+	UDMMaterialStageInputGradient* ChangeInput_Gradient(int32 InInputIdx, int32 InInputChannel, TSubclassOf<UDMMaterialStageGradient> InGradientClass, int32 InOutputChannel);
+
+	UFUNCTION(BlueprintCallable, Category = "Material Designer")
+	UDMMaterialStageInputFunction* ChangeInput_MaterialFunction(int32 InInputIdx, int32 InInputChannel);
+
+	UFUNCTION(BlueprintCallable, Category = "Material Designer")
+	void RemoveUnusedInputs();
+
+	/** Returns true if any changes were made */
+	UFUNCTION(BlueprintCallable, Category = "Material Designer")
+	bool VerifyAllInputMaps();
+
+	/** Returns true if any changes were made */
+	UFUNCTION(BlueprintCallable, Category = "Material Designer")
+	bool VerifyInputMap(int32 InInputIdx);
+
+	UFUNCTION(BlueprintPure, Category = "Material Designer")
+	UMaterialInterface* GetPreviewMaterial();
+
+	const FDMMaterialStageConnectorChannel* FindInputChannel(UDMMaterialStageInput* InStageInput);
+
+	void UpdateInputMap(int32 InInputIdx, int32 InSourceIndex, int32 InInputChannel, int32 InOutputIdx, int32 InOutputChannel, EDMMaterialPropertyType InStageProperty);
+
+	UFUNCTION(BlueprintCallable, Category = "Material Designer")
+	int32 FindIndex() const;
+
+	UFUNCTION(BlueprintCallable, Category = "Material Designer")
+	UDMMaterialStage* GetPreviousStage() const;
+
+	UFUNCTION(BlueprintCallable, Category = "Material Designer")
+	UDMMaterialStage* GetNextStage() const;
+
+	UFUNCTION(BlueprintPure, Category = "Material Designer")
+	virtual bool IsRootStage() const;
+
+	//~ Begin UDMMaterialComponent
+	virtual void Update(EDMUpdateType InUpdateType) override;
+	virtual void DoClean() override;
+	virtual FString GetComponentPathComponent() const override;
+	virtual UDMMaterialComponent* GetParentComponent() const override;
+	virtual void PostEditorDuplicate(UDynamicMaterialModel* InMaterialModel, UDMMaterialComponent* InParent) override;
+	//~ End UDMMaterialComponent
+
+	//~ Begin UObject
+	virtual bool Modify(bool bInAlwaysMarkDirty = true) override;
+	virtual void PostEditUndo() override;
+	virtual void PostLoad() override;
+	virtual void PostEditImport() override;
+	//~ End UObject
+
+protected:
+	UDMMaterialStage();
+
+	UPROPERTY(VisibleInstanceOnly, BlueprintReadOnly, Instanced, Category = "Material Designer")
+	TObjectPtr<UDMMaterialStageSource> Source;
+
+	UPROPERTY(VisibleInstanceOnly, BlueprintReadOnly, Instanced, Category = "Material Designer")
+	TArray<TObjectPtr<UDMMaterialStageInput>> Inputs;
+
+	/** How our inputs connect to the inputs of this stage's source */
+	UPROPERTY(VisibleInstanceOnly, BlueprintReadOnly, Category = "Material Designer")
+	TArray<FDMMaterialStageConnection> InputConnectionMap;
+
+	UPROPERTY(VisibleInstanceOnly, BlueprintReadOnly, Category = "Material Designer")
+	bool bEnabled;
+
+	UPROPERTY(VisibleInstanceOnly, BlueprintReadOnly, Category = "Material Designer")
+	bool bCanChangeSource;
+
+	UPROPERTY(VisibleInstanceOnly, BlueprintReadOnly, Transient, DuplicateTransient, TextExportTransient, Category = "Material Designer")
+	bool bIsBeingEdited;
+
+	UPROPERTY(VisibleInstanceOnly, BlueprintReadOnly, Transient, DuplicateTransient, TextExportTransient, Category = "Material Designer")
+	TObjectPtr<UMaterial> PreviewMaterialBase;
+
+	UPROPERTY(VisibleInstanceOnly, BlueprintReadOnly, Transient, DuplicateTransient, TextExportTransient, Category = "Material Designer")
+	TObjectPtr<UMaterialInstanceDynamic> PreviewMaterialDynamic;
+
+	void CreatePreviewMaterial();
+	void UpdatePreviewMaterial();
+
+	//~ Begin UDMMaterialComponent
+	virtual void OnComponentAdded() override;
+	virtual void OnComponentRemoved() override;
+	virtual void GetComponentPathInternal(TArray<FString>& OutChildComponentPathComponents) const override;
+	virtual UDMMaterialComponent* GetSubComponentByPath(FDMComponentPath& InPath, const FDMComponentPathSegment& InPathSegment) const override;
+	//~ End UDMMaterialComponent
+
+	virtual void AddDelegates();
+	virtual void RemoveDelegates();
+
+	void OnValueUpdated(UDynamicMaterialModel* InMaterialModel, UDMMaterialValue* InValue);
+	void OnTextureUVUpdated(UDynamicMaterialModel* InMaterialModel, UDMTextureUV* InTextureUV);
+
+	void ApplyWholeLayerValue(UDMMaterialValue* InValue);
+};
diff --git a/Engine/Plugins/Experimental/DynamicMaterial/Source/DynamicMaterialEditor/Public/Components/DMMaterialStageBlend.h b/Engine/Plugins/Experimental/DynamicMaterial/Source/DynamicMaterialEditor/Public/Components/DMMaterialStageBlend.h
new file mode 100644
--- /dev/null
+++ b/Engine/Plugins/Experimental/DynamicMaterial/Source/DynamicMaterialEditor/Public/Components/DMMaterialStageBlend.h
@@ -0,0 +1,65 @@
+// Copyright Epic Games, Inc. All Rights Reserved.
+
+#pragma once
+
+#include "DMMaterialStageThroughput.h"
+#include "UObject/StrongObjectPtr.h"
+#include "DMMaterialStageBlend.generated.h"
+
+class FMenuBuilder;
+class UDMMaterialLayerObject;
+class UDMMaterialStageInput;
+class UDMMaterialValueFloat1;
+class UMaterial;
+struct FDMMaterialBuildState;
+
+/**
+ * A node which represents a blend operation.
+ */
+UCLASS(Abstract, BlueprintType, ClassGroup = "Material Designer", meta = (DisplayName = "Material Designer Stage Blend"))
+class DYNAMICMATERIALEDITOR_API UDMMaterialStageBlend : public UDMMaterialStageThroughput
+{
+	GENERATED_BODY()
+
+public:
+	static constexpr int32 InputAlpha = 0;
+	static constexpr int32 InputA = 1;
+	static constexpr int32 InputB = 2;
+
+	static UDMMaterialStage* CreateStage(TSubclassOf<UDMMaterialStageBlend> InMaterialStageBlendClass, UDMMaterialLayerObject* InLayer = nullptr);
+
+	static const TArray<TStrongObjectPtr<UClass>>& GetAvailableBlends();
+
+	UDMMaterialValueFloat1* GetInputAlpha() const;
+	UDMMaterialStageInput* GetInputB() const;
+
+	//~ Begin UDMMaterialStageThroughput
+	virtual bool CanInputAcceptType(int32 InputIndex, EDMValueType ValueType) const override;
+	virtual void AddDefaultInput(int32 InInputIndex) const override;
+	virtual bool CanChangeInput(int32 InputIndex) const override;
+	virtual bool CanChangeInputType(int32 InputIndex) const override;
+	virtual bool IsInputVisible(int32 InputIndex) const override;
+	virtual int32 ResolveInput(const TSharedRef<FDMMaterialBuildState>& InBuildState, int32 InputIndex, FDMMaterialStageConnectorChannel& OutChannel,
+		TArray<UMaterialExpression*>& OutExpressions) const override;
+	//~ End UDMMaterialStageThroughput
+
+	//~ Begin UDMMaterialStageSource
+	virtual FText GetStageDescription() const override;
+	virtual bool SupportsLayerMaskTextureUVLink() const override { return true; }
+	virtual FDMExpressionInput GetLayerMaskLinkTextureUVInputExpressions(const TSharedRef<FDMMaterialBuildState>& InBuildState) const;
+	virtual void GetMaskAlphaBlendNode(const TSharedRef<FDMMaterialBuildState>& InBuildState, UMaterialExpression*& OutExpression,
+		int32& OutOutputIndex, int32& OutOutputChannel) const override;
+	virtual bool UpdateStagePreviewMaterial(UDMMaterialStage* InStage, UMaterial* InPreviewMaterial, UMaterialExpression*& OutMaterialExpression,
+		int32& OutputIndex) override;
+	//~ End UDMMaterialStageSource
+
+protected:
+	static TArray<TStrongObjectPtr<UClass>> Blends;
+
+	static void GenerateBlendList();
+
+	UDMMaterialStageBlend();
+	UDMMaterialStageBlend(const FText& InName);
+
+	virtual void UpdatePreviewMaterial(UMaterial* InPreviewMaterial = nullptr) override;
+};
\ No newline at end of file
diff --git a/Engine/Plugins/Experimental/DynamicMaterial/Source/DynamicMaterialEditor/Public/Components/DMMaterialStageBlendFunction.h b/Engine/Plugins/Experimental/DynamicMaterial/Source/DynamicMaterialEditor/Public/Components/DMMaterialStageBlendFunction.h
new file mode 100644
--- /dev/null
+++ b/Engine/Plugins/Experimental/DynamicMaterial/Source/DynamicMaterialEditor/Public/Components/DMMaterialStageBlendFunction.h
@@ -0,0 +1,29 @@
+// Copyright Epic Games, Inc. All Rights Reserved.
+
+#pragma once
+
+#include "Components/DMMaterialStageBlend.h"
+#include "DMMaterialStageBlendFunction.generated.h"
+
+class UMaterialExpression;
+class UMaterialFunctionInterface;
+struct FDMMaterialBuildState;
+
+UCLASS(Abstract, BlueprintType, ClassGroup = "Material Designer")
+class DYNAMICMATERIALEDITOR_API UDMMaterialStageBlendFunction : public UDMMaterialStageBlend
+{
+	GENERATED_BODY()
+
+public:
+	virtual void GenerateExpressions(const TSharedRef<FDMMaterialBuildState>& InBuildState) const override;
+	virtual void ConnectOutputToInput(const TSharedRef<FDMMaterialBuildState>& InBuildState, int32 InInputIndex, UMaterialExpression* InSourceExpression,
+		int32 InSourceOutputIndex, int32 InSourceOutputChannel) override;
+
+protected:
+	UPROPERTY()
+	TObjectPtr<UMaterialFunctionInterface> MaterialFunction;
+
+	UDMMaterialStageBlendFunction();
+	UDMMaterialStageBlendFunction(const FText& InName, UMaterialFunctionInterface* InMaterialFunction);
+	UDMMaterialStageBlendFunction(const FText& InName, const FName& InFunctionName, const FString& InFunctionPath);
+};
diff --git a/Engine/Plugins/Experimental/DynamicMaterial/Source/DynamicMaterialEditor/Public/Components/DMMaterialStageFunction.h b/Engine/Plugins/Experimental/DynamicMaterial/Source/DynamicMaterialEditor/Public/Components/DMMaterialStageFunction.h
new file mode 100644
--- /dev/null
+++ b/Engine/Plugins/Experimental/DynamicMaterial/Source/DynamicMaterialEditor/Public/Components/DMMaterialStageFunction.h
@@ -0,0 +1,56 @@
+// Copyright Epic Games, Inc. All Rights Reserved.
+
+#pragma once
+
+#include "Components/DMMaterialStageThroughput.h"
+#include "DMMaterialStageFunction.generated.h"
+
+class UDMMaterialLayerObject;
+class UMaterialFunctionInterface;
+
+UCLASS(BlueprintType, ClassGroup = "Material Designer")
+class DYNAMICMATERIALEDITOR_API UDMMaterialStageFunction : public UDMMaterialStageThroughput
+{
+	GENERATED_BODY()
+
+public:
+	static constexpr int32 InputPreviousStage = 0;
+
+	static UDMMaterialStage* CreateStage(UDMMaterialLayerObject* InLayer = nullptr);
+
+	UFUNCTION(BlueprintPure, Category = "Material Designer")
+	UMaterialFunctionInterface* GetMaterialFunction() const { return MaterialFunction.Get(); }
+
+	UFUNCTION(BlueprintCallable, Category = "Material Designer")
+	void SetMaterialFunction(UMaterialFunctionInterface* InMaterialFunction);
+
+	//~ Begin UDMMaterialStageThroughput
+	virtual void AddDefaultInput(int32 InInputIndex) const override;
+	virtual bool CanChangeInput(int32 InputIndex) const override;
+	virtual bool CanChangeInputType(int32 InputIndex) const override;
+	virtual bool IsInputVisible(int32 InputIndex) const override;
+	//~ End UDMMaterialStageThroughput
+
+	//~ Begin UDMMaterialStageSource
+	virtual void GenerateExpressions(const TSharedRef<FDMMaterialBuildState>& InBuildState) const override;
+	//~ End UDMMaterialStageSource
+
+	//~ Begin UObject
+	virtual void PreEditChange(FEditPropertyChain& PropertyAboutToChange) override;
+	virtual void PostEditChangeProperty(FPropertyChangedEvent& PropertyChangedEvent) override;
+	//~ End UObject
+
+protected:
+	static TSoftObjectPtr<UMaterialFunctionInterface> NoOp;
+
+	UPROPERTY(EditAnywhere, BlueprintReadWrite, Getter = GetMaterialFunction, Setter = SetMaterialFunction, BlueprintSetter = SetMaterialFunction, Category = "Material Designer",
+		meta = (DisplayThumbnail = true, AllowPrivateAccess = "true", HighPriority, NotKeyframeable))
+	TObjectPtr<UMaterialFunctionInterface> MaterialFunction;
+
+	UPROPERTY(Transient, DuplicateTransient, TextExportTransient)
+	TObjectPtr<UMaterialFunctionInterface> MaterialFunction_PreEdit;
+
+	UDMMaterialStageFunction();
+
+	void OnMaterialFunctionChanged();
+};
diff --git a/Engine/Plugins/Experimental/DynamicMaterial/Source/DynamicMaterialEditor/Public/Components/DMMaterialStageGradient.h b/Engine/Plugins/Experimental/DynamicMaterial/Source/DynamicMaterialEditor/Public/Components/DMMaterialStageGradient.h
new file mode 100644
--- /dev/null
+++ b/Engine/Plugins/Experimental/DynamicMaterial/Source/DynamicMaterialEditor/Public/Components/DMMaterialStageGradient.h
@@ -0,0 +1,44 @@
+// Copyright Epic Games, Inc. All Rights Reserved.
+
+#pragma once
+
+#include "DMMaterialStageThroughput.h"
+#include "UObject/StrongObjectPtr.h"
+#include "DMMaterialStageGradient.generated.h"
+
+class FMenuBuilder;
+class UDMMaterialLayerObject;
+class UDMMaterialStageInput;
+struct FDMMaterialBuildState;
+
+/**
+ * A node which represents UV-based gradient.
+ */
+UCLASS(Abstract, BlueprintType, ClassGroup = "Material Designer", meta = (DisplayName = "Material Designer Stage Gradient"))
+class DYNAMICMATERIALEDITOR_API UDMMaterialStageGradient : public UDMMaterialStageThroughput
+{
+	GENERATED_BODY()
+
+public:
+	static UDMMaterialStage* CreateStage(TSubclassOf<UDMMaterialStageGradient> InMaterialStageGradientClass, UDMMaterialLayerObject* InLayer = nullptr);
+
+	static const TArray<TStrongObjectPtr<UClass>>& GetAvailableGradients();
+
+	//~ Begin UDMMaterialStageThroughput
+	virtual bool CanChangeInputType(int32 InputIndex) const override;
+	//~ End UDMMaterialStageThroughput
+	
+	//~ Begin UDMMaterialStageSource
+	virtual bool SupportsLayerMaskTextureUVLink() const override { return true; }
+	virtual int32 GetLayerMaskTextureUVLinkInputIndex() const override { return 0; }
+	//~ End UDMMaterialStageSource
+	
+protected:
+	static TArray<TStrongObjectPtr<UClass>> Gradients;
+
+	static void GenerateGradientList();
+
+	UDMMaterialStageGradient();
+	UDMMaterialStageGradient(const FText& InName);
+
+};
\ No newline at end of file
diff --git a/Engine/Plugins/Experimental/DynamicMaterial/Source/DynamicMaterialEditor/Public/Components/DMMaterialStageInput.h b/Engine/Plugins/Experimental/DynamicMaterial/Source/DynamicMaterialEditor/Public/Components/DMMaterialStageInput.h
new file mode 100644
--- /dev/null
+++ b/Engine/Plugins/Experimental/DynamicMaterial/Source/DynamicMaterialEditor/Public/Components/DMMaterialStageInput.h
@@ -0,0 +1,45 @@
+// Copyright Epic Games, Inc. All Rights Reserved.
+
+#pragma once
+
+#include "DMMaterialStageSource.h"
+#include "DMMaterialStageInput.generated.h"
+
+class UDMMaterialSlot;
+class UMaterial;
+struct FDMMaterialStageConnectorChannel;
+
+/**
+ * A node which produces an output (e.g. Texture coordinate.)
+ */
+UCLASS(Abstract, BlueprintType, ClassGroup = "Material Designer", meta = (DisplayName = "Material Designer Stage Input"))
+class DYNAMICMATERIALEDITOR_API UDMMaterialStageInput : public UDMMaterialStageSource
+{
+	GENERATED_BODY()
+
+public:
+	static const FString StageInputPrefixStr;
+
+	virtual FText GetChannelDescription(const FDMMaterialStageConnectorChannel& Channel)
+		PURE_VIRTUAL(UDMMaterialStageInput::GetChannelDescription, return FText::GetEmpty();)
+		
+	virtual void Update(EDMUpdateType InUpdateType) override;
+
+	//~ Begin UDMMaterialComponent
+	virtual FString GetComponentPathComponent() const override;
+	//~ End UDMMaterialComponent
+
+protected:
+	UDMMaterialStageInput()
+	{
+	}
+
+	//~ Begin UDMMaterialStageSource
+	virtual void UpdateOutputConnectors() { }
+	virtual void UpdatePreviewMaterial(UMaterial* InPreviewMaterial = nullptr) override;
+	//~ End UDMMaterialStageSource
+
+	//~ Begin UDMMaterialComponent
+	virtual void GetComponentPathInternal(TArray<FString>& OutChildComponentPathComponents) const override;
+	//~ End UDMMaterialComponent
+};
diff --git a/Engine/Plugins/Experimental/DynamicMaterial/Source/DynamicMaterialEditor/Public/Components/DMMaterialStageSource.h b/Engine/Plugins/Experimental/DynamicMaterial/Source/DynamicMaterialEditor/Public/Components/DMMaterialStageSource.h
new file mode 100644
--- /dev/null
+++ b/Engine/Plugins/Experimental/DynamicMaterial/Source/DynamicMaterialEditor/Public/Components/DMMaterialStageSource.h
@@ -0,0 +1,101 @@
+// Copyright Epic Games, Inc. All Rights Reserved.
+
+#pragma once
+
+#include "Components/DMMaterialComponent.h"
+#include "DMEDefs.h"
+#include "UObject/StrongObjectPtr.h"
+#include "DMMaterialStageSource.generated.h"
+
+class UDMMaterialSlot;
+class UDMMaterialStageInputTextureUV;
+class UMaterialExpression;
+struct FDMMaterialBuildState;
+
+/**
+ * A node which produces an output.
+ */
+UCLASS(Abstract, BlueprintType, ClassGroup = "Material Designer", meta = (DisplayName = "Material Designer Stage Source"))
+class DYNAMICMATERIALEDITOR_API UDMMaterialStageSource : public UDMMaterialComponent
+{
+	GENERATED_BODY()
+
+	friend class UDMMaterialStage;
+
+public:
+	static const TArray<TStrongObjectPtr<UClass>>& GetAvailableSourceClasses();
+
+	UFUNCTION(BlueprintPure, Category = "Material Designer")
+	UDMMaterialStage* GetStage() const;
+
+	/* Returns a description of the stage for which this is the source. */
+	UFUNCTION(BlueprintPure, Category = "Material Designer")
+	virtual FText GetStageDescription() const { return GetComponentDescription(); }
+
+	UFUNCTION(BlueprintPure, Category = "Material Designer")
+	const TArray<FDMMaterialStageConnector>& GetOutputConnectors() const { return OutputConnectors; }
+
+	virtual void GetMaskAlphaBlendNode(const TSharedRef<FDMMaterialBuildState>& InBuildState, UMaterialExpression*& OutExpression, 
+		int32& OutOutputIndex, int32& OutOutputChannel) const;
+
+	virtual void GenerateExpressions(const TSharedRef<FDMMaterialBuildState>& InBuildState) const
+		PURE_VIRTUAL(UDMMaterialStageSource::GenerateExpressions)
+
+	virtual void AddExpressionProperties(const TArray<UMaterialExpression*>& Expressions) const {}
+
+	UFUNCTION(BlueprintPure, Category = "Material Designer")
+	UMaterial* GetPreviewMaterial();
+
+	void UpdateOntoPreviewMaterial(UMaterial* ExternalPreviewMaterial);
+
+	/** 
+	 * Returns the output index (channel WHOLE_CHANNEL) if this expression has pre-masked outputs.
+	 * Returns INDEX_NONE if it is not supported.
+	 */
+	virtual int32 GetInnateMaskOutput(int32 OutputIndex, int32 OutputChannels) const;
+
+	/**
+	 * Given an output index, may return an override for output channels on that output.
+	 * E.g. The texture sample alpha output may override to FOURTH_CHANNEL.
+	 * Returns INDEX_NONE with no override.
+	 */
+	virtual int32 GetOutputChannelOverride(int32 InOutputIndex) const { return INDEX_NONE; }
+
+	virtual bool UpdateStagePreviewMaterial(UDMMaterialStage* InStage, UMaterial* InPreviewMaterial, UMaterialExpression*& OutMaterialExpression, 
+		int32& OutputIndex);
+
+	//~ Begin FNotifyHook
+	virtual void NotifyPostChange(const FPropertyChangedEvent& InPropertyChangedEvent, class FEditPropertyChain* InPropertyThatChanged) override;
+	//~ End FNotifyHook
+
+	//~ Begin UDMMaterialComponent
+	virtual void Update(EDMUpdateType InUpdateType) override;
+	virtual void DoClean() override;
+	virtual UDMMaterialComponent* GetParentComponent() const override;
+	virtual void PostEditorDuplicate(UDynamicMaterialModel* InMaterialModel, UDMMaterialComponent* InParent) override;
+	//~ End UDMMaterialComponent
+
+	//~ Begin UObject
+	virtual void PostEditUndo() override;
+	//~ End UObject
+
+protected:
+	static TArray<TStrongObjectPtr<UClass>> SourceClasses;
+
+	static void GenerateClassList();
+
+	UPROPERTY(VisibleInstanceOnly, BlueprintReadOnly, Category = "Material Designer")
+	TArray<FDMMaterialStageConnector> OutputConnectors;
+
+	UPROPERTY(VisibleInstanceOnly, BlueprintReadOnly, Transient, DuplicateTransient, TextExportTransient, Category = "Material Designer")
+	TObjectPtr<UMaterial> PreviewMaterial;
+
+	UDMMaterialStageSource();
+
+	void CreatePreviewMaterial();
+	virtual void UpdatePreviewMaterial(UMaterial* InPreviewMaterial = nullptr);
+
+	//~ Begin UDMMaterialComponent
+	virtual void OnComponentAdded() override;
+	//~ End UDMMaterialComponent
+};
diff --git a/Engine/Plugins/Experimental/DynamicMaterial/Source/DynamicMaterialEditor/Public/Components/DMMaterialStageThroughput.h b/Engine/Plugins/Experimental/DynamicMaterial/Source/DynamicMaterialEditor/Public/Components/DMMaterialStageThroughput.h
new file mode 100644
--- /dev/null
+++ b/Engine/Plugins/Experimental/DynamicMaterial/Source/DynamicMaterialEditor/Public/Components/DMMaterialStageThroughput.h
@@ -0,0 +1,133 @@
+// Copyright Epic Games, Inc. All Rights Reserved.
+
+#pragma once
+
+#include "DMMaterialStageSource.h"
+#include "UObject/StrongObjectPtr.h"
+#include "DMMaterialStageThroughput.generated.h"
+
+class FMenuBuilder;
+class UDMMaterialStageInput;
+class UMaterialExpression;
+struct FDMExpressionInput;
+struct FDMMaterialBuildState;
+
+struct FDMExpressionInput
+{
+	TArray<UMaterialExpression*> OutputExpressions = {};
+	int32 OutputIndex = INDEX_NONE;
+	int32 OutputChannel = INDEX_NONE;
+
+	bool IsValid()
+	{
+		return OutputExpressions.IsEmpty() == false && OutputIndex != INDEX_NONE && OutputChannel != INDEX_NONE;
+	}
+};
+
+/**
+ * A node which take one or more inputs and produces an output (e.g. Multiply)
+ */
+UCLASS(Abstract, BlueprintType, ClassGroup = "Material Designer", meta = (DisplayName = "Material Designer Stage Throughput"))
+class DYNAMICMATERIALEDITOR_API UDMMaterialStageThroughput : public UDMMaterialStageSource
+{
+	GENERATED_BODY()
+
+public:
+	static const TArray<TStrongObjectPtr<UClass>>& GetAvailableThroughputs();
+
+	virtual FText GetComponentDescription() const override { return GetDescription(); }
+
+	UFUNCTION(BlueprintPure, Category = "Material Designer")
+	const FText& GetDescription() const { return Name; }
+
+	UFUNCTION(BlueprintPure, Category = "Material Designer")
+	bool IsInputRequired() const { return bInputRequired; }
+
+	UFUNCTION(BlueprintPure, Category = "Material Designer")
+	bool AllowsNestedInputs() const { return bAllowNestedInputs; }
+
+	UFUNCTION(BlueprintPure, Category = "Material Designer")
+	const TArray<FDMMaterialStageConnector>& GetInputConnectors() const { return InputConnectors; }
+
+	UFUNCTION(BlueprintCallable, Category = "Material Designer")
+	virtual bool CanInputAcceptType(int32 InputIndex, EDMValueType ValueType) const;
+
+	UFUNCTION(BlueprintCallable, Category = "Material Designer")
+	virtual bool CanInputConnectTo(int32 InputIndex, const FDMMaterialStageConnector& OutputConnector, int32 OutputChannel,
+		bool bCheckSingleFloat = false);
+
+	virtual bool CanChangeInput(int32 InputIndex) const;
+	virtual bool CanChangeInputType(int32 InputIndex) const;
+
+	virtual bool IsInputVisible(int32 InputIndex) const;
+
+	virtual void ConnectOutputToInput(const TSharedRef<FDMMaterialBuildState>& InBuildState, int32 InInputIndex, UMaterialExpression* InSourceExpression,
+		int32 InSourceOutputIndex, int32 InSourceOutputChannel);
+
+	/** Returns true if the layer and mask can have their Texture UV linked. */
+	virtual bool SupportsLayerMaskTextureUVLink() const { return false; }
+
+	/** Returns the input index for the default implementation of the below method. */
+	virtual int32 GetLayerMaskTextureUVLinkInputIndex() const;
+
+	/** 
+	 * Returns all the material nodes requires to create this node's Texture UV input. 
+	 * If you override this method, you do not need to override GetLayerMaskTextureUVLinkInputIndex.
+	 */
+	virtual FDMExpressionInput GetLayerMaskLinkTextureUVInputExpressions(const TSharedRef<FDMMaterialBuildState>& InBuildState) const;
+
+	/**
+	 * Override this to redirect inputs to other nodes.
+	 * Returns the first node in the array by default
+	 * --> In [ ]-[ ]-[ ] Out -->
+	 */
+	virtual UMaterialExpression* GetExpressionForInput(const TArray<UMaterialExpression*>& StageSourceExpressions, int32 InputIdx);
+
+	virtual void AddDefaultInput(int32 InInputIndex) const;
+
+	/** Returns the actual output index of the material expression */
+	virtual int32 ResolveInput(const TSharedRef<FDMMaterialBuildState>& InBuildState, int32 InputIndex, FDMMaterialStageConnectorChannel& OutChannel,
+		TArray<UMaterialExpression*>& OutExpressions) const;
+
+	virtual int32 ResolveLayerMaskTextureUVLinkInput(const TSharedRef<FDMMaterialBuildState>& InBuildState, int32 InputIndex, 
+		FDMMaterialStageConnectorChannel& OutChannel, TArray<UMaterialExpression*>& OutExpressions) const;
+
+	virtual void InputUpdated(int32 InInputIndex, EDMUpdateType InUpdateType) { }
+
+protected:
+	static TArray<TStrongObjectPtr<UClass>> Throughputs;
+
+	static void GenerateThroughputList();
+
+	static int32 ResolveLayerMaskTextureUVLinkInputImpl(const TSharedRef<FDMMaterialBuildState>& InBuildState, const UDMMaterialStageSource* StageSource,
+		FDMMaterialStageConnectorChannel& OutChannel, TArray<UMaterialExpression*>& OutExpressions);
+
+	UDMMaterialStageThroughput();
+	UDMMaterialStageThroughput(const FText& InName);
+
+	UPROPERTY(VisibleInstanceOnly, BlueprintReadOnly, Category = "Material Designer")
+	FText Name;
+
+	UPROPERTY(VisibleInstanceOnly, BlueprintReadOnly, Category = "Material Designer")
+	bool bInputRequired;
+
+	UPROPERTY(VisibleInstanceOnly, BlueprintReadOnly, Category = "Material Designer")
+	bool bAllowNestedInputs;
+
+	UPROPERTY(VisibleInstanceOnly, BlueprintReadOnly, Category = "Material Designer")
+	TArray<FDMMaterialStageConnector> InputConnectors;
+
+	virtual bool ShouldKeepInput(int32 InInputIdx);
+
+	void ConnectOutputToInput_Internal(const TSharedRef<FDMMaterialBuildState>& InBuildState, UMaterialExpression* TargetExpression,
+		int32 InputIndex, UMaterialExpression* SourceExpression, int32 SourceOutputIndex, int32 SourceOutputChannel) const;
+
+	virtual int32 ResolveInputChannel(const TSharedRef<FDMMaterialBuildState>& InBuildState, int32 InputIndex, int32 ChannelIndex, 
+		FDMMaterialStageConnectorChannel& OutChannel, TArray<UMaterialExpression*>& OutExpressions) const;
+
+	virtual void UpdatePreviewMaterial(UMaterial* InPreviewMaterial = nullptr);
+
+	//~ Begin UDMMaterialComponent
+	virtual void OnComponentAdded() override;
+	//~ End UDMMaterialComponent
+};
diff --git a/Engine/Plugins/Experimental/DynamicMaterial/Source/DynamicMaterialEditor/Public/Components/DMMaterialStageThroughputLayerBlend.h b/Engine/Plugins/Experimental/DynamicMaterial/Source/DynamicMaterialEditor/Public/Components/DMMaterialStageThroughputLayerBlend.h
new file mode 100644
--- /dev/null
+++ b/Engine/Plugins/Experimental/DynamicMaterial/Source/DynamicMaterialEditor/Public/Components/DMMaterialStageThroughputLayerBlend.h
@@ -0,0 +1,131 @@
+// Copyright Epic Games, Inc. All Rights Reserved.
+
+#pragma once
+
+#include "Components/DMMaterialStageThroughput.h"
+#include "DMMaterialStageThroughputLayerBlend.generated.h"
+
+class UDMMaterialLayerObject;
+class UMaterial;
+class UMaterialExpression;
+
+UENUM(BlueprintType)
+enum class EAvaColorChannel : uint8
+{
+	None =  0,
+	Red =   1,
+	Green = 2,
+	Blue =  3,
+	Alpha = 4
+};
+
+UCLASS(BlueprintType, ClassGroup = "Material Designer")
+class DYNAMICMATERIALEDITOR_API UDMMaterialStageThroughputLayerBlend : public UDMMaterialStageThroughput
+{
+	GENERATED_BODY()
+
+public:
+	static constexpr int32 InputPreviousLayer = 0;
+	static constexpr int32 InputBaseStage = 1;
+	static constexpr int32 InputMaskSource = 2;
+
+	static UDMMaterialStage* CreateStage(UDMMaterialLayerObject* InLayer = nullptr);
+
+	UDMMaterialStageThroughputLayerBlend();
+
+	UDMMaterialStageInput* GetInputMask() const;
+
+	UFUNCTION(BlueprintCallable, Category = "Material Designer")
+	EAvaColorChannel GetMaskChannelOverride() const;
+
+	UFUNCTION(BlueprintCallable, Category="Material Designer")
+	void SetMaskChannelOverride(EAvaColorChannel InMaskChannel);
+
+	//~ Begin UDMMaterialComponent
+	virtual void OnComponentAdded() override;
+	virtual void Update(EDMUpdateType InUpdateType) override;
+	virtual void GenerateExpressions(const TSharedRef<FDMMaterialBuildState>& InBuildState) const override;
+	virtual void PostEditorDuplicate(UDynamicMaterialModel* InMaterialModel, UDMMaterialComponent* InParent) override;
+	virtual bool IsPropertyVisible(FName InProperty) const override;
+	//~ End UDMMaterialComponent
+
+	//~ Begin UDMMaterialStageSource
+	virtual FText GetStageDescription() const override;
+	virtual bool UpdateStagePreviewMaterial(UDMMaterialStage* InStage, UMaterial* InPreviewMaterial, UMaterialExpression*& OutMaterialExpression,
+		int32& OutputIndex) override;
+	//~ End UDMMaterialStageSource
+
+	//~ Begin FNotifyHook
+	virtual void NotifyPostChange(const FPropertyChangedEvent& InPropertyChangedEvent, class FEditPropertyChain* InPropertyThatChanged) override;
+	//~ End FNotifyHook
+
+	//~ Begin UDMMaterialStageThroughput
+	virtual void AddDefaultInput(int32 InInputIndex) const override;
+	virtual bool CanChangeInput(int32 InputIndex) const override;
+	virtual bool IsInputVisible(int32 InputIndex) const override;
+	virtual int32 ResolveInput(const TSharedRef<FDMMaterialBuildState>& InBuildState, int32 InputIndex, FDMMaterialStageConnectorChannel& OutChannel,
+		TArray<UMaterialExpression*>& OutExpressions) const override;
+	virtual void ConnectOutputToInput(const TSharedRef<FDMMaterialBuildState>& InBuildState, int32 InInputIndex, UMaterialExpression* InSourceExpression,
+		int32 InSourceOutputIndex, int32 InSourceOutputChannel) override;
+	//~ End UDMMaterialStageThroughput
+
+	void GetMaskOutput(const TSharedRef<FDMMaterialBuildState>& InBuildState, UMaterialExpression*& OutExpression, int32& OutOutputIndex, int32& OutOutputChannel) const;
+
+	UFUNCTION(BlueprintPure, Category = "Material Designer")
+	bool UsePremultiplyAlpha() const { return bPremultiplyAlpha; }
+
+	UFUNCTION(BlueprintCallable, Category = "Material Designer")
+	void SetPremultiplyAlpha(bool bInValue);
+
+	//~ Begin UObject
+	virtual void PostLoad() override;
+	virtual void PostEditImport() override;
+	//~ End UObject
+
+protected:
+	UPROPERTY(EditInstanceOnly, BlueprintReadOnly, Getter=GetMaskChannelOverride, Setter=SetMaskChannelOverride, BlueprintGetter=GetMaskChannelOverride,
+		Category = "Material Designer", DisplayName = "ChannelOverride",
+		meta=(NotKeyframeable, ToolTip="Changes the output channel of the mask input."))
+	mutable EAvaColorChannel MaskChannelOverride;
+
+	UPROPERTY(EditInstanceOnly, BlueprintReadOnly, Category = "Material Designer", 
+		meta=(NotKeyframeable, ToolTip="Additionally multiplies the output of the RGB Stage by the output from this Stage."))
+	bool bPremultiplyAlpha;
+
+	UPROPERTY(EditInstanceOnly, BlueprintReadOnly, Category = "Material Designer",
+		meta=(NotKeyframeable))
+	bool bIsAlphaOnlyBlend;
+
+	bool bBlockUpdate;
+
+	virtual int32 ResolveMaskInput(const TSharedRef<FDMMaterialBuildState>& InBuildState, int32 InputIndex, FDMMaterialStageConnectorChannel& OutChannel,
+		TArray<UMaterialExpression*>& OutExpressions) const;
+	
+	virtual void UpdatePreviewMaterial(UMaterial* InPreviewMaterial = nullptr) override;
+
+	virtual void UpdateAlphaOnlyMaskStatus();
+	virtual void OnStageUpdated(UDMMaterialComponent* InComponent, EDMUpdateType InUpdateType);
+	virtual void UpdateAlphaOnlyMasks(EDMUpdateType InUpdateType);
+
+	void InitBlendStage();
+
+	void GenerateMainExpressions(const TSharedRef<FDMMaterialBuildState>& InBuildState) const;
+	void GeneratePreviewExpressions(const TSharedRef<FDMMaterialBuildState>& InBuildState) const;
+
+	void UpdateLinkedInputStage(EDMUpdateType InUpdateType);
+
+	/* Returns true if there are any outputs on the mask input that have more than 1 channel. */
+	bool CanUseMaskChannelOverride() const;
+
+	/* Returns the first output on the mask input that has more than 1 channel. */
+	int32 GetDefaultMaskChannelOverrideOutputIndex() const;
+
+	/* Returns true if the given mask output supports more than 1 channel. */
+	bool IsValidMaskChannelOverrideOutputIndex(int32 InIndex) const;
+
+	/* Reads the current output setting from the input map. */
+	void PullMaskChannelOverride() const;
+
+	/* Takes the override setting and applies it to the input map. */
+	void PushMaskChannelOverride();
+};
diff --git a/Engine/Plugins/Experimental/DynamicMaterial/Source/DynamicMaterialEditor/Public/Components/DMMaterialSubStage.h b/Engine/Plugins/Experimental/DynamicMaterial/Source/DynamicMaterialEditor/Public/Components/DMMaterialSubStage.h
new file mode 100644
--- /dev/null
+++ b/Engine/Plugins/Experimental/DynamicMaterial/Source/DynamicMaterialEditor/Public/Components/DMMaterialSubStage.h
@@ -0,0 +1,57 @@
+// Copyright Epic Games, Inc. All Rights Reserved.
+
+#pragma once
+
+#include "DMMaterialStage.h"
+#include "DMMaterialSubStage.generated.h"
+
+/**
+ * A stage that is a subobject of another stage.
+ */
+UCLASS(BlueprintType, ClassGroup = "Material Designer", meta = (DisplayName = "Material Designer Sub Stage"))
+class DYNAMICMATERIALEDITOR_API UDMMaterialSubStage : public UDMMaterialStage
+{
+	GENERATED_BODY()
+
+	friend class UDMMaterialStageInputThroughput;
+
+public:
+	UFUNCTION(BlueprintCallable, Category = "Material Designer")
+	static UDMMaterialSubStage* CreateMaterialSubStage(UDMMaterialStage* InParentStage);
+
+	UDMMaterialStage* GetParentStage() const;
+	UDMMaterialStage* GetParentMostStage() const;
+
+	//~ Begin UDMMaterialStage
+	virtual bool IsCompatibleWithPreviousStage(const UDMMaterialStage* PreviousStage) const override;
+	virtual bool IsCompatibleWithNextStage(const UDMMaterialStage* NextStage) const override;
+	virtual bool IsRootStage() const override;
+	//~ End UDMMaterialStage
+
+	//~ Begin UDMMaterialComponent
+	virtual FString GetComponentPathComponent() const override;
+	virtual UDMMaterialComponent* GetParentComponent() const override;
+	virtual void PostEditorDuplicate(UDynamicMaterialModel* InMaterialModel, UDMMaterialComponent* InParent) override;
+	//~ End UDMMaterialComponent
+
+	void SetParentComponent(UDMMaterialComponent* InParentComponent);
+
+protected:
+	UPROPERTY(VisibleInstanceOnly, BlueprintReadOnly, Category = "Material Designer")
+	TObjectPtr<UDMMaterialStage> ParentStage;
+
+	UPROPERTY()
+	TObjectPtr<UDMMaterialComponent> ParentComponent;
+
+	UDMMaterialSubStage()
+		: ParentStage(nullptr)
+		, ParentComponent(nullptr)
+	{
+	}
+
+	//~ Begin UDMMaterialStage
+	// Sub stages don't need these bound.
+	virtual void AddDelegates() override {}
+	virtual void RemoveDelegates() override {}
+	//~ End UDMMaterialStage
+};
diff --git a/Engine/Plugins/Experimental/DynamicMaterial/Source/DynamicMaterialEditor/Public/Components/MaterialStageBlends/DMMSBAdd.h b/Engine/Plugins/Experimental/DynamicMaterial/Source/DynamicMaterialEditor/Public/Components/MaterialStageBlends/DMMSBAdd.h
new file mode 100644
--- /dev/null
+++ b/Engine/Plugins/Experimental/DynamicMaterial/Source/DynamicMaterialEditor/Public/Components/MaterialStageBlends/DMMSBAdd.h
@@ -0,0 +1,15 @@
+// Copyright Epic Games, Inc. All Rights Reserved.
+
+#pragma once
+
+#include "Components/DMMaterialStageBlendFunction.h"
+#include "DMMSBAdd.generated.h"
+
+UCLASS(BlueprintType, ClassGroup = "Material Designer")
+class DYNAMICMATERIALEDITOR_API UDMMaterialStageBlendAdd : public UDMMaterialStageBlendFunction
+{
+	GENERATED_BODY()
+
+public:
+	UDMMaterialStageBlendAdd();
+};
diff --git a/Engine/Plugins/Experimental/DynamicMaterial/Source/DynamicMaterialEditor/Public/Components/MaterialStageBlends/DMMSBColor.h b/Engine/Plugins/Experimental/DynamicMaterial/Source/DynamicMaterialEditor/Public/Components/MaterialStageBlends/DMMSBColor.h
new file mode 100644
--- /dev/null
+++ b/Engine/Plugins/Experimental/DynamicMaterial/Source/DynamicMaterialEditor/Public/Components/MaterialStageBlends/DMMSBColor.h
@@ -0,0 +1,15 @@
+// Copyright Epic Games, Inc. All Rights Reserved.
+
+#pragma once
+
+#include "Components/DMMaterialStageBlendFunction.h"
+#include "DMMSBColor.generated.h"
+
+UCLASS(BlueprintType, ClassGroup = "Material Designer")
+class DYNAMICMATERIALEDITOR_API UDMMaterialStageBlendColor : public UDMMaterialStageBlendFunction
+{
+	GENERATED_BODY()
+
+public:
+	UDMMaterialStageBlendColor();
+};
diff --git a/Engine/Plugins/Experimental/DynamicMaterial/Source/DynamicMaterialEditor/Public/Components/MaterialStageBlends/DMMSBColorBurn.h b/Engine/Plugins/Experimental/DynamicMaterial/Source/DynamicMaterialEditor/Public/Components/MaterialStageBlends/DMMSBColorBurn.h
new file mode 100644
--- /dev/null
+++ b/Engine/Plugins/Experimental/DynamicMaterial/Source/DynamicMaterialEditor/Public/Components/MaterialStageBlends/DMMSBColorBurn.h
@@ -0,0 +1,15 @@
+// Copyright Epic Games, Inc. All Rights Reserved.
+
+#pragma once
+
+#include "Components/DMMaterialStageBlendFunction.h"
+#include "DMMSBColorBurn.generated.h"
+
+UCLASS(BlueprintType, ClassGroup = "Material Designer")
+class DYNAMICMATERIALEDITOR_API UDMMaterialStageBlendColorBurn : public UDMMaterialStageBlendFunction
+{
+	GENERATED_BODY()
+
+public:
+	UDMMaterialStageBlendColorBurn();
+};
diff --git a/Engine/Plugins/Experimental/DynamicMaterial/Source/DynamicMaterialEditor/Public/Components/MaterialStageBlends/DMMSBColorDodge.h b/Engine/Plugins/Experimental/DynamicMaterial/Source/DynamicMaterialEditor/Public/Components/MaterialStageBlends/DMMSBColorDodge.h
new file mode 100644
--- /dev/null
+++ b/Engine/Plugins/Experimental/DynamicMaterial/Source/DynamicMaterialEditor/Public/Components/MaterialStageBlends/DMMSBColorDodge.h
@@ -0,0 +1,15 @@
+// Copyright Epic Games, Inc. All Rights Reserved.
+
+#pragma once
+
+#include "Components/DMMaterialStageBlendFunction.h"
+#include "DMMSBColorDodge.generated.h"
+
+UCLASS(BlueprintType, ClassGroup = "Material Designer")
+class DYNAMICMATERIALEDITOR_API UDMMaterialStageBlendColorDodge : public UDMMaterialStageBlendFunction
+{
+	GENERATED_BODY()
+
+public:
+	UDMMaterialStageBlendColorDodge();
+};
diff --git a/Engine/Plugins/Experimental/DynamicMaterial/Source/DynamicMaterialEditor/Public/Components/MaterialStageBlends/DMMSBContrastBase.h b/Engine/Plugins/Experimental/DynamicMaterial/Source/DynamicMaterialEditor/Public/Components/MaterialStageBlends/DMMSBContrastBase.h
new file mode 100644
--- /dev/null
+++ b/Engine/Plugins/Experimental/DynamicMaterial/Source/DynamicMaterialEditor/Public/Components/MaterialStageBlends/DMMSBContrastBase.h
@@ -0,0 +1,21 @@
+// Copyright Epic Games, Inc. All Rights Reserved.
+
+#pragma once
+
+#include "Components/DMMaterialStageBlend.h"
+#include "DMMSBContrastBase.generated.h"
+
+class UMaterialExpression;
+struct FDMMaterialBuildState;
+
+UCLASS(BlueprintType, Abstract, ClassGroup = "Material Designer")
+class DYNAMICMATERIALEDITOR_API UDMMaterialStageBlendContrastBase : public UDMMaterialStageBlend
+{
+	GENERATED_BODY()
+
+public:
+	UDMMaterialStageBlendContrastBase(const FText& InName);
+
+protected:
+	UDMMaterialStageBlendContrastBase();
+};
diff --git a/Engine/Plugins/Experimental/DynamicMaterial/Source/DynamicMaterialEditor/Public/Components/MaterialStageBlends/DMMSBDarken.h b/Engine/Plugins/Experimental/DynamicMaterial/Source/DynamicMaterialEditor/Public/Components/MaterialStageBlends/DMMSBDarken.h
new file mode 100644
--- /dev/null
+++ b/Engine/Plugins/Experimental/DynamicMaterial/Source/DynamicMaterialEditor/Public/Components/MaterialStageBlends/DMMSBDarken.h
@@ -0,0 +1,15 @@
+// Copyright Epic Games, Inc. All Rights Reserved.
+
+#pragma once
+
+#include "Components/DMMaterialStageBlendFunction.h"
+#include "DMMSBDarken.generated.h"
+
+UCLASS(BlueprintType, ClassGroup = "Material Designer")
+class DYNAMICMATERIALEDITOR_API UDMMaterialStageBlendDarken : public UDMMaterialStageBlendFunction
+{
+	GENERATED_BODY()
+
+public:
+	UDMMaterialStageBlendDarken();
+};
diff --git a/Engine/Plugins/Experimental/DynamicMaterial/Source/DynamicMaterialEditor/Public/Components/MaterialStageBlends/DMMSBDarkenColor.h b/Engine/Plugins/Experimental/DynamicMaterial/Source/DynamicMaterialEditor/Public/Components/MaterialStageBlends/DMMSBDarkenColor.h
new file mode 100644
--- /dev/null
+++ b/Engine/Plugins/Experimental/DynamicMaterial/Source/DynamicMaterialEditor/Public/Components/MaterialStageBlends/DMMSBDarkenColor.h
@@ -0,0 +1,15 @@
+// Copyright Epic Games, Inc. All Rights Reserved.
+
+#pragma once
+
+#include "Components/DMMaterialStageBlendFunction.h"
+#include "DMMSBDarkenColor.generated.h"
+
+UCLASS(BlueprintType, ClassGroup = "Material Designer")
+class DYNAMICMATERIALEDITOR_API UDMMaterialStageBlendDarkenColor : public UDMMaterialStageBlendFunction
+{
+	GENERATED_BODY()
+
+public:
+	UDMMaterialStageBlendDarkenColor();
+};
diff --git a/Engine/Plugins/Experimental/DynamicMaterial/Source/DynamicMaterialEditor/Public/Components/MaterialStageBlends/DMMSBDifference.h b/Engine/Plugins/Experimental/DynamicMaterial/Source/DynamicMaterialEditor/Public/Components/MaterialStageBlends/DMMSBDifference.h
new file mode 100644
--- /dev/null
+++ b/Engine/Plugins/Experimental/DynamicMaterial/Source/DynamicMaterialEditor/Public/Components/MaterialStageBlends/DMMSBDifference.h
@@ -0,0 +1,15 @@
+// Copyright Epic Games, Inc. All Rights Reserved.
+
+#pragma once
+
+#include "Components/DMMaterialStageBlendFunction.h"
+#include "DMMSBDifference.generated.h"
+
+UCLASS(BlueprintType, ClassGroup = "Material Designer")
+class DYNAMICMATERIALEDITOR_API UDMMaterialStageBlendDifference : public UDMMaterialStageBlendFunction
+{
+	GENERATED_BODY()
+
+public:
+	UDMMaterialStageBlendDifference();
+};
diff --git a/Engine/Plugins/Experimental/DynamicMaterial/Source/DynamicMaterialEditor/Public/Components/MaterialStageBlends/DMMSBDivide.h b/Engine/Plugins/Experimental/DynamicMaterial/Source/DynamicMaterialEditor/Public/Components/MaterialStageBlends/DMMSBDivide.h
new file mode 100644
--- /dev/null
+++ b/Engine/Plugins/Experimental/DynamicMaterial/Source/DynamicMaterialEditor/Public/Components/MaterialStageBlends/DMMSBDivide.h
@@ -0,0 +1,15 @@
+// Copyright Epic Games, Inc. All Rights Reserved.
+
+#pragma once
+
+#include "Components/DMMaterialStageBlendFunction.h"
+#include "DMMSBDivide.generated.h"
+
+UCLASS(BlueprintType, ClassGroup = "Material Designer")
+class DYNAMICMATERIALEDITOR_API UDMMaterialStageBlendDivide : public UDMMaterialStageBlendFunction
+{
+	GENERATED_BODY()
+
+public:
+	UDMMaterialStageBlendDivide();
+};
diff --git a/Engine/Plugins/Experimental/DynamicMaterial/Source/DynamicMaterialEditor/Public/Components/MaterialStageBlends/DMMSBExclusion.h b/Engine/Plugins/Experimental/DynamicMaterial/Source/DynamicMaterialEditor/Public/Components/MaterialStageBlends/DMMSBExclusion.h
new file mode 100644
--- /dev/null
+++ b/Engine/Plugins/Experimental/DynamicMaterial/Source/DynamicMaterialEditor/Public/Components/MaterialStageBlends/DMMSBExclusion.h
@@ -0,0 +1,15 @@
+// Copyright Epic Games, Inc. All Rights Reserved.
+
+#pragma once
+
+#include "Components/DMMaterialStageBlendFunction.h"
+#include "DMMSBExclusion.generated.h"
+
+UCLASS(BlueprintType, ClassGroup = "Material Designer")
+class DYNAMICMATERIALEDITOR_API UDMMaterialStageBlendExclusion : public UDMMaterialStageBlendFunction
+{
+	GENERATED_BODY()
+
+public:
+	UDMMaterialStageBlendExclusion();
+};
