[MaterialDesigner]

diff --git a/Engine/Plugins/Experimental/DynamicMaterial/Source/DynamicMaterialEditor/Public/Components/MaterialStageBlends/DMMSBHardLight.h b/Engine/Plugins/Experimental/DynamicMaterial/Source/DynamicMaterialEditor/Public/Components/MaterialStageBlends/DMMSBHardLight.h
new file mode 100644
--- /dev/null
+++ b/Engine/Plugins/Experimental/DynamicMaterial/Source/DynamicMaterialEditor/Public/Components/MaterialStageBlends/DMMSBHardLight.h
@@ -0,0 +1,15 @@
+// Copyright Epic Games, Inc. All Rights Reserved.
+
+#pragma once
+
+#include "Components/DMMaterialStageBlendFunction.h"
+#include "DMMSBHardLight.generated.h"
+
+UCLASS(BlueprintType, ClassGroup = "Material Designer")
+class DYNAMICMATERIALEDITOR_API UDMMaterialStageBlendHardLight : public UDMMaterialStageBlendFunction
+{
+	GENERATED_BODY()
+
+public:
+	UDMMaterialStageBlendHardLight();
+};
diff --git a/Engine/Plugins/Experimental/DynamicMaterial/Source/DynamicMaterialEditor/Public/Components/MaterialStageBlends/DMMSBHardMix.h b/Engine/Plugins/Experimental/DynamicMaterial/Source/DynamicMaterialEditor/Public/Components/MaterialStageBlends/DMMSBHardMix.h
new file mode 100644
--- /dev/null
+++ b/Engine/Plugins/Experimental/DynamicMaterial/Source/DynamicMaterialEditor/Public/Components/MaterialStageBlends/DMMSBHardMix.h
@@ -0,0 +1,15 @@
+// Copyright Epic Games, Inc. All Rights Reserved.
+
+#pragma once
+
+#include "Components/DMMaterialStageBlendFunction.h"
+#include "DMMSBHardMix.generated.h"
+
+UCLASS(BlueprintType, ClassGroup = "Material Designer")
+class DYNAMICMATERIALEDITOR_API UDMMaterialStageBlendHardMix : public UDMMaterialStageBlendFunction
+{
+	GENERATED_BODY()
+
+public:
+	UDMMaterialStageBlendHardMix();
+};
diff --git a/Engine/Plugins/Experimental/DynamicMaterial/Source/DynamicMaterialEditor/Public/Components/MaterialStageBlends/DMMSBHue.h b/Engine/Plugins/Experimental/DynamicMaterial/Source/DynamicMaterialEditor/Public/Components/MaterialStageBlends/DMMSBHue.h
new file mode 100644
--- /dev/null
+++ b/Engine/Plugins/Experimental/DynamicMaterial/Source/DynamicMaterialEditor/Public/Components/MaterialStageBlends/DMMSBHue.h
@@ -0,0 +1,15 @@
+// Copyright Epic Games, Inc. All Rights Reserved.
+
+#pragma once
+
+#include "Components/DMMaterialStageBlendFunction.h"
+#include "DMMSBHue.generated.h"
+
+UCLASS(BlueprintType, ClassGroup = "Material Designer")
+class DYNAMICMATERIALEDITOR_API UDMMaterialStageBlendHue : public UDMMaterialStageBlendFunction
+{
+	GENERATED_BODY()
+
+public:
+	UDMMaterialStageBlendHue();
+};
diff --git a/Engine/Plugins/Experimental/DynamicMaterial/Source/DynamicMaterialEditor/Public/Components/MaterialStageBlends/DMMSBLighten.h b/Engine/Plugins/Experimental/DynamicMaterial/Source/DynamicMaterialEditor/Public/Components/MaterialStageBlends/DMMSBLighten.h
new file mode 100644
--- /dev/null
+++ b/Engine/Plugins/Experimental/DynamicMaterial/Source/DynamicMaterialEditor/Public/Components/MaterialStageBlends/DMMSBLighten.h
@@ -0,0 +1,15 @@
+// Copyright Epic Games, Inc. All Rights Reserved.
+
+#pragma once
+
+#include "Components/DMMaterialStageBlendFunction.h"
+#include "DMMSBLighten.generated.h"
+
+UCLASS(BlueprintType, ClassGroup = "Material Designer")
+class DYNAMICMATERIALEDITOR_API UDMMaterialStageBlendLighten : public UDMMaterialStageBlendFunction
+{
+	GENERATED_BODY()
+
+public:
+	UDMMaterialStageBlendLighten();
+};
diff --git a/Engine/Plugins/Experimental/DynamicMaterial/Source/DynamicMaterialEditor/Public/Components/MaterialStageBlends/DMMSBLightenColor.h b/Engine/Plugins/Experimental/DynamicMaterial/Source/DynamicMaterialEditor/Public/Components/MaterialStageBlends/DMMSBLightenColor.h
new file mode 100644
--- /dev/null
+++ b/Engine/Plugins/Experimental/DynamicMaterial/Source/DynamicMaterialEditor/Public/Components/MaterialStageBlends/DMMSBLightenColor.h
@@ -0,0 +1,15 @@
+// Copyright Epic Games, Inc. All Rights Reserved.
+
+#pragma once
+
+#include "Components/DMMaterialStageBlendFunction.h"
+#include "DMMSBLightenColor.generated.h"
+
+UCLASS(BlueprintType, ClassGroup = "Material Designer")
+class DYNAMICMATERIALEDITOR_API UDMMaterialStageBlendLightenColor : public UDMMaterialStageBlendFunction
+{
+	GENERATED_BODY()
+
+public:
+	UDMMaterialStageBlendLightenColor();
+};
diff --git a/Engine/Plugins/Experimental/DynamicMaterial/Source/DynamicMaterialEditor/Public/Components/MaterialStageBlends/DMMSBLinearBurn.h b/Engine/Plugins/Experimental/DynamicMaterial/Source/DynamicMaterialEditor/Public/Components/MaterialStageBlends/DMMSBLinearBurn.h
new file mode 100644
--- /dev/null
+++ b/Engine/Plugins/Experimental/DynamicMaterial/Source/DynamicMaterialEditor/Public/Components/MaterialStageBlends/DMMSBLinearBurn.h
@@ -0,0 +1,15 @@
+// Copyright Epic Games, Inc. All Rights Reserved.
+
+#pragma once
+
+#include "Components/DMMaterialStageBlendFunction.h"
+#include "DMMSBLinearBurn.generated.h"
+
+UCLASS(BlueprintType, ClassGroup = "Material Designer")
+class DYNAMICMATERIALEDITOR_API UDMMaterialStageBlendLinearBurn : public UDMMaterialStageBlendFunction
+{
+	GENERATED_BODY()
+
+public:
+	UDMMaterialStageBlendLinearBurn();
+};
diff --git a/Engine/Plugins/Experimental/DynamicMaterial/Source/DynamicMaterialEditor/Public/Components/MaterialStageBlends/DMMSBLinearDodge.h b/Engine/Plugins/Experimental/DynamicMaterial/Source/DynamicMaterialEditor/Public/Components/MaterialStageBlends/DMMSBLinearDodge.h
new file mode 100644
--- /dev/null
+++ b/Engine/Plugins/Experimental/DynamicMaterial/Source/DynamicMaterialEditor/Public/Components/MaterialStageBlends/DMMSBLinearDodge.h
@@ -0,0 +1,15 @@
+// Copyright Epic Games, Inc. All Rights Reserved.
+
+#pragma once
+
+#include "Components/DMMaterialStageBlendFunction.h"
+#include "DMMSBLinearDodge.generated.h"
+
+UCLASS(BlueprintType, ClassGroup = "Material Designer")
+class DYNAMICMATERIALEDITOR_API UDMMaterialStageBlendLinearDodge : public UDMMaterialStageBlendFunction
+{
+	GENERATED_BODY()
+
+public:
+	UDMMaterialStageBlendLinearDodge();
+};
diff --git a/Engine/Plugins/Experimental/DynamicMaterial/Source/DynamicMaterialEditor/Public/Components/MaterialStageBlends/DMMSBLinearLight.h b/Engine/Plugins/Experimental/DynamicMaterial/Source/DynamicMaterialEditor/Public/Components/MaterialStageBlends/DMMSBLinearLight.h
new file mode 100644
--- /dev/null
+++ b/Engine/Plugins/Experimental/DynamicMaterial/Source/DynamicMaterialEditor/Public/Components/MaterialStageBlends/DMMSBLinearLight.h
@@ -0,0 +1,15 @@
+// Copyright Epic Games, Inc. All Rights Reserved.
+
+#pragma once
+
+#include "Components/DMMaterialStageBlendFunction.h"
+#include "DMMSBLinearLight.generated.h"
+
+UCLASS(BlueprintType, ClassGroup = "Material Designer")
+class DYNAMICMATERIALEDITOR_API UDMMaterialStageBlendLinearLight : public UDMMaterialStageBlendFunction
+{
+	GENERATED_BODY()
+
+public:
+	UDMMaterialStageBlendLinearLight();
+};
diff --git a/Engine/Plugins/Experimental/DynamicMaterial/Source/DynamicMaterialEditor/Public/Components/MaterialStageBlends/DMMSBLuminosity.h b/Engine/Plugins/Experimental/DynamicMaterial/Source/DynamicMaterialEditor/Public/Components/MaterialStageBlends/DMMSBLuminosity.h
new file mode 100644
--- /dev/null
+++ b/Engine/Plugins/Experimental/DynamicMaterial/Source/DynamicMaterialEditor/Public/Components/MaterialStageBlends/DMMSBLuminosity.h
@@ -0,0 +1,15 @@
+// Copyright Epic Games, Inc. All Rights Reserved.
+
+#pragma once
+
+#include "Components/DMMaterialStageBlendFunction.h"
+#include "DMMSBLuminosity.generated.h"
+
+UCLASS(BlueprintType, ClassGroup = "Material Designer")
+class DYNAMICMATERIALEDITOR_API UDMMaterialStageBlendLuminosity : public UDMMaterialStageBlendFunction
+{
+	GENERATED_BODY()
+
+public:
+	UDMMaterialStageBlendLuminosity();
+};
diff --git a/Engine/Plugins/Experimental/DynamicMaterial/Source/DynamicMaterialEditor/Public/Components/MaterialStageBlends/DMMSBMultiply.h b/Engine/Plugins/Experimental/DynamicMaterial/Source/DynamicMaterialEditor/Public/Components/MaterialStageBlends/DMMSBMultiply.h
new file mode 100644
--- /dev/null
+++ b/Engine/Plugins/Experimental/DynamicMaterial/Source/DynamicMaterialEditor/Public/Components/MaterialStageBlends/DMMSBMultiply.h
@@ -0,0 +1,15 @@
+// Copyright Epic Games, Inc. All Rights Reserved.
+
+#pragma once
+
+#include "Components/DMMaterialStageBlendFunction.h"
+#include "DMMSBMultiply.generated.h"
+
+UCLASS(BlueprintType, ClassGroup = "Material Designer")
+class DYNAMICMATERIALEDITOR_API UDMMaterialStageBlendMultiply : public UDMMaterialStageBlendFunction
+{
+	GENERATED_BODY()
+
+public:
+	UDMMaterialStageBlendMultiply();
+};
diff --git a/Engine/Plugins/Experimental/DynamicMaterial/Source/DynamicMaterialEditor/Public/Components/MaterialStageBlends/DMMSBNormal.h b/Engine/Plugins/Experimental/DynamicMaterial/Source/DynamicMaterialEditor/Public/Components/MaterialStageBlends/DMMSBNormal.h
new file mode 100644
--- /dev/null
+++ b/Engine/Plugins/Experimental/DynamicMaterial/Source/DynamicMaterialEditor/Public/Components/MaterialStageBlends/DMMSBNormal.h
@@ -0,0 +1,15 @@
+// Copyright Epic Games, Inc. All Rights Reserved.
+
+#pragma once
+
+#include "Components/DMMaterialStageBlendFunction.h"
+#include "DMMSBNormal.generated.h"
+
+UCLASS(BlueprintType, ClassGroup = "Material Designer")
+class DYNAMICMATERIALEDITOR_API UDMMaterialStageBlendNormal : public UDMMaterialStageBlendFunction
+{
+	GENERATED_BODY()
+
+public:
+	UDMMaterialStageBlendNormal();
+};
diff --git a/Engine/Plugins/Experimental/DynamicMaterial/Source/DynamicMaterialEditor/Public/Components/MaterialStageBlends/DMMSBOverlay.h b/Engine/Plugins/Experimental/DynamicMaterial/Source/DynamicMaterialEditor/Public/Components/MaterialStageBlends/DMMSBOverlay.h
new file mode 100644
--- /dev/null
+++ b/Engine/Plugins/Experimental/DynamicMaterial/Source/DynamicMaterialEditor/Public/Components/MaterialStageBlends/DMMSBOverlay.h
@@ -0,0 +1,15 @@
+// Copyright Epic Games, Inc. All Rights Reserved.
+
+#pragma once
+
+#include "Components/DMMaterialStageBlendFunction.h"
+#include "DMMSBOverlay.generated.h"
+
+UCLASS(BlueprintType, ClassGroup = "Material Designer")
+class DYNAMICMATERIALEDITOR_API UDMMaterialStageBlendOverlay : public UDMMaterialStageBlendFunction
+{
+	GENERATED_BODY()
+
+public:
+	UDMMaterialStageBlendOverlay();
+};
diff --git a/Engine/Plugins/Experimental/DynamicMaterial/Source/DynamicMaterialEditor/Public/Components/MaterialStageBlends/DMMSBPinLight.h b/Engine/Plugins/Experimental/DynamicMaterial/Source/DynamicMaterialEditor/Public/Components/MaterialStageBlends/DMMSBPinLight.h
new file mode 100644
--- /dev/null
+++ b/Engine/Plugins/Experimental/DynamicMaterial/Source/DynamicMaterialEditor/Public/Components/MaterialStageBlends/DMMSBPinLight.h
@@ -0,0 +1,15 @@
+// Copyright Epic Games, Inc. All Rights Reserved.
+
+#pragma once
+
+#include "Components/DMMaterialStageBlendFunction.h"
+#include "DMMSBPinLight.generated.h"
+
+UCLASS(BlueprintType, ClassGroup = "Material Designer")
+class DYNAMICMATERIALEDITOR_API UDMMaterialStageBlendPinLight : public UDMMaterialStageBlendFunction
+{
+	GENERATED_BODY()
+
+public:
+	UDMMaterialStageBlendPinLight();
+};
diff --git a/Engine/Plugins/Experimental/DynamicMaterial/Source/DynamicMaterialEditor/Public/Components/MaterialStageBlends/DMMSBSaturation.h b/Engine/Plugins/Experimental/DynamicMaterial/Source/DynamicMaterialEditor/Public/Components/MaterialStageBlends/DMMSBSaturation.h
new file mode 100644
--- /dev/null
+++ b/Engine/Plugins/Experimental/DynamicMaterial/Source/DynamicMaterialEditor/Public/Components/MaterialStageBlends/DMMSBSaturation.h
@@ -0,0 +1,15 @@
+// Copyright Epic Games, Inc. All Rights Reserved.
+
+#pragma once
+
+#include "Components/DMMaterialStageBlendFunction.h"
+#include "DMMSBSaturation.generated.h"
+
+UCLASS(BlueprintType, ClassGroup = "Material Designer")
+class DYNAMICMATERIALEDITOR_API UDMMaterialStageBlendSaturation : public UDMMaterialStageBlendFunction
+{
+	GENERATED_BODY()
+
+public:
+	UDMMaterialStageBlendSaturation();
+};
diff --git a/Engine/Plugins/Experimental/DynamicMaterial/Source/DynamicMaterialEditor/Public/Components/MaterialStageBlends/DMMSBScreen.h b/Engine/Plugins/Experimental/DynamicMaterial/Source/DynamicMaterialEditor/Public/Components/MaterialStageBlends/DMMSBScreen.h
new file mode 100644
--- /dev/null
+++ b/Engine/Plugins/Experimental/DynamicMaterial/Source/DynamicMaterialEditor/Public/Components/MaterialStageBlends/DMMSBScreen.h
@@ -0,0 +1,15 @@
+// Copyright Epic Games, Inc. All Rights Reserved.
+
+#pragma once
+
+#include "Components/DMMaterialStageBlendFunction.h"
+#include "DMMSBScreen.generated.h"
+
+UCLASS(BlueprintType, ClassGroup = "Material Designer")
+class DYNAMICMATERIALEDITOR_API UDMMaterialStageBlendScreen : public UDMMaterialStageBlendFunction
+{
+	GENERATED_BODY()
+
+public:
+	UDMMaterialStageBlendScreen();
+};
diff --git a/Engine/Plugins/Experimental/DynamicMaterial/Source/DynamicMaterialEditor/Public/Components/MaterialStageBlends/DMMSBSoftLight.h b/Engine/Plugins/Experimental/DynamicMaterial/Source/DynamicMaterialEditor/Public/Components/MaterialStageBlends/DMMSBSoftLight.h
new file mode 100644
--- /dev/null
+++ b/Engine/Plugins/Experimental/DynamicMaterial/Source/DynamicMaterialEditor/Public/Components/MaterialStageBlends/DMMSBSoftLight.h
@@ -0,0 +1,15 @@
+// Copyright Epic Games, Inc. All Rights Reserved.
+
+#pragma once
+
+#include "Components/DMMaterialStageBlendFunction.h"
+#include "DMMSBSoftLight.generated.h"
+
+UCLASS(BlueprintType, ClassGroup = "Material Designer")
+class DYNAMICMATERIALEDITOR_API UDMMaterialStageBlendSoftLight : public UDMMaterialStageBlendFunction
+{
+	GENERATED_BODY()
+
+public:
+	UDMMaterialStageBlendSoftLight();
+};
diff --git a/Engine/Plugins/Experimental/DynamicMaterial/Source/DynamicMaterialEditor/Public/Components/MaterialStageBlends/DMMSBSubtract.h b/Engine/Plugins/Experimental/DynamicMaterial/Source/DynamicMaterialEditor/Public/Components/MaterialStageBlends/DMMSBSubtract.h
new file mode 100644
--- /dev/null
+++ b/Engine/Plugins/Experimental/DynamicMaterial/Source/DynamicMaterialEditor/Public/Components/MaterialStageBlends/DMMSBSubtract.h
@@ -0,0 +1,15 @@
+// Copyright Epic Games, Inc. All Rights Reserved.
+
+#pragma once
+
+#include "Components/DMMaterialStageBlendFunction.h"
+#include "DMMSBSubtract.generated.h"
+
+UCLASS(BlueprintType, ClassGroup = "Material Designer")
+class DYNAMICMATERIALEDITOR_API UDMMaterialStageBlendSubtract : public UDMMaterialStageBlendFunction
+{
+	GENERATED_BODY()
+
+public:
+	UDMMaterialStageBlendSubtract();
+};
diff --git a/Engine/Plugins/Experimental/DynamicMaterial/Source/DynamicMaterialEditor/Public/Components/MaterialStageBlends/DMMSBVividLight.h b/Engine/Plugins/Experimental/DynamicMaterial/Source/DynamicMaterialEditor/Public/Components/MaterialStageBlends/DMMSBVividLight.h
new file mode 100644
--- /dev/null
+++ b/Engine/Plugins/Experimental/DynamicMaterial/Source/DynamicMaterialEditor/Public/Components/MaterialStageBlends/DMMSBVividLight.h
@@ -0,0 +1,15 @@
+// Copyright Epic Games, Inc. All Rights Reserved.
+
+#pragma once
+
+#include "Components/DMMaterialStageBlendFunction.h"
+#include "DMMSBVividLight.generated.h"
+
+UCLASS(BlueprintType, ClassGroup = "Material Designer")
+class DYNAMICMATERIALEDITOR_API UDMMaterialStageBlendVividLight : public UDMMaterialStageBlendFunction
+{
+	GENERATED_BODY()
+
+public:
+	UDMMaterialStageBlendVividLight();
+};
diff --git a/Engine/Plugins/Experimental/DynamicMaterial/Source/DynamicMaterialEditor/Public/Components/MaterialStageExpressions/DMMSEMathBase.h b/Engine/Plugins/Experimental/DynamicMaterial/Source/DynamicMaterialEditor/Public/Components/MaterialStageExpressions/DMMSEMathBase.h
new file mode 100644
--- /dev/null
+++ b/Engine/Plugins/Experimental/DynamicMaterial/Source/DynamicMaterialEditor/Public/Components/MaterialStageExpressions/DMMSEMathBase.h
@@ -0,0 +1,48 @@
+// Copyright Epic Games, Inc. All Rights Reserved.
+
+#pragma once
+
+#include "Components/DMMaterialStageExpression.h"
+
+#include "DMMSEMathBase.generated.h"
+
+class UDMMaterialStageInput;
+
+UCLASS(BlueprintType, ClassGroup = "Material Designer")
+class DYNAMICMATERIALEDITOR_API UDMMaterialStageExpressionMathBase : public UDMMaterialStageExpression
+{
+	GENERATED_BODY()
+
+public:
+
+	UDMMaterialStageExpressionMathBase(const FText& InName, TSubclassOf<UMaterialExpression> InClass);
+
+	virtual bool CanInputAcceptType(int32 InInputIndex, EDMValueType InValueType) const override;
+
+	virtual void Update(EDMUpdateType InUpdateType) override;
+
+	virtual void AddDefaultInput(int32 InInputIndex) const override;
+
+	//~ Begin UDMMaterialStageSource
+	virtual bool UpdateStagePreviewMaterial(UDMMaterialStage* InStage, UMaterial* InPreviewMaterial, UMaterialExpression*& OutMaterialExpression,
+		int32& OutputIndex) override;
+	//~ End UDMMaterialStageSource
+
+protected:
+
+	UPROPERTY(VisibleInstanceOnly, BlueprintReadOnly, Category = "Material Designer")
+	bool bSingleChannelOnly;
+
+	UPROPERTY(VisibleInstanceOnly, BlueprintReadOnly, Category = "Material Designer")
+	int32 VariadicInputCount;
+
+	UPROPERTY(VisibleInstanceOnly, BlueprintReadOnly, Category = "Material Designer")
+	bool bAllowSingleFloatMatch;
+
+	UDMMaterialStageExpressionMathBase();
+
+	void SetupInputs(int32 InCount);
+
+	virtual void UpdatePreviewMaterial(UMaterial* InPreviewMaterial = nullptr) override;
+
+};
diff --git a/Engine/Plugins/Experimental/DynamicMaterial/Source/DynamicMaterialEditor/Public/Components/MaterialStageExpressions/DMMSEParticleSubUV.h b/Engine/Plugins/Experimental/DynamicMaterial/Source/DynamicMaterialEditor/Public/Components/MaterialStageExpressions/DMMSEParticleSubUV.h
new file mode 100644
--- /dev/null
+++ b/Engine/Plugins/Experimental/DynamicMaterial/Source/DynamicMaterialEditor/Public/Components/MaterialStageExpressions/DMMSEParticleSubUV.h
@@ -0,0 +1,15 @@
+// Copyright Epic Games, Inc. All Rights Reserved.
+
+#pragma once
+
+#include "Components/MaterialStageExpressions/DMMSETextureSampleBase.h"
+#include "DMMSEParticleSubUV.generated.h"
+
+UCLASS(BlueprintType, ClassGroup = "Material Designer")
+class DYNAMICMATERIALEDITOR_API UDMMaterialStageExpressionParticleSubUV : public UDMMaterialStageExpressionTextureSampleBase
+{
+	GENERATED_BODY()
+
+public:
+	UDMMaterialStageExpressionParticleSubUV();
+};
diff --git a/Engine/Plugins/Experimental/DynamicMaterial/Source/DynamicMaterialEditor/Public/Components/MaterialStageExpressions/DMMSETextureCoordinate.h b/Engine/Plugins/Experimental/DynamicMaterial/Source/DynamicMaterialEditor/Public/Components/MaterialStageExpressions/DMMSETextureCoordinate.h
new file mode 100644
--- /dev/null
+++ b/Engine/Plugins/Experimental/DynamicMaterial/Source/DynamicMaterialEditor/Public/Components/MaterialStageExpressions/DMMSETextureCoordinate.h
@@ -0,0 +1,27 @@
+// Copyright Epic Games, Inc. All Rights Reserved.
+
+#pragma once
+
+#include "Components/DMMaterialStageExpression.h"
+#include "DMMSETextureCoordinate.generated.h"
+
+UCLASS(BlueprintType, ClassGroup = "Material Designer")
+class DYNAMICMATERIALEDITOR_API UDMMaterialStageExpressionTextureCoordinate : public UDMMaterialStageExpression
+{
+	GENERATED_BODY()
+
+public:
+	UDMMaterialStageExpressionTextureCoordinate();
+
+	virtual void AddExpressionProperties(const TArray<UMaterialExpression*>& InExpressions) const override;
+
+protected:
+	UPROPERTY(EditInstanceOnly, BlueprintReadOnly, Category = "Material Designer")
+	int32 CoordinateIndex;
+
+	UPROPERTY(EditInstanceOnly, BlueprintReadOnly, Category = "Material Designer")
+	float UTiling;
+
+	UPROPERTY(EditInstanceOnly, BlueprintReadOnly, Category = "Material Designer")
+	float VTiling;
+};
diff --git a/Engine/Plugins/Experimental/DynamicMaterial/Source/DynamicMaterialEditor/Public/Components/MaterialStageExpressions/DMMSETextureSample.h b/Engine/Plugins/Experimental/DynamicMaterial/Source/DynamicMaterialEditor/Public/Components/MaterialStageExpressions/DMMSETextureSample.h
new file mode 100644
--- /dev/null
+++ b/Engine/Plugins/Experimental/DynamicMaterial/Source/DynamicMaterialEditor/Public/Components/MaterialStageExpressions/DMMSETextureSample.h
@@ -0,0 +1,17 @@
+// Copyright Epic Games, Inc. All Rights Reserved.
+
+#pragma once
+
+#include "Components/MaterialStageExpressions/DMMSETextureSampleBase.h"
+#include "DMMSETextureSample.generated.h"
+
+UCLASS(BlueprintType, ClassGroup = "Material Designer")
+class DYNAMICMATERIALEDITOR_API UDMMaterialStageExpressionTextureSample : public UDMMaterialStageExpressionTextureSampleBase
+{
+	GENERATED_BODY()
+
+	friend class SDMComponentEdit;
+
+public:
+	UDMMaterialStageExpressionTextureSample();
+};
diff --git a/Engine/Plugins/Experimental/DynamicMaterial/Source/DynamicMaterialEditor/Public/Components/MaterialStageExpressions/DMMSETextureSampleBase.h b/Engine/Plugins/Experimental/DynamicMaterial/Source/DynamicMaterialEditor/Public/Components/MaterialStageExpressions/DMMSETextureSampleBase.h
new file mode 100644
--- /dev/null
+++ b/Engine/Plugins/Experimental/DynamicMaterial/Source/DynamicMaterialEditor/Public/Components/MaterialStageExpressions/DMMSETextureSampleBase.h
@@ -0,0 +1,45 @@
+// Copyright Epic Games, Inc. All Rights Reserved.
+
+#pragma once
+
+#include "Components/DMMaterialStageExpression.h"
+#include "DMMSETextureSampleBase.generated.h"
+
+UCLASS(BlueprintType, ClassGroup = "Material Designer")
+class DYNAMICMATERIALEDITOR_API UDMMaterialStageExpressionTextureSampleBase : public UDMMaterialStageExpression
+{
+	GENERATED_BODY()
+
+public:
+	UDMMaterialStageExpressionTextureSampleBase(const FText& InName, TSubclassOf<UMaterialExpression> InClass);
+
+	//~ Begin UDMMaterialStageSource
+	virtual void OnComponentAdded() override;
+	virtual bool IsPropertyVisible(FName Property) const override;
+	virtual void AddExpressionProperties(const TArray<UMaterialExpression*>& InExpressions) const override;
+	virtual int32 GetInnateMaskOutput(int32 InOutputIndex, int32 InOutputChannels) const override;
+	virtual int32 GetOutputChannelOverride(int32 InOutputIndex) const override;
+	//~ End UDMMaterialStageSource
+
+	//~ Begin UDMMaterialStageThroughput
+	virtual bool CanChangeInputType(int32 InInputIndex) const override;
+	virtual bool SupportsLayerMaskTextureUVLink() const override { return true; }
+	virtual int32 GetLayerMaskTextureUVLinkInputIndex() const override { return 1; }
+	virtual void InputUpdated(int32 InInputIndex, EDMUpdateType InUpdateType) override;
+	//~ End UDMMaterialStageThroughput
+
+	UFUNCTION(BlueprintPure, Category = "Material Designer")
+	virtual bool IsClampTextureEnabled() const { return bClampTexture; }
+
+	UFUNCTION(BlueprintCallable, Category = "Material Designer")
+	virtual void SetClampTextureEnabled(bool bInValue);
+
+protected:
+	UPROPERTY(EditInstanceOnly, BlueprintReadOnly, Setter=SetClampTextureEnabled, Category = "Material Designer", 
+		meta=(ToolTip="Forces a material rebuild.", LowPriority, NotKeyframeable, AllowPrivateAccess = "true"))
+	bool bClampTexture;
+
+	UDMMaterialStageExpressionTextureSampleBase();
+
+	void UpdateMask();
+};
diff --git a/Engine/Plugins/Experimental/DynamicMaterial/Source/DynamicMaterialEditor/Public/Components/MaterialStageExpressions/DMMSETextureSampleEdgeColor.h b/Engine/Plugins/Experimental/DynamicMaterial/Source/DynamicMaterialEditor/Public/Components/MaterialStageExpressions/DMMSETextureSampleEdgeColor.h
new file mode 100644
--- /dev/null
+++ b/Engine/Plugins/Experimental/DynamicMaterial/Source/DynamicMaterialEditor/Public/Components/MaterialStageExpressions/DMMSETextureSampleEdgeColor.h
@@ -0,0 +1,46 @@
+// Copyright Epic Games, Inc. All Rights Reserved.
+
+#pragma once
+
+#include "Components/DMMaterialStageExpression.h"
+#include "DMMSETextureSampleEdgeColor.generated.h"
+
+UENUM(BlueprintType)
+enum class EDMEdgeLocation : uint8
+{
+	TopLeft,
+	Top,
+	TopRight,
+	Left,
+	Center,
+	Right,
+	BottomLeft,
+	Bottom,
+	BottomRight,
+	Custom
+};
+
+UCLASS(BlueprintType, ClassGroup = "Material Designer")
+class DYNAMICMATERIALEDITOR_API UDMMaterialStageExpressionTextureSampleEdgeColor : public UDMMaterialStageExpression
+{
+	GENERATED_BODY()
+
+public:
+	UDMMaterialStageExpressionTextureSampleEdgeColor();
+
+	//~ Begin UDMMaterialStageThroughput
+	virtual bool CanChangeInputType(int32 InInputIndex) const override;
+	virtual bool IsInputVisible(int32 InInputIndex) const override;
+	virtual void AddDefaultInput(int32 InInputIndex) const override;
+	//~ End UDMMaterialStageThroughput
+
+	//~ Begin UObject
+	virtual void PostEditChangeProperty(struct FPropertyChangedEvent& PropertyChangedEvent) override;
+	//~ End UObject
+
+protected:
+	UPROPERTY(EditInstanceOnly, BlueprintReadOnly, Category = "Material Designer", meta=(AllowPrivateAccess = "true"))
+	EDMEdgeLocation EdgeLocation;
+
+	void OnEdgeLocationChanged();
+};
diff --git a/Engine/Plugins/Experimental/DynamicMaterial/Source/DynamicMaterialEditor/Public/Components/MaterialStageGradients/DMMSGLinear.h b/Engine/Plugins/Experimental/DynamicMaterial/Source/DynamicMaterialEditor/Public/Components/MaterialStageGradients/DMMSGLinear.h
new file mode 100644
--- /dev/null
+++ b/Engine/Plugins/Experimental/DynamicMaterial/Source/DynamicMaterialEditor/Public/Components/MaterialStageGradients/DMMSGLinear.h
@@ -0,0 +1,38 @@
+// Copyright Epic Games, Inc. All Rights Reserved.
+
+#pragma once
+
+#include "Components/DMMaterialStageGradient.h"
+#include "DMMSGLinear.generated.h"
+
+struct FDMMaterialBuildState;
+class UMaterialExpression;
+
+UENUM(BlueprintType)
+enum class ELinearGradientTileType : uint8
+{
+	NoTile,
+	Tile,
+	TileAndMirror
+};
+
+UCLASS(BlueprintType, ClassGroup = "Material Designer")
+class DYNAMICMATERIALEDITOR_API UDMMaterialStageGradientLinear : public UDMMaterialStageGradient
+{
+	GENERATED_BODY()
+
+public:
+	UDMMaterialStageGradientLinear();
+
+	virtual void GenerateExpressions(const TSharedRef<FDMMaterialBuildState>& InBuildState) const override;
+
+	UFUNCTION(BlueprintPure, Category = "Material Designer")
+	virtual ELinearGradientTileType GetTilingType() const { return Tiling; }
+
+	UFUNCTION(BlueprintCallable, Category = "Material Designer")
+	virtual void SetTilingType(ELinearGradientTileType InType);
+
+protected:
+	UPROPERTY(EditInstanceOnly, BlueprintReadOnly, Getter=GetTilingType, Setter=SetTilingType, Category = "Material Designer")
+	ELinearGradientTileType Tiling;
+};
diff --git a/Engine/Plugins/Experimental/DynamicMaterial/Source/DynamicMaterialEditor/Public/Components/MaterialStageGradients/DMMSGRadial.h b/Engine/Plugins/Experimental/DynamicMaterial/Source/DynamicMaterialEditor/Public/Components/MaterialStageGradients/DMMSGRadial.h
new file mode 100644
--- /dev/null
+++ b/Engine/Plugins/Experimental/DynamicMaterial/Source/DynamicMaterialEditor/Public/Components/MaterialStageGradients/DMMSGRadial.h
@@ -0,0 +1,20 @@
+// Copyright Epic Games, Inc. All Rights Reserved.
+
+#pragma once
+
+#include "Components/DMMaterialStageGradient.h"
+#include "DMMSGRadial.generated.h"
+
+struct FDMMaterialBuildState;
+class UMaterialExpression;
+
+UCLASS(BlueprintType, ClassGroup = "Material Designer")
+class DYNAMICMATERIALEDITOR_API UDMMaterialStageGradientRadial : public UDMMaterialStageGradient
+{
+	GENERATED_BODY()
+
+public:
+	UDMMaterialStageGradientRadial();
+
+	virtual void GenerateExpressions(const TSharedRef<FDMMaterialBuildState>& InBuildState) const override;
+};
diff --git a/Engine/Plugins/Experimental/DynamicMaterial/Source/DynamicMaterialEditor/Public/Components/MaterialStageInputs/DMMSIExpression.h b/Engine/Plugins/Experimental/DynamicMaterial/Source/DynamicMaterialEditor/Public/Components/MaterialStageInputs/DMMSIExpression.h
new file mode 100644
--- /dev/null
+++ b/Engine/Plugins/Experimental/DynamicMaterial/Source/DynamicMaterialEditor/Public/Components/MaterialStageInputs/DMMSIExpression.h
@@ -0,0 +1,37 @@
+// Copyright Epic Games, Inc. All Rights Reserved.
+
+#pragma once
+
+#include "Components/MaterialStageInputs/DMMSIThroughput.h"
+#include "UObject/StrongObjectPtr.h"
+#include "DMMSIExpression.generated.h"
+
+class UDMMaterialLayerObject;
+class UDMMaterialStageExpression;
+
+UCLASS(BlueprintType, ClassGroup = "Material Designer")
+class DYNAMICMATERIALEDITOR_API UDMMaterialStageInputExpression : public UDMMaterialStageInputThroughput
+{
+	GENERATED_BODY()
+
+public:
+	static UDMMaterialStage* CreateStage(TSubclassOf<UDMMaterialStageExpression> InMaterialStageExpressionClass, UDMMaterialLayerObject* InLayer = nullptr);
+
+	static const TArray<TStrongObjectPtr<UClass>>& GetAvailableInputExpressions();
+
+	UFUNCTION(BlueprintPure, Category = "Material Designer")
+	TSubclassOf<UDMMaterialStageExpression> GetMaterialStageExpressionClass() const;
+
+	UFUNCTION(BlueprintCallable, Category = "Material Designer")
+	void SetMaterialStageExpressionClass(TSubclassOf<UDMMaterialStageExpression> InMaterialStageExpressionClass);
+
+	UFUNCTION(BlueprintPure, Category = "Material Designer")
+	UDMMaterialStageExpression* GetMaterialStageExpression() const;
+
+protected:
+	static TArray<TStrongObjectPtr<UClass>> InputExpressions;
+
+	static void GenerateExpressionList();
+
+	UDMMaterialStageInputExpression() = default;
+};
diff --git a/Engine/Plugins/Experimental/DynamicMaterial/Source/DynamicMaterialEditor/Public/Components/MaterialStageInputs/DMMSIFunction.h b/Engine/Plugins/Experimental/DynamicMaterial/Source/DynamicMaterialEditor/Public/Components/MaterialStageInputs/DMMSIFunction.h
new file mode 100644
--- /dev/null
+++ b/Engine/Plugins/Experimental/DynamicMaterial/Source/DynamicMaterialEditor/Public/Components/MaterialStageInputs/DMMSIFunction.h
@@ -0,0 +1,26 @@
+// Copyright Epic Games, Inc. All Rights Reserved.
+
+#pragma once
+
+#include "Components/MaterialStageInputs/DMMSIThroughput.h"
+#include "DMMSIFunction.generated.h"
+
+class UDMMaterialLayerObject;
+class UDMMaterialStageFunction;
+
+UCLASS(BlueprintType, ClassGroup = "Material Designer")
+class DYNAMICMATERIALEDITOR_API UDMMaterialStageInputFunction : public UDMMaterialStageInputThroughput
+{
+	GENERATED_BODY()
+
+public:
+	static UDMMaterialStage* CreateStage(UDMMaterialLayerObject* InLayer = nullptr);
+
+	void Init();
+
+	UFUNCTION(BlueprintPure, Category = "Material Designer")
+	UDMMaterialStageFunction* GetMaterialStageFunction() const;
+
+protected:
+	UDMMaterialStageInputFunction() = default;
+};
diff --git a/Engine/Plugins/Experimental/DynamicMaterial/Source/DynamicMaterialEditor/Public/Components/MaterialStageInputs/DMMSIGradient.h b/Engine/Plugins/Experimental/DynamicMaterial/Source/DynamicMaterialEditor/Public/Components/MaterialStageInputs/DMMSIGradient.h
new file mode 100644
--- /dev/null
+++ b/Engine/Plugins/Experimental/DynamicMaterial/Source/DynamicMaterialEditor/Public/Components/MaterialStageInputs/DMMSIGradient.h
@@ -0,0 +1,37 @@
+// Copyright Epic Games, Inc. All Rights Reserved.
+
+#pragma once
+
+#include "Components/MaterialStageInputs/DMMSIThroughput.h"
+#include "UObject/StrongObjectPtr.h"
+#include "DMMSIGradient.generated.h"
+
+class UDMMaterialLayerObject;
+class UDMMaterialStageGradient;
+
+UCLASS(BlueprintType, ClassGroup = "Material Designer")
+class DYNAMICMATERIALEDITOR_API UDMMaterialStageInputGradient : public UDMMaterialStageInputThroughput
+{
+	GENERATED_BODY()
+
+public:
+	static UDMMaterialStage* CreateStage(TSubclassOf<UDMMaterialStageGradient> InMaterialStageGradientClass, UDMMaterialLayerObject* InLayer = nullptr);
+
+	static const TArray<TStrongObjectPtr<UClass>>& GetAvailableGradients();
+
+	UFUNCTION(BlueprintPure, Category = "Material Designer")
+	TSubclassOf<UDMMaterialStageGradient> GetMaterialStageGradientClass() const;
+
+	UFUNCTION(BlueprintCallable, Category = "Material Designer")
+	void SetMaterialStageGradientClass(TSubclassOf<UDMMaterialStageGradient> InMaterialStageGradientClass);
+
+	UFUNCTION(BlueprintPure, Category = "Material Designer")
+	UDMMaterialStageGradient* GetMaterialStageGradient() const;
+
+protected:
+	static TArray<TStrongObjectPtr<UClass>> Gradients;
+
+	static void GenerateGradientList();
+
+	UDMMaterialStageInputGradient() = default;
+};
diff --git a/Engine/Plugins/Experimental/DynamicMaterial/Source/DynamicMaterialEditor/Public/Components/MaterialStageInputs/DMMSISlot.h b/Engine/Plugins/Experimental/DynamicMaterial/Source/DynamicMaterialEditor/Public/Components/MaterialStageInputs/DMMSISlot.h
new file mode 100644
--- /dev/null
+++ b/Engine/Plugins/Experimental/DynamicMaterial/Source/DynamicMaterialEditor/Public/Components/MaterialStageInputs/DMMSISlot.h
@@ -0,0 +1,74 @@
+// Copyright Epic Games, Inc. All Rights Reserved.
+
+#pragma once
+
+#include "Components/DMMaterialStageInput.h"
+#include "DMMSISlot.generated.h"
+
+class UDMMaterialLayerObject;
+class UDMMaterialStage;
+class UDynamicMaterialModel;
+struct FDMMaterialBuildState;
+
+UCLASS(BlueprintType, ClassGroup = "Material Designer")
+class DYNAMICMATERIALEDITOR_API UDMMaterialStageInputSlot : public UDMMaterialStageInput
+{
+	GENERATED_BODY()
+
+public:
+	static const FString SlotPathToken;
+
+	static UDMMaterialStage* CreateStage(UDMMaterialSlot* InSourceSlot, EDMMaterialPropertyType InMaterialProperty, 
+		UDMMaterialLayerObject* InLayer = nullptr);
+
+	UFUNCTION(BlueprintPure, Category = "Material Designer")
+	UDMMaterialSlot* GetSlot() const { return Slot; }
+
+	UFUNCTION(BlueprintPure, Category = "Material Designer")
+	EDMMaterialPropertyType GetMaterialProperty() const { return MaterialProperty; }
+
+	virtual FText GetComponentDescription() const override;
+	virtual FText GetChannelDescription(const FDMMaterialStageConnectorChannel& Channel) override;
+
+	UFUNCTION(BlueprintCallable, Category = "Material Designer")
+	void SetSlot(UDMMaterialSlot* InSlot);
+
+	UFUNCTION(BlueprintCallable, Category = "Material Designer")
+	void SetMaterialProperty(EDMMaterialPropertyType InMaterialProperty);
+
+	virtual void GenerateExpressions(const TSharedRef<FDMMaterialBuildState>& InBuildState) const override;
+
+	//~ Begin UObject
+	virtual void PostLoad() override;
+	virtual void PostEditImport() override;
+	//~ End UObject
+
+	//~ Start UDMMaterialComponent
+	virtual void PostEditorDuplicate(UDynamicMaterialModel* InMaterialModel, UDMMaterialComponent* InParent) override;
+	//~ End UDMMaterialComponent
+
+protected:
+	UDMMaterialStageInputSlot();
+	virtual ~UDMMaterialStageInputSlot() override = default;
+
+	UPROPERTY(VisibleInstanceOnly, BlueprintReadOnly, Category = "Material Designer")
+	TObjectPtr<UDMMaterialSlot> Slot;
+
+	UPROPERTY(VisibleInstanceOnly, BlueprintReadOnly, Category = "Material Designer")
+	EDMMaterialPropertyType MaterialProperty;
+
+	void OnSlotUpdated(UDMMaterialComponent* InComponent, EDMUpdateType InUpdateType);
+	void OnSlotConnectorsUpdated(UDMMaterialSlot* InSlot);
+	void OnSlotRemoved(UDMMaterialComponent* InComponent, EDMComponentLifetimeState InLifetimeState);
+	void OnParentSlotRemoved(UDMMaterialComponent* InComponent, EDMComponentLifetimeState InLifetimeState);
+
+	virtual void UpdateOutputConnectors() override;
+
+	void InitSlot();
+	void DeinitSlot();
+
+	//~ Begin UDMMaterialComponent
+	virtual void OnComponentRemoved() override;
+	virtual UDMMaterialComponent* GetSubComponentByPath(FDMComponentPath& InPath, const FDMComponentPathSegment& InPathSegment) const override;
+	//~ End UDMMaterialComponent
+};
diff --git a/Engine/Plugins/Experimental/DynamicMaterial/Source/DynamicMaterialEditor/Public/Components/MaterialStageInputs/DMMSITextureUV.h b/Engine/Plugins/Experimental/DynamicMaterial/Source/DynamicMaterialEditor/Public/Components/MaterialStageInputs/DMMSITextureUV.h
new file mode 100644
--- /dev/null
+++ b/Engine/Plugins/Experimental/DynamicMaterial/Source/DynamicMaterialEditor/Public/Components/MaterialStageInputs/DMMSITextureUV.h
@@ -0,0 +1,73 @@
+// Copyright Epic Games, Inc. All Rights Reserved.
+
+#pragma once
+
+#include "Components/DMMaterialStageInput.h"
+#include "Templates/SharedPointer.h"
+#include "DMMSITextureUV.generated.h"
+
+class SWidget;
+class UDMMaterialLayerObject;
+class UDMMaterialParameter;
+class UDMMaterialStage;
+class UDMTextureUV;
+class UDynamicMaterialModel;
+class UMaterialExpressionScalarParameter;
+class UMaterialInstanceDynamic;
+struct FDMMaterialBuildState;
+
+UCLASS(BlueprintType, ClassGroup = "Material Designer")
+class DYNAMICMATERIALEDITOR_API UDMMaterialStageInputTextureUV : public UDMMaterialStageInput
+{
+	GENERATED_BODY()
+
+public:
+	static const FString TextureUVPathToken;
+
+	static UDMMaterialStage* CreateStage(UDynamicMaterialModel* InMaterialModel, UDMMaterialLayerObject* InLayer = nullptr);
+
+	virtual ~UDMMaterialStageInputTextureUV() override = default;
+
+	virtual FText GetComponentDescription() const override;
+	virtual FText GetChannelDescription(const FDMMaterialStageConnectorChannel& Channel) override;
+
+	void Init(UDynamicMaterialModel* InMaterialModel);
+
+	UFUNCTION(BlueprintPure, Category = "Material Designer")
+	UDMTextureUV* GetTextureUV() { return TextureUV; }
+
+	virtual void GenerateExpressions(const TSharedRef<FDMMaterialBuildState>& InBuildState) const override;
+
+	//~ Begin UObject
+	virtual bool Modify(bool bInAlwaysMarkDirty = true) override;
+	virtual void PostLoad() override;
+	virtual void PostEditImport() override;
+	//~ End UObject
+
+	//~ Start UDMMaterialComponent
+	virtual void PostEditorDuplicate(UDynamicMaterialModel* InMaterialModel, UDMMaterialComponent* InParent) override;
+	//~ End UDMMaterialComponent
+
+protected:
+	static UMaterialExpressionScalarParameter* CreateScalarParameter(const TSharedRef<FDMMaterialBuildState>& InBuildState, FName InParamName, float InValue = 0.f);
+	static TArray<UMaterialExpression*> CreateTextureUVExpressions(const TSharedRef<FDMMaterialBuildState>& InBuildState, UDMTextureUV* InTextureUV);
+
+	UPROPERTY(EditInstanceOnly, BlueprintReadOnly, Instanced, Category = "Material Designer")
+	TObjectPtr<UDMTextureUV> TextureUV;
+
+	UDMMaterialStageInputTextureUV();
+
+	virtual void UpdateOutputConnectors() override;
+
+	void OnTextureUVUpdated(UDMMaterialComponent* InComponent, EDMUpdateType InUpdateType);
+
+	void InitTextureUV();
+
+	void AddEffects(const TSharedRef<FDMMaterialBuildState>& InBuildState, TArray<UMaterialExpression*>& InOutExpressions) const;
+
+	//~ Begin UDMMaterialComponent
+	virtual void OnComponentAdded() override;
+	virtual void OnComponentRemoved() override;
+	virtual UDMMaterialComponent* GetSubComponentByPath(FDMComponentPath& InPath, const FDMComponentPathSegment& InPathSegment) const override;
+	//~ End UDMMaterialComponent
+};
diff --git a/Engine/Plugins/Experimental/DynamicMaterial/Source/DynamicMaterialEditor/Public/Components/MaterialStageInputs/DMMSIThroughput.h b/Engine/Plugins/Experimental/DynamicMaterial/Source/DynamicMaterialEditor/Public/Components/MaterialStageInputs/DMMSIThroughput.h
new file mode 100644
--- /dev/null
+++ b/Engine/Plugins/Experimental/DynamicMaterial/Source/DynamicMaterialEditor/Public/Components/MaterialStageInputs/DMMSIThroughput.h
@@ -0,0 +1,68 @@
+// Copyright Epic Games, Inc. All Rights Reserved.
+
+#pragma once
+
+#include "Components/DMMaterialStageInput.h"
+#include "DMMSIThroughput.generated.h"
+
+class UDMMaterialStage;
+class UDMMaterialStageThroughput;
+class UDMMaterialSubStage;
+class UDynamicMaterialModel;
+struct FDMMaterialBuildState;
+
+UCLASS(Abstract, BlueprintType, ClassGroup = "Material Designer")
+class DYNAMICMATERIALEDITOR_API UDMMaterialStageInputThroughput : public UDMMaterialStageInput
+{
+	GENERATED_BODY()
+
+public:
+	static const FString SubStagePathToken;
+
+	virtual FText GetComponentDescription() const override;
+	virtual FText GetChannelDescription(const FDMMaterialStageConnectorChannel& Channel) override;
+
+	UFUNCTION(BlueprintPure, Category = "Material Designer")
+	TSubclassOf<UDMMaterialStageThroughput> GetMaterialStageThroughputClass() const;
+
+	UFUNCTION(BlueprintPure, Category = "Material Designer")
+	UDMMaterialStageThroughput* GetMaterialStageThroughput() const;
+
+	UFUNCTION(BlueprintPure, Category = "Material Designer")
+	UDMMaterialSubStage* GetSubStage() const { return SubStage; }
+
+	virtual void GenerateExpressions(const TSharedRef<FDMMaterialBuildState>& InBuildState) const override;
+	virtual int32 GetInnateMaskOutput(int32 OutputIndex, int32 OutputChannels) const;
+	virtual int32 GetOutputChannelOverride(int32 InOutputIndex) const override;
+
+	virtual bool IsPropertyVisible(FName Property) const override;
+
+	//~ Begin UObject
+	virtual bool Modify(bool bInAlwaysMarkDirty = true) override;
+	virtual void PostLoad() override;
+	virtual void PostEditImport() override;
+	//~ End UObject
+
+	//~ Start UDMMaterialComponent
+	virtual void PostEditorDuplicate(UDynamicMaterialModel* InMaterialModel, UDMMaterialComponent* InParent) override;
+	//~ End UDMMaterialComponent
+
+protected:
+	UPROPERTY(VisibleInstanceOnly, BlueprintReadOnly, Instanced, Category = "Material Designer")
+	TObjectPtr<UDMMaterialSubStage> SubStage;
+
+	UDMMaterialStageInputThroughput();
+
+	void OnSubStageUpdated(UDMMaterialComponent* InComponent, EDMUpdateType InUpdateType);
+
+	void SetMaterialStageThroughputClass(TSubclassOf<UDMMaterialStageThroughput> InMaterialStageThroughputClass);
+
+	void InitSubStage();
+
+	//~ Begin UDMMaterialComponent
+	virtual void OnComponentAdded() override;
+	virtual void OnComponentRemoved() override;
+	virtual void GetComponentPathInternal(TArray<FString>& OutChildComponentPathComponents) const override;
+	virtual UDMMaterialComponent* GetSubComponentByPath(FDMComponentPath& InPath, const FDMComponentPathSegment& InPathSegment) const override;
+	//~ End UDMMaterialComponent
+};
diff --git a/Engine/Plugins/Experimental/DynamicMaterial/Source/DynamicMaterialEditor/Public/Components/MaterialStageInputs/DMMSIValue.h b/Engine/Plugins/Experimental/DynamicMaterial/Source/DynamicMaterialEditor/Public/Components/MaterialStageInputs/DMMSIValue.h
new file mode 100644
--- /dev/null
+++ b/Engine/Plugins/Experimental/DynamicMaterial/Source/DynamicMaterialEditor/Public/Components/MaterialStageInputs/DMMSIValue.h
@@ -0,0 +1,68 @@
+// Copyright Epic Games, Inc. All Rights Reserved.
+
+#pragma once
+
+#include "Components/DMMaterialStageInput.h"
+#include "DMMSIValue.generated.h"
+
+class UDMMaterialLayerObject;
+class UDMMaterialStage;
+class UDynamicMaterialModel;
+
+UCLASS(BlueprintType, ClassGroup = "Material Designer")
+class DYNAMICMATERIALEDITOR_API UDMMaterialStageInputValue : public UDMMaterialStageInput
+{
+	GENERATED_BODY()
+		
+public:
+	static const FString ValuePathToken;
+
+	static FName GetValuePropertyName() { return GET_MEMBER_NAME_CHECKED(UDMMaterialStageInputValue, Value); }
+
+	static UDMMaterialStage* CreateStage(UDMMaterialValue* InValue, UDMMaterialLayerObject* InLayer = nullptr);
+
+	virtual FText GetComponentDescription() const override;
+	virtual FText GetChannelDescription(const FDMMaterialStageConnectorChannel& Channel) override;
+
+	UFUNCTION(BlueprintPure, Category = "Material Designer")
+	UDMMaterialValue* GetValue() const { return Value; }
+
+	UFUNCTION(BlueprintCallable, Category = "Material Designer")
+	void SetValue(UDMMaterialValue* InValue);
+
+	virtual void GenerateExpressions(const TSharedRef<FDMMaterialBuildState>& InBuildState) const override;
+
+	virtual int32 GetInnateMaskOutput(int32 OutputIndex, int32 OutputChannels) const override;
+
+	//~ Begin UObject
+	virtual bool Modify(bool bInAlwaysMarkDirty = true) override;
+	virtual void PostLoad() override;
+	virtual void PostEditImport() override;
+	//~ End UObject
+
+	//~ Start UDMMaterialComponent
+	virtual void PostEditorDuplicate(UDynamicMaterialModel* InMaterialModel, UDMMaterialComponent* InParent) override;
+	//~ End UDMMaterialComponent
+
+protected:
+	UDMMaterialStageInputValue();
+	virtual ~UDMMaterialStageInputValue() override = default;
+
+	UPROPERTY(VisibleInstanceOnly, BlueprintReadOnly, Instanced, Category = "Material Designer")
+	TObjectPtr<UDMMaterialValue> Value;
+
+	void OnValueUpdated(UDMMaterialComponent* InComponent, EDMUpdateType InUpdateType);
+
+	virtual void UpdateOutputConnectors() override;
+
+	void InitInputValue();
+	void DeinitInputValue();
+
+	//~ Begin UDMMaterialStageSource
+	virtual void OnComponentAdded() override;
+	virtual void OnComponentRemoved() override;
+	virtual UDMMaterialComponent* GetSubComponentByPath(FDMComponentPath& InPath, const FDMComponentPathSegment& InPathSegment) const override;
+	//~ End UDMMaterialStageSource
+
+	bool IsSharedStageValue() const;
+};
diff --git a/Engine/Plugins/Experimental/DynamicMaterial/Source/DynamicMaterialEditor/Public/DMBlueprintFunctionLibrary.h b/Engine/Plugins/Experimental/DynamicMaterial/Source/DynamicMaterialEditor/Public/DMBlueprintFunctionLibrary.h
new file mode 100644
--- /dev/null
+++ b/Engine/Plugins/Experimental/DynamicMaterial/Source/DynamicMaterialEditor/Public/DMBlueprintFunctionLibrary.h
@@ -0,0 +1,78 @@
+// Copyright Epic Games, Inc. All Rights Reserved.
+
+#pragma once
+
+#include "DMDefs.h"
+#include "DMEDefs.h"
+#include "Kismet/BlueprintFunctionLibrary.h"
+#include "DMBlueprintFunctionLibrary.generated.h"
+
+class AActor;
+class UDMMaterialSlot;
+class UDMMaterialStage;
+class UDMMaterialStageBlend;
+class UDMMaterialStageExpression;
+class UDMMaterialStageGradient;
+class UDMMaterialStageInputValue;
+class UDMMaterialValue;
+class UDynamicMaterialModel;
+class UTexture;
+struct FDMObjectMaterialProperty;
+
+/**
+ * Material Designer Blueprint Function Library
+ */
+UCLASS()
+class DYNAMICMATERIALEDITOR_API UDMBlueprintFunctionLibrary : public UBlueprintFunctionLibrary
+{
+	GENERATED_BODY()
+	
+public:
+	UFUNCTION(BlueprintCallable, Category = "Material Designer")
+	static UDMMaterialStageInputValue* FindDefaultStageOpacityInputValue(UDMMaterialStage* InStage);
+
+	UFUNCTION(BlueprintCallable, Category = "Material Designer")
+	static void SetDefaultStageSourceTexture(UDMMaterialStage* InStage, UTexture* InTexture);
+
+	UFUNCTION(BlueprintCallable, Category = "Material Designer")
+	static TArray<FDMObjectMaterialProperty> GetActorDynamicMaterialSlots(AActor* InActor, const bool bInIncludeEmptySlots);
+
+	UFUNCTION(BlueprintCallable, Category = "Material Designer")
+	static TArray<FDMObjectMaterialProperty> GetActorMaterialProperties(AActor* InActor);
+	
+	UFUNCTION(BlueprintCallable, Category = "Material Designer")
+	static UDynamicMaterialModel* FindFirstValidActorMaterialSlot(const TArray<AActor*>& InActors);
+
+	UFUNCTION(BlueprintCallable, Category = "Material Designer")
+	static FDMObjectMaterialProperty MakePrimitiveComponentSlot(UPrimitiveComponent* InComponent, int32 InSlot);
+
+	UFUNCTION(BlueprintCallable, Category = "Material Designer")
+	static FDMObjectMaterialProperty MakeObjectMaterialProperty(UObject* InObject, FName PropertyName);
+
+	UFUNCTION(BlueprintCallable, Category = "Material Designer")
+	static FDMObjectMaterialProperty MakeObjectMaterialPropertyArray(UObject* InObject, FName PropertyName, int32 ArrayIndex);
+
+	UFUNCTION(BlueprintCallable, Category = "Material Designer")
+	static UDynamicMaterialModel* CreateDynamicMaterialInObject(FDMObjectMaterialProperty& InMaterialProperty);
+
+	UFUNCTION(BlueprintCallable, Category = "Material Designer")
+	UDynamicMaterialModel* GetObjectPropertyMaterialModel(const FDMObjectMaterialProperty& InMaterialProperty) const;
+
+	UFUNCTION(BlueprintCallable, Category = "Material Designer")
+	UDynamicMaterialInstance* GetObjectPropertyMaterial(const FDMObjectMaterialProperty& InMaterialProperty) const;
+
+	UFUNCTION(BlueprintCallable, Category = "Material Designer")
+	void SetObjectPropertyMaterial(UPARAM(Ref) FDMObjectMaterialProperty& InMaterialProperty, UDynamicMaterialInstance* InDynamicMaterial);
+
+	UFUNCTION(BlueprintCallable, Category = "Material Designer")
+	bool IsObjectPropertyValid(const FDMObjectMaterialProperty& InMaterialProperty) const;
+
+	UFUNCTION(BlueprintCallable, Category = "Material Designer")
+	FText GetObjectPropertyName(const FDMObjectMaterialProperty& InMaterialProperty, bool bInIgnoreNewStatus) const;
+
+	UFUNCTION(BlueprintCallable, Category = "Material Designer")
+	static bool ExportMaterialInstance(UDynamicMaterialModel* InMaterialModel, const FString& InSavePath);
+
+	UFUNCTION(BlueprintCallable, Category = "Material Designer")
+	static bool ExportGeneratedMaterial(UDynamicMaterialModel* InMaterialModel, const FString& InSavePath);
+};
diff --git a/Engine/Plugins/Experimental/DynamicMaterial/Source/DynamicMaterialEditor/Public/DMObjectMaterialProperty.h b/Engine/Plugins/Experimental/DynamicMaterial/Source/DynamicMaterialEditor/Public/DMObjectMaterialProperty.h
new file mode 100644
--- /dev/null
+++ b/Engine/Plugins/Experimental/DynamicMaterial/Source/DynamicMaterialEditor/Public/DMObjectMaterialProperty.h
@@ -0,0 +1,57 @@
+// Copyright Epic Games, Inc. All Rights Reserved.
+
+#pragma once
+
+#include "UObject/WeakObjectPtr.h"
+#include "UObject/WeakObjectPtrTemplates.h"
+#include "UObject/NameTypes.h"
+#include "UObject/Object.h"
+#include "DMObjectMaterialProperty.generated.h"
+
+class FProperty;
+class UDynamicMaterialInstance;
+class UDynamicMaterialModel;
+class UObject;
+class UPrimitiveComponent;
+
+/**
+ * Defines a material property slot that can be a Material Designer Instance.
+ */
+USTRUCT(BlueprintType)
+struct DYNAMICMATERIALEDITOR_API FDMObjectMaterialProperty
+{
+	GENERATED_BODY()
+
+	FDMObjectMaterialProperty();
+
+	// UPrimitiveComponent Material Index
+	FDMObjectMaterialProperty(UPrimitiveComponent* InOuter, int32 InIndex);
+
+	// Class Property (including potential array index)
+	FDMObjectMaterialProperty(UObject* InOuter, FProperty* InProperty, int32 InIndex = INDEX_NONE);
+
+	UPROPERTY(VisibleInstanceOnly, BlueprintReadOnly, Category = "Material")
+	TWeakObjectPtr<UObject> OuterWeak;
+
+	/** C++ version of property */
+	FProperty* Property;
+
+	/** Blueprint version of property */
+	UPROPERTY(VisibleInstanceOnly, BlueprintReadOnly, Category = "Material")
+	FName PropertyName;
+
+	/** Component or array property index. */
+	UPROPERTY(VisibleInstanceOnly, BlueprintReadOnly, Category = "Material")
+	int32 Index;
+
+	UDynamicMaterialModel* GetMaterialModel() const;
+
+	UDynamicMaterialInstance* GetMaterial() const;
+	void SetMaterial(UDynamicMaterialInstance* DynamicMaterial);
+
+	bool IsValid() const;
+
+	FText GetPropertyName(bool bInIgnoreNewStatus) const;
+
+	void Reset();
+};
diff --git a/Engine/Plugins/Experimental/DynamicMaterial/Source/DynamicMaterialEditor/Public/DMWorldSubsystem.h b/Engine/Plugins/Experimental/DynamicMaterial/Source/DynamicMaterialEditor/Public/DMWorldSubsystem.h
new file mode 100644
--- /dev/null
+++ b/Engine/Plugins/Experimental/DynamicMaterial/Source/DynamicMaterialEditor/Public/DMWorldSubsystem.h
@@ -0,0 +1,57 @@
+// Copyright Epic Games, Inc. All Rights Reserved.
+ 
+#pragma once
+
+#include "DynamicMaterialEditorModule.h"
+#include "Subsystems/WorldSubsystem.h"
+#include "Templates/SharedPointer.h"
+#include "DMWorldSubsystem.generated.h"
+
+class AActor;
+class IDetailKeyframeHandler;
+class UDynamicMaterialInstance;
+class UDynamicMaterialModel;
+struct FDMObjectMaterialProperty;
+
+DECLARE_DELEGATE_RetVal_OneParam(bool, FDMIsValidDelegate, UDynamicMaterialModel*)
+DECLARE_DELEGATE_RetVal_TwoParams(bool, FDMSetMaterialValueDelegate, const FDMObjectMaterialProperty&, UDynamicMaterialInstance*)
+DECLARE_DELEGATE(FDMInvokeTabDelegate)
+ 
+UCLASS()
+class DYNAMICMATERIALEDITOR_API UDMWorldSubsystem : public UWorldSubsystem
+{
+	GENERATED_BODY()
+ 
+public:
+	UDMWorldSubsystem();
+ 
+	const TSharedPtr<IDetailKeyframeHandler>& GetKeyframeHandler() const { return KeyframeHandler; }
+	void SetKeyframeHandler(const TSharedPtr<IDetailKeyframeHandler>& InKeyframeHandler);
+
+	/** Sets material in a custom editor tab. */
+	FDMSetMaterialModelDelegate& GetSetCustomEditorModelDelegate() { return CustomModelEditorDelegate; }
+
+	/** Sets the object property in custom editor tab. */
+	FDMSetMaterialObjectPropertyDelegate& GetCustomObjectPropertyEditorDelegate() { return CustomObjectPropertyEditorDelegate; }
+
+	/** Sets actor in a custom editor tab. */
+	FDMSetMaterialActorDelegate& GetSetCustomEditorActorDelegate() { return CustomActorEditorDelegate; }
+
+	/** Returns true if the supplied material is valid for this world. */
+	FDMIsValidDelegate& GetIsValidDelegate() { return IsValidDelegate; }
+
+	/** Used to redirect SetMaterial to different objects/paths. */
+	FDMSetMaterialValueDelegate& GetMaterialValueSetterDelegate() { return SetMaterialValueDelegate; }
+
+	/** Used to show the tab to the user. */
+	FDMInvokeTabDelegate& GetInvokeTabDelegate() { return InvokeTabDelegate; }
+
+protected:
+	TSharedPtr<IDetailKeyframeHandler> KeyframeHandler;
+	FDMSetMaterialModelDelegate CustomModelEditorDelegate;
+	FDMSetMaterialObjectPropertyDelegate CustomObjectPropertyEditorDelegate;
+	FDMSetMaterialActorDelegate CustomActorEditorDelegate;
+	FDMIsValidDelegate IsValidDelegate;
+	FDMSetMaterialValueDelegate SetMaterialValueDelegate;
+	FDMInvokeTabDelegate InvokeTabDelegate;
+};
diff --git a/Engine/Plugins/Experimental/DynamicMaterial/Source/DynamicMaterialEditor/Public/Material/DynamicMaterialInstanceFactory.h b/Engine/Plugins/Experimental/DynamicMaterial/Source/DynamicMaterialEditor/Public/Material/DynamicMaterialInstanceFactory.h
new file mode 100644
--- /dev/null
+++ b/Engine/Plugins/Experimental/DynamicMaterial/Source/DynamicMaterialEditor/Public/Material/DynamicMaterialInstanceFactory.h
@@ -0,0 +1,25 @@
+// Copyright Epic Games, Inc. All Rights Reserved.
+
+#pragma once
+
+#include "Factories/Factory.h"
+#include "Containers/UnrealString.h"
+#include "DynamicMaterialInstanceFactory.generated.h"
+
+class UDynamicMaterialInstance;
+
+UCLASS(MinimalAPI, hidecategories = Object)
+class UDynamicMaterialInstanceFactory : public UFactory
+{
+	GENERATED_BODY()
+
+public:
+	static const FString BaseDirectory;
+	static const FString BaseName;
+
+	UDynamicMaterialInstanceFactory();
+
+	virtual UObject* FactoryCreateNew(UClass* Class, UObject* InParent, FName Name, EObjectFlags Flags, UObject* Context, FFeedbackContext* Warn) override;
+	virtual FText GetDisplayName() const override;
+	virtual FText GetToolTip() const override;
+};
diff --git a/Engine/Plugins/Experimental/DynamicMaterial/Source/DynamicMaterialEditor/Public/Model/DMMaterialBuildState.h b/Engine/Plugins/Experimental/DynamicMaterial/Source/DynamicMaterialEditor/Public/Model/DMMaterialBuildState.h
new file mode 100644
--- /dev/null
+++ b/Engine/Plugins/Experimental/DynamicMaterial/Source/DynamicMaterialEditor/Public/Model/DMMaterialBuildState.h
@@ -0,0 +1,163 @@
+// Copyright Epic Games, Inc. All Rights Reserved.
+
+#pragma once
+
+#include "Model/IDMMaterialBuildStateInterface.h"
+#include "Containers/Array.h"
+#include "Containers/Map.h"
+#include "Containers/Set.h"
+#include "Model/DMMaterialBuildUtils.h"
+
+class UDMMaterialLayerObject;
+class UDMMaterialSlot;
+class UDMMaterialStage;
+class UDMMaterialStageInput;
+class UDMMaterialStageSource;
+class UDMMaterialStageThroughput;
+class UDynamicMaterialModel;
+class UMaterial;
+class UMaterialExpression;
+class UMaterialExpressionAppendVector;
+class UMaterialExpressionComponentMask;
+struct FDMExpressionInput;
+struct FDMMaterialBuildUtils;
+struct FDMMaterialStageConnection;
+
+/**
+ * BuildState is a class that stores the current state of a material that is being built.
+ * It stores useful lists of UMaterialExpressions relating to various object within the builder, such a Stages or Sources.
+ * It is an entirely transient object. It is not meant to be saved outside of the material building processs.
+ * It also provides some helper functions for creating UMaterialExpressions.
+ */
+struct DYNAMICMATERIALEDITOR_API FDMMaterialBuildState : public IDMMaterialBuildStateInterface, public TSharedFromThis<FDMMaterialBuildState>
+{
+	FDMMaterialBuildState(UMaterial* InDynamicMaterial, UDynamicMaterialModel* InMaterialModel, bool bInDirtyAssets = true);
+
+	virtual ~FDMMaterialBuildState() override;
+
+	virtual UMaterial* GetDynamicMaterial() const override;
+
+	virtual UDynamicMaterialModel* GetMaterialModel() const override;
+
+	bool ShouldDirtyAssets() const { return bDirtyAssets; }
+
+	void SetIgnoreUVs();
+
+	bool IsIgnoringUVs() const { return bIgnoreUVs; }
+
+	void SetPreviewMaterial();
+
+	bool IsPreviewMaterial() const { return bIsPreviewMaterial; }
+
+	virtual IDMMaterialBuildUtilsInterface& GetBuildUtils() const override;
+
+	FExpressionInput* GetMaterialProperty(EDMMaterialPropertyType InProperty) const;
+
+	/** Slots */
+	bool HasSlot(const UDMMaterialSlot* InSlot) const;
+
+	const TArray<UMaterialExpression*>& GetSlotExpressions(const UDMMaterialSlot* InSlot) const;
+
+	UMaterialExpression* GetLastSlotExpression(const UDMMaterialSlot* InSlot) const;
+
+	void AddSlotExpressions(const UDMMaterialSlot* InSlot, const TArray<UMaterialExpression*>& InSlotExpressions);
+
+	bool HasSlotProperties(const UDMMaterialSlot* InSlot) const;
+
+	void AddSlotPropertyExpressions(const UDMMaterialSlot* InSlot, const TMap<EDMMaterialPropertyType, 
+		TArray<UMaterialExpression*>>& InSlotPropertyExpressions);
+
+	const TMap<EDMMaterialPropertyType, TArray<UMaterialExpression*>>& GetSlotPropertyExpressions(const UDMMaterialSlot* InSlot);
+
+	UMaterialExpression* GetLastSlotPropertyExpression(const UDMMaterialSlot* InSlot, EDMMaterialPropertyType InMaterialProperty) const;
+
+	TArray<const UDMMaterialSlot*> GetSlots() const;
+
+	const TMap<const UDMMaterialSlot*, TArray<UMaterialExpression*>>& GetSlotMap() const;
+
+	/** Layers */
+	bool HasLayer(const UDMMaterialLayerObject* InLayer) const;
+
+	const TArray<UMaterialExpression*>& GetLayerExpressions(const UDMMaterialLayerObject* InLayer) const;
+
+	UMaterialExpression* GetLastLayerExpression(const UDMMaterialLayerObject* InLayer) const;
+
+	void AddLayerExpressions(const UDMMaterialLayerObject* InLayer, const TArray<UMaterialExpression*>& InLayerExpressions);
+
+	TArray<const UDMMaterialLayerObject*> GetLayers() const;
+
+	const TMap<const UDMMaterialLayerObject*, TArray<UMaterialExpression*>>& GetLayerMap() const;
+
+	/** Stages */
+	bool HasStage(const UDMMaterialStage* InStage) const;
+
+	const TArray<UMaterialExpression*>& GetStageExpressions(const UDMMaterialStage* InStage) const;
+
+	UMaterialExpression* GetLastStageExpression(const UDMMaterialStage* InStage) const;
+
+	void AddStageExpressions(const UDMMaterialStage* InStage, const TArray<UMaterialExpression*>& InStageExpressions);
+
+	TArray<const UDMMaterialStage*> GetStages() const;
+
+	const TMap<const UDMMaterialStage*, TArray<UMaterialExpression*>>& GetStageMap() const;
+
+	/** Stage Sources */
+	bool HasStageSource(const UDMMaterialStageSource* InStageSource) const;
+
+	const TArray<UMaterialExpression*>& GetStageSourceExpressions(const UDMMaterialStageSource* InStageSource) const;
+
+	UMaterialExpression* GetLastStageSourceExpression(const UDMMaterialStageSource* InStageSource) const;
+
+	void AddStageSourceExpressions(const UDMMaterialStageSource* InStageSource, 
+		const TArray<UMaterialExpression*>& InStageSourceExpressions);
+
+	TArray<const UDMMaterialStageSource*> GetStageSources() const;
+
+	const TMap<const UDMMaterialStageSource*, TArray<UMaterialExpression*>>& GetStageSourceMap() const;
+
+	/** Material Values */
+	virtual bool HasValue(const UDMMaterialValue* InValue) const override;
+
+	virtual const TArray<UMaterialExpression*>& GetValueExpressions(const UDMMaterialValue* InValue) const override;
+
+	virtual UMaterialExpression* GetLastValueExpression(const UDMMaterialValue* InValue) const override;
+
+	virtual void AddValueExpressions(const UDMMaterialValue* InValue, const TArray<UMaterialExpression*>& InValueExpressions) override;
+
+	virtual TArray<const UDMMaterialValue*> GetValues() const override;
+
+	virtual const TMap<const UDMMaterialValue*, TArray<UMaterialExpression*>>& GetValueMap() const override;
+
+	/** Callbacks */
+	bool HasCallback(const UDMMaterialStageSource* InCallback) const;
+
+	const TArray<UDMMaterialStageSource*>& GetCallbackExpressions(const UDMMaterialStageSource* InCallback) const;
+
+	void AddCallbackExpressions(const UDMMaterialStageSource* InCallback, const TArray<UDMMaterialStageSource*>& InCallbackExpressions);
+
+	TArray<const UDMMaterialStageSource*> GetCallbacks() const;
+
+	const TMap<const UDMMaterialStageSource*, TArray<UDMMaterialStageSource*>>& GetCallbackMap() const;
+
+	/** Other */
+	virtual void AddOtherExpressions(const TArray<UMaterialExpression*>& InOtherExpressions) override;
+
+	virtual const TSet<UMaterialExpression*>& GetOtherExpressions() override;
+
+private:
+	UMaterial* DynamicMaterial = nullptr;
+	UDynamicMaterialModel* MaterialModel = nullptr;
+	bool bDirtyAssets;
+	bool bIgnoreUVs;
+	bool bIsPreviewMaterial;
+	TSharedRef<FDMMaterialBuildUtils> Utils;
+
+	TMap<const UDMMaterialValue*, TArray<UMaterialExpression*>> Values;
+	TMap<const UDMMaterialSlot*, TArray<UMaterialExpression*>> Slots;
+	TMap<const UDMMaterialSlot*, TMap<EDMMaterialPropertyType, TArray<UMaterialExpression*>>> SlotProperties;
+	TMap<const UDMMaterialLayerObject*, TArray<UMaterialExpression*>> Layers;
+	TMap<const UDMMaterialStage*, TArray<UMaterialExpression*>> Stages;
+	TMap<const UDMMaterialStageSource*, TArray<UMaterialExpression*>> StageSources;
+	TMap<const UDMMaterialStageSource*, TArray<UDMMaterialStageSource*>> Callbacks;
+	TSet<UMaterialExpression*> OtherExpressions;
+};
diff --git a/Engine/Plugins/Experimental/DynamicMaterial/Source/DynamicMaterialEditor/Public/Model/DMMaterialBuildUtils.h b/Engine/Plugins/Experimental/DynamicMaterial/Source/DynamicMaterialEditor/Public/Model/DMMaterialBuildUtils.h
new file mode 100644
--- /dev/null
+++ b/Engine/Plugins/Experimental/DynamicMaterial/Source/DynamicMaterialEditor/Public/Model/DMMaterialBuildUtils.h
@@ -0,0 +1,52 @@
+// Copyright Epic Games, Inc. All Rights Reserved.
+
+#pragma once
+
+#include "Model/IDMMaterialBuildUtilsInterface.h"
+
+struct FDMMaterialBuildState;
+
+struct DYNAMICMATERIALEDITOR_API FDMMaterialBuildUtils : public IDMMaterialBuildUtilsInterface
+{
+	FDMMaterialBuildUtils(FDMMaterialBuildState& InBuildState);
+
+	UMaterialExpression* CreateDefaultExpression() const;
+
+	virtual UMaterialExpression* CreateExpression(TSubclassOf<UMaterialExpression> InExpressionClass, const FString& InComment,
+		UObject* InAsset = nullptr) const override;
+
+	virtual UMaterialExpression* CreateExpressionParameter(TSubclassOf<UMaterialExpression> InExpressionClass, FName InParameterName,
+		const FString& InComment, UObject* InAsset = nullptr) const override;
+
+	template<typename InExpressionClass
+		UE_REQUIRES(std::derived_from<InExpressionClass, UMaterialExpression>)>
+	InExpressionClass* CreateExpression(const FString& InComment, UObject* InAsset = nullptr) const
+	{
+		return Cast<InExpressionClass>(CreateExpression(TSubclassOf<UMaterialExpression>(InExpressionClass::StaticClass()), InComment, InAsset));
+	}
+
+	template<typename InExpressionClass
+		UE_REQUIRES(std::derived_from<InExpressionClass, UMaterialExpression>)>
+	InExpressionClass* CreateExpressionParameter(FName InParameterName, const FString& InComment, UObject* InAsset = nullptr) const
+	{
+		return Cast<InExpressionClass>(CreateExpressionParameter(TSubclassOf<UMaterialExpression>(InExpressionClass::StaticClass()), InParameterName, InComment, InAsset));
+	}
+
+	virtual TArray<UMaterialExpression*> CreateExpressionInputs(const TArray<FDMMaterialStageConnection>& InInputConnectionMap,
+		int32 InStageSourceInputIdx, const TArray<UDMMaterialStageInput*>& InStageInputs, int32& OutOutputIndex, 
+		int32& OutOutputChannel) const override;
+
+	virtual TArray<UMaterialExpression*> CreateExpressionInput(UDMMaterialStageInput* InInput) const override;
+
+	virtual UMaterialExpressionComponentMask* CreateExpressionBitMask(UMaterialExpression* InExpression, int32 InOutputIndex,
+		int32 InOutputChannels) const override;
+
+	virtual UMaterialExpressionAppendVector* CreateExpressionAppend(UMaterialExpression* InExpressionA, int32 InOutputIndexA, 
+		UMaterialExpression* InExpressionB, int32 InOutputIndexB) const override;
+
+	virtual void UpdatePreviewMaterial(UMaterialExpression* InLastExpression, int32 InOutputIndex, int32 InOutputChannel,
+		int32 InSize) const override;
+
+protected:
+	FDMMaterialBuildState& BuildState;
+};
diff --git a/Engine/Plugins/Experimental/DynamicMaterial/Source/DynamicMaterialEditor/Public/Model/DynamicMaterialModelEditorOnlyData.h b/Engine/Plugins/Experimental/DynamicMaterial/Source/DynamicMaterialEditor/Public/Model/DynamicMaterialModelEditorOnlyData.h
new file mode 100644
--- /dev/null
+++ b/Engine/Plugins/Experimental/DynamicMaterial/Source/DynamicMaterialEditor/Public/Model/DynamicMaterialModelEditorOnlyData.h
@@ -0,0 +1,228 @@
+// Copyright Epic Games, Inc. All Rights Reserved.
+
+#pragma once
+
+#include "Model/IDynamicMaterialModelEditorOnlyDataInterface.h"
+#include "DMEDefs.h"
+#include "Engine/EngineTypes.h"
+#include "MaterialDomain.h"
+#include "Misc/NotifyHook.h"
+#include "UObject/WeakObjectPtrFwd.h"
+#include "UObject/WeakObjectPtrTemplatesFwd.h"
+#include "DynamicMaterialModelEditorOnlyData.generated.h"
+
+class UDMMaterialComponent;
+class UDMMaterialParameter;
+class UDMMaterialProperty;
+class UDMMaterialSlot;
+class UDMMaterialValue;
+class UDMMaterialValueFloat1;
+class UDMTextureUV;
+class UDynamicMaterialModel;
+class UDynamicMaterialModelEditorOnlyData;
+class UMaterialExpression;
+struct FDMComponentPathSegment;
+struct FDMComponentPath;
+struct FDMMaterialBuildState;
+
+DECLARE_MULTICAST_DELEGATE_OneParam(FDMOnMaterialBuilt, UDynamicMaterialModel*);
+DECLARE_MULTICAST_DELEGATE_OneParam(FDMOnValueListUpdated, UDynamicMaterialModel*);
+DECLARE_MULTICAST_DELEGATE_TwoParams(FDMOnValueUpdated, UDynamicMaterialModel*, UDMMaterialValue*);
+DECLARE_MULTICAST_DELEGATE_OneParam(FDMOnSlotListUpdated, UDynamicMaterialModel*);
+DECLARE_MULTICAST_DELEGATE_TwoParams(FDMOnTextureUVUpdated, UDynamicMaterialModel*, UDMTextureUV*);
+
+UENUM(BlueprintType)
+enum class EDMState : uint8
+{
+	Idle,
+	Building
+};
+
+UCLASS(BlueprintType, EditInlineNew, DefaultToInstanced, ClassGroup = "Material Designer")
+class DYNAMICMATERIALEDITOR_API UDynamicMaterialModelEditorOnlyData : public UObject, public IDynamicMaterialModelEditorOnlyDataInterface, public FNotifyHook, public IDMBuildable
+{
+	GENERATED_BODY()
+
+	friend class UDynamicMaterialModelFactory;
+
+public:
+	static const FString SlotsPathToken;
+	static const FString RGBSlotPathToken; // @TODO This will probably need to change when other slot types are opened up.
+	static const FString OpacitySlotPathToken; // @TODO This will probably need to change when other slot types are opened up.
+	static const FString PropertiesPathToken;
+
+	static UDynamicMaterialModelEditorOnlyData* Get(UDynamicMaterialModel* InModel);
+	static UDynamicMaterialModelEditorOnlyData* Get(TWeakObjectPtr<UDynamicMaterialModel> InModelWeak);
+	static UDynamicMaterialModelEditorOnlyData* Get(const TScriptInterface<IDynamicMaterialModelEditorOnlyDataInterface>& InInterface);
+	static UDynamicMaterialModelEditorOnlyData* Get(IDynamicMaterialModelEditorOnlyDataInterface* InInterface);
+
+	UDynamicMaterialModelEditorOnlyData();
+
+	UFUNCTION(BlueprintPure, Category = "Material Designer")
+	UDynamicMaterialModel* GetMaterialModel() const { return MaterialModel; }
+
+	UFUNCTION(BlueprintPure, Category = "Material Designer")
+	UMaterial* GetGeneratedMaterial() const;
+
+	UFUNCTION(BlueprintPure, Category = "Material Designer")
+	EDMState GetState() const { return State; }
+
+	UFUNCTION(BlueprintPure, Category = "Material Designer")
+	TEnumAsByte<EMaterialDomain> GetDomain() const { return Domain; }
+
+	UFUNCTION(BlueprintCallable, Category = "Material Designer")
+	void SetDomain(TEnumAsByte<EMaterialDomain> InDomain);
+
+	UFUNCTION(BlueprintPure, Category = "Material Designer")
+	TEnumAsByte<EBlendMode> GetBlendMode() const { return BlendMode; }
+
+	UFUNCTION(BlueprintCallable, Category = "Material Designer")
+	void SetBlendMode(TEnumAsByte<EBlendMode> InBlendMode);
+
+	UFUNCTION(BlueprintPure, Category = "Material Designer")
+	EDMMaterialShadingModel GetShadingModel() const { return ShadingModel; }
+
+	UFUNCTION(BlueprintCallable, Category = "Material Designer")
+	void SetShadingModel(EDMMaterialShadingModel InShadingModel);
+
+	UFUNCTION(BlueprintCallable, Category = "Material Designer")
+	void OpenMaterialEditor() const;
+
+	UFUNCTION(BlueprintPure, Category = "Material Designer")
+	TMap<EDMMaterialPropertyType, UDMMaterialProperty*> GetMaterialProperties() const;
+
+	UFUNCTION(BlueprintPure, Category = "Material Designer")
+	UDMMaterialProperty* GetMaterialProperty(EDMMaterialPropertyType MaterialProperty) const;
+
+	UFUNCTION(BlueprintPure, Category = "Material Designer")
+	const TArray<UDMMaterialSlot*>& GetSlots() const { return Slots; }
+
+	UFUNCTION(BlueprintPure, Category = "Material Designer")
+	UDMMaterialSlot* GetSlot(int32 Index) const;
+
+	UFUNCTION(BlueprintCallable, Category = "Material Designer")
+	UDMMaterialSlot* AddSlot();
+
+	UFUNCTION(BlueprintCallable, Category = "Material Designer")
+	void RemoveSlot(int32 Index);
+
+	UFUNCTION(BlueprintPure, Category = "Material Designer")
+	UDMMaterialSlot* GetSlotForMaterialProperty(EDMMaterialPropertyType Property) const;
+
+	UFUNCTION(BlueprintPure, Category = "Material Designer")
+	TArray<EDMMaterialPropertyType> GetMaterialPropertiesForSlot(UDMMaterialSlot* Slot) const;
+
+	UFUNCTION(BlueprintCallable, Category = "Material Designer")
+	void AssignMaterialPropertyToSlot(EDMMaterialPropertyType Property, UDMMaterialSlot* Slot);
+
+	UFUNCTION(BlueprintCallable, Category = "Material Designer")
+	void UnassignMaterialProperty(EDMMaterialPropertyType Property);
+
+	UFUNCTION(BlueprintPure, Category = "Material Designer")
+	bool IsPixelAnimationFlagSet() const;
+
+	UFUNCTION(BlueprintCallable, Category = "Material Designer")
+	void SetPixelAnimationFlag(bool bInFlagValue);
+	
+	FDMOnMaterialBuilt& GetOnMaterialBuiltDelegate() { return OnMaterialBuiltDelegate; }
+	FDMOnValueListUpdated& GetOnValueListUpdateDelegate() { return OnValueListUpdateDelegate; }
+	FDMOnValueUpdated& GetOnValueUpdateDelegate() { return OnValueUpdateDelegate; }
+	FDMOnSlotListUpdated& GetOnSlotListUpdateDelegate() { return OnSlotListUpdateDelegate; }
+	FDMOnTextureUVUpdated& GetOnTextureUVUpdateDelegate() { return OnTextureUVUpdateDelegate; }
+
+	void GenerateOpacityExpressions(const TSharedRef<FDMMaterialBuildState>& InBuildState, UDMMaterialSlot* InFromSlot,
+		EDMMaterialPropertyType InFromProperty, UMaterialExpression*& OutExpression, int32& OutOutputIndex, int32& OutOutputChannel) const;
+
+	TSharedRef<FDMMaterialBuildState> CreateBuildState(UMaterial* InMaterialToBuild, bool bInDirtyAssets = true) const;
+
+	//~ Begin FNotifyHook
+	virtual void NotifyPreChange(class FEditPropertyChain* PropertyAboutToChange) {}
+	virtual void NotifyPostChange(const FPropertyChangedEvent& PropertyChangedEvent, class FEditPropertyChain* PropertyThatChanged);
+	//~ End FNotifyHook
+
+	//~ Begin UObject
+	virtual void PostLoad() override;
+	virtual void PostEditUndo() override;
+	virtual void PostEditImport() override;
+	virtual void PostDuplicate(bool bDuplicateForPIE) override;
+	//~ End UObject
+
+	void SaveEditor();
+
+	//~ Begin IDMBuildable
+	virtual void DoBuild_Implementation(bool bInDirtyAssets) override;
+	//~ End IDMBuildable
+
+	//~ Begin IDynamicMaterialModelEditorOnlyDataInterface
+	virtual void PostEditorDuplicate() override;
+	virtual void RequestMaterialBuild() override;
+	virtual void OnValueUpdated(UDMMaterialValue* InValue, EDMUpdateType InUpdateType) override;
+	virtual void OnValueListUpdate() override;
+	virtual void OnTextureUVUpdated(UDMTextureUV* InTextureUV) override;
+	virtual void LoadDeprecatedModelData(UDynamicMaterialModel* InMaterialModel) override;
+	virtual TSharedRef<IDMMaterialBuildStateInterface> CreateBuildStateInterface(UMaterial* InMaterialToBuild) const override;
+	virtual UDMMaterialComponent* GetSubComponentByPath(FDMComponentPath& InPath, const FDMComponentPathSegment& InPathSegment) const override;
+	//~ End IDynamicMaterialModelEditorOnlyDataInterface
+
+protected:
+	UPROPERTY(VisibleInstanceOnly, BlueprintReadOnly, Category = "Material Designer")
+	TObjectPtr<UDynamicMaterialModel> MaterialModel;
+
+	UPROPERTY(VisibleInstanceOnly, BlueprintReadOnly, Transient, DuplicateTransient, TextExportTransient, Category = "Material Designer")
+	EDMState State = EDMState::Idle;
+
+	UPROPERTY(EditInstanceOnly, BlueprintReadOnly, Category = "Material Designer")
+	TEnumAsByte<EMaterialDomain> Domain = EMaterialDomain::MD_Surface;
+
+	UPROPERTY(EditInstanceOnly, BlueprintReadOnly, Category = "Material Designer")
+	TEnumAsByte<EBlendMode> BlendMode = EBlendMode::BLEND_Translucent;
+
+	UPROPERTY(EditInstanceOnly, BlueprintReadOnly, Category = "Material Designer")
+	EDMMaterialShadingModel ShadingModel = EDMMaterialShadingModel::Unlit;
+
+	UPROPERTY(VisibleInstanceOnly, BlueprintReadOnly, Instanced, Category = "Material Designer")
+	TMap<EDMMaterialPropertyType, TObjectPtr<UDMMaterialProperty>> Properties;
+
+	UPROPERTY(VisibleInstanceOnly, BlueprintReadOnly, DuplicateTransient, TextExportTransient, Category = "Material Designer")
+	TMap<EDMMaterialPropertyType, TObjectPtr<UDMMaterialSlot>> PropertySlotMap;
+
+	UPROPERTY(VisibleInstanceOnly, BlueprintReadOnly, Instanced, Category = "Material Designer")
+	TArray<TObjectPtr<UDMMaterialSlot>> Slots;
+
+	UPROPERTY(VisibleInstanceOnly, BlueprintReadOnly, Instanced, Category = "Material Designer")
+	TArray<TObjectPtr<UMaterialExpression>> Expressions;
+
+	UPROPERTY(VisibleInstanceOnly, BlueprintReadOnly, Category = "Material Designer")
+	bool bCreateMaterialPackage;
+
+	UPROPERTY(VisibleInstanceOnly, BlueprintReadOnly, Category = "Material Designer")
+	bool bPixelAnimationFlag;
+
+	FDMOnMaterialBuilt OnMaterialBuiltDelegate;
+	FDMOnValueListUpdated OnValueListUpdateDelegate;
+	FDMOnValueUpdated OnValueUpdateDelegate;
+	FDMOnSlotListUpdated OnSlotListUpdateDelegate;
+	FDMOnTextureUVUpdated OnTextureUVUpdateDelegate;
+
+	void CreateMaterial();
+
+	UFUNCTION(BlueprintCallable, Category = "Material Designer")
+	void BuildMaterial(bool bInDirtyAssets);
+
+	FString GetMaterialAssetPath() const;
+	FString GetMaterialAssetName() const;
+	FString GetMaterialPackageName(const FString& MaterialBaseName) const;
+
+	void OnSlotConnectorsUpdated(UDMMaterialSlot* Slot);
+
+	//~ Begin IDynamicMaterialModelEditorOnlyDataInterface
+	virtual void ReinitComponents() override;
+	virtual void ResetData() override;
+	//~ End IDynamicMaterialModelEditorOnlyDataInterface
+
+	void LoadDeprecatedModelData_Base(bool bInCreateMaterialPackage, EBlendMode InBlendMode, EDMMaterialShadingModel InShadingModel);
+	void LoadDeprecatedModelData_Expressions(TArray<TObjectPtr<UMaterialExpression>>& InExpressions);
+	void LoadDeprecatedModelData_Properties(TMap<EDMMaterialPropertyType, TObjectPtr<UObject>>& InProperties);
+	void LoadDeprecatedModelData_Slots(TArray<TObjectPtr<UObject>>& InSlots);
+	void LoadDeprecatedModelData_PropertySlotMap(TMap<EDMMaterialPropertyType, TObjectPtr<UObject>>& InPropertySlotMap);
+};
diff --git a/Engine/Plugins/Experimental/DynamicMaterial/Source/DynamicMaterialEditor/Public/Slate/Properties/Generators/DMComponentPropertyRowGenerator.h b/Engine/Plugins/Experimental/DynamicMaterial/Source/DynamicMaterialEditor/Public/Slate/Properties/Generators/DMComponentPropertyRowGenerator.h
new file mode 100644
--- /dev/null
+++ b/Engine/Plugins/Experimental/DynamicMaterial/Source/DynamicMaterialEditor/Public/Slate/Properties/Generators/DMComponentPropertyRowGenerator.h
@@ -0,0 +1,33 @@
+// Copyright Epic Games, Inc. All Rights Reserved.
+
+#pragma once
+
+#include "Containers/Array.h"
+#include "Containers/Set.h"
+#include "Templates/SharedPointer.h"
+
+class FName;
+class FProperty;
+class SDMComponentEdit;
+class UDMMaterialComponent;
+struct FDMPropertyHandle;
+
+class DYNAMICMATERIALEDITOR_API FDMComponentPropertyRowGenerator
+{
+public:
+	static const TSharedRef<FDMComponentPropertyRowGenerator>& Get();
+
+	FDMComponentPropertyRowGenerator() = default;
+	virtual ~FDMComponentPropertyRowGenerator() = default;
+
+	virtual void AddComponentProperties(const TSharedRef<SDMComponentEdit>& InComponentEditWidget, UDMMaterialComponent* InComponent, 
+		TArray<FDMPropertyHandle>& InOutPropertyRows, TSet<UDMMaterialComponent*>& InOutProcessedObjects);
+
+	virtual void AddPropertyEditRows(const TSharedRef<SDMComponentEdit>& InComponentEditWidget, UDMMaterialComponent* InComponent, 
+		const FName& InProperty, TArray<FDMPropertyHandle>& InOutPropertyRows, TSet<UDMMaterialComponent*>& InOutProcessedObjects);
+
+	virtual void AddPropertyEditRows(const TSharedRef<SDMComponentEdit>& InComponentEditWidget, UDMMaterialComponent* InComponent, FProperty* InProperty, 
+		void* MemoryPtr, TArray<FDMPropertyHandle>& InOutPropertyRows, TSet<UDMMaterialComponent*>& InOutProcessedObjects);
+
+	virtual bool AllowKeyframeButton(UDMMaterialComponent* InComponent, FProperty* InProperty);
+};
