Fix duplicated loctext keys

diff --git a/Engine/Plugins/Experimental/ActorModifierCore/Source/ActorModifierCore/Private/Modifiers/ActorModifierCoreStack.cpp b/Engine/Plugins/Experimental/ActorModifierCore/Source/ActorModifierCore/Private/Modifiers/ActorModifierCoreStack.cpp
--- a/Engine/Plugins/Experimental/ActorModifierCore/Source/ActorModifierCore/Private/Modifiers/ActorModifierCoreStack.cpp
+++ b/Engine/Plugins/Experimental/ActorModifierCore/Source/ActorModifierCore/Private/Modifiers/ActorModifierCoreStack.cpp
@@ -382,7 +382,7 @@ UActorModifierCoreBase* UActorModifierCoreStack::CloneModifier(FActorModifierCor
 		if (InCloneOp.FailReason)
 		{
 			*InCloneOp.FailReason = FText::Format(
-				LOCTEXT("InsertModifierNotFound", "Could not clone modifier {0} in this stack"),
+				LOCTEXT("CloneModifierNotFound", "Could not clone modifier {0} in this stack"),
 				FText::FromName(InCloneOp.CloneModifier->GetModifierName()));
 		}
 
@@ -719,7 +719,7 @@ bool UActorModifierCoreStack::MoveModifier(FActorModifierCoreStackMoveOp& InMove
 		if (InMoveOp.FailReason)
 		{
 			*InMoveOp.FailReason = FText::Format(
-			LOCTEXT("MoveModifierNotFound", "Could not move modifier {0} in this stack"),
+			LOCTEXT("MoveSourceModifierNotFound", "Could not move modifier {0} in this stack"),
 				FText::FromName(InMoveOp.MoveModifier->GetModifierName()));
 		}
 
@@ -770,7 +770,7 @@ bool UActorModifierCoreStack::MoveModifier(FActorModifierCoreStackMoveOp& InMove
 		if (InMoveOp.FailReason)
 		{
 			*InMoveOp.FailReason = FText::Format(
-				LOCTEXT("MoveModifierNotFound", "Could not move modifier {0} in this stack"),
+				LOCTEXT("MoveTargetModifierNotFound", "Could not move modifier {0} in this stack"),
 				FText::FromName(InMoveOp.MoveModifier->GetModifierName()));
 		}
 
@@ -1387,7 +1387,7 @@ void UActorModifierCoreStack::CheckModifierOptimization(bool bInInvalidateAll)
 				{
 					UnoptimizedModifiers.Add(Modifier);
 
-					const FText Message = FText::Format(LOCTEXT("AvoidBeforeCategory", "Should be moved above {0} modifiers"), FText::FromName(Category));
+					const FText Message = FText::Format(LOCTEXT("AvoidAfterCategory", "Should be moved above {0} modifiers"), FText::FromName(Category));
 					Modifier->Status = FActorModifierCoreStatus(EActorModifierCoreStatus::Warning, Message);
 
 					LogModifier(FString::Printf(TEXT("Optimisation possible for modifier %s : %s"), *Modifier->GetModifierName().ToString(), *Message.ToString()), true);
diff --git a/Engine/Plugins/Experimental/ActorModifierCore/Source/ActorModifierCoreEditor/Private/Modifiers/ActorModifierCoreEditorMenu.cpp b/Engine/Plugins/Experimental/ActorModifierCore/Source/ActorModifierCoreEditor/Private/Modifiers/ActorModifierCoreEditorMenu.cpp
--- a/Engine/Plugins/Experimental/ActorModifierCore/Source/ActorModifierCoreEditor/Private/Modifiers/ActorModifierCoreEditorMenu.cpp
+++ b/Engine/Plugins/Experimental/ActorModifierCore/Source/ActorModifierCoreEditor/Private/Modifiers/ActorModifierCoreEditorMenu.cpp
@@ -443,8 +443,8 @@ void UE::ActorModifierCoreEditor::OnExtendEnableModifierMenu(UToolMenu* InToolMe
 		
 		if (InData.Context.ContainsEnabledModifier())
 		{
-			static const FText DisableModifierLabel = LOCTEXT("EnableModifier.Label", "Disable selected modifier");
-			static const FText DisableModifierTooltip = LOCTEXT("EnableModifier.Tooltip", "Disable selected modifier");
+			static const FText DisableModifierLabel = LOCTEXT("DisableModifier.Label", "Disable selected modifier");
+			static const FText DisableModifierTooltip = LOCTEXT("DisableModifier.Tooltip", "Disable selected modifier");
 			static const FName DisableModifierEntryName(TEXT("DisableModifier"));
 		
 			const FUIAction DisableModifierAction = FUIAction(FExecuteAction::CreateLambda(&UE::ActorModifierCoreEditor::OnEnableModifierMenuAction, InData, false, false));
diff --git a/Engine/Plugins/Experimental/ActorModifierCore/Source/ActorModifierCoreEditor/Private/Subsystems/ActorModifierCoreEditorSubsystem.cpp b/Engine/Plugins/Experimental/ActorModifierCore/Source/ActorModifierCoreEditor/Private/Subsystems/ActorModifierCoreEditorSubsystem.cpp
--- a/Engine/Plugins/Experimental/ActorModifierCore/Source/ActorModifierCoreEditor/Private/Subsystems/ActorModifierCoreEditorSubsystem.cpp
+++ b/Engine/Plugins/Experimental/ActorModifierCore/Source/ActorModifierCoreEditor/Private/Subsystems/ActorModifierCoreEditorSubsystem.cpp
@@ -232,7 +232,7 @@ bool UActorModifierCoreEditorSubsystem::RemoveActorsModifiers(TSet<AActor*>& InA
 
 	// create transaction
 	const FText TransactionText = FText::Format(
-		LOCTEXT("AddModifier", "Remove all modifiers from {0} actor(s)"),
+		LOCTEXT("RemoveModifier", "Remove all modifiers from {0} actor(s)"),
 		FText::FromString(FString::FromInt(InActors.Num())));
 	FScopedTransaction Transaction(TransactionText, bInShouldTransact);
 
@@ -334,14 +334,14 @@ bool UActorModifierCoreEditorSubsystem::InsertModifier(const FName& InModifierNa
 		if (IsValid(InPositionModifier))
 		{
 			TransactionText = FText::Format(
-				LOCTEXT("InsertModifier", "Insert {0} modifier before {1}"),
+				LOCTEXT("InsertModifierBefore", "Insert {0} modifier before {1}"),
 				FText::FromName(InModifierName),
 				FText::FromName(InPositionModifier->GetModifierName()));
 		}
 		else
 		{
 			TransactionText = FText::Format(
-				LOCTEXT("InsertModifier", "Insert {0} modifier at the end of stack"),
+				LOCTEXT("InsertModifierAtEnd", "Insert {0} modifier at the end of stack"),
 				FText::FromName(InModifierName));
 		}
 	}
@@ -350,14 +350,14 @@ bool UActorModifierCoreEditorSubsystem::InsertModifier(const FName& InModifierNa
 		if (IsValid(InPositionModifier))
 		{
 			TransactionText = FText::Format(
-				LOCTEXT("InsertModifier", "Insert {0} modifier after {1}"),
+				LOCTEXT("InsertModifierAfter", "Insert {0} modifier after {1}"),
 				FText::FromName(InModifierName),
 				FText::FromName(InPositionModifier->GetModifierName()));
 		}
 		else
 		{
 			TransactionText = FText::Format(
-				LOCTEXT("InsertModifier", "Insert {0} modifier at the start of stack"),
+				LOCTEXT("InsertModifierAtStart", "Insert {0} modifier at the start of stack"),
 				FText::FromName(InModifierName));
 		}
 	}
@@ -399,14 +399,14 @@ bool UActorModifierCoreEditorSubsystem::MoveModifier(UActorModifierCoreBase* InM
 		if (IsValid(InPositionModifier))
 		{
 			TransactionText = FText::Format(
-				LOCTEXT("InsertModifier", "Move {0} modifier before {1}"),
+				LOCTEXT("MoveModifierBefore", "Move {0} modifier before {1}"),
 				FText::FromName(MoveModifierName),
 				FText::FromName(InPositionModifier->GetModifierName()));
 		}
 		else
 		{
 			TransactionText = FText::Format(
-				LOCTEXT("InsertModifier", "Move {0} modifier at the end of stack"),
+				LOCTEXT("MoveModifierToEnd", "Move {0} modifier at the end of stack"),
 				FText::FromName(MoveModifierName));
 		}
 	}
@@ -415,14 +415,14 @@ bool UActorModifierCoreEditorSubsystem::MoveModifier(UActorModifierCoreBase* InM
 		if (IsValid(InPositionModifier))
 		{
 			TransactionText = FText::Format(
-				LOCTEXT("InsertModifier", "Move {0} modifier after {1}"),
+				LOCTEXT("MoveModifierAfter", "Move {0} modifier after {1}"),
 				FText::FromName(MoveModifierName),
 				FText::FromName(InPositionModifier->GetModifierName()));
 		}
 		else
 		{
 			TransactionText = FText::Format(
-				LOCTEXT("InsertModifier", "Move {0} modifier at the start of stack"),
+				LOCTEXT("MoveModifierToStart", "Move {0} modifier at the start of stack"),
 				FText::FromName(MoveModifierName));
 		}
 	}
diff --git a/Engine/Plugins/Experimental/DynamicMaterial/Source/DynamicMaterialEditor/Private/Components/MaterialStageExpressions/DMMSEDeltaTime.cpp b/Engine/Plugins/Experimental/DynamicMaterial/Source/DynamicMaterialEditor/Private/Components/MaterialStageExpressions/DMMSEDeltaTime.cpp
--- a/Engine/Plugins/Experimental/DynamicMaterial/Source/DynamicMaterialEditor/Private/Components/MaterialStageExpressions/DMMSEDeltaTime.cpp
+++ b/Engine/Plugins/Experimental/DynamicMaterial/Source/DynamicMaterialEditor/Private/Components/MaterialStageExpressions/DMMSEDeltaTime.cpp
@@ -13,7 +13,7 @@ UDMMaterialStageExpressionDeltaTime::UDMMaterialStageExpressionDeltaTime()
 {
 	Menus.Add(EDMExpressionMenu::Time);
 
-	OutputConnectors.Add({0, LOCTEXT("DeltaTime", "Delta Time"), EDMValueType::VT_Float1});
+	OutputConnectors.Add({0, LOCTEXT("DeltaTimeLabel", "Delta Time"), EDMValueType::VT_Float1});
 }
 
 #undef LOCTEXT_NAMESPACE
diff --git a/Engine/Plugins/Experimental/DynamicMaterial/Source/DynamicMaterialEditor/Private/Components/MaterialStageInputs/DMMSIValue.cpp b/Engine/Plugins/Experimental/DynamicMaterial/Source/DynamicMaterialEditor/Private/Components/MaterialStageInputs/DMMSIValue.cpp
--- a/Engine/Plugins/Experimental/DynamicMaterial/Source/DynamicMaterialEditor/Private/Components/MaterialStageInputs/DMMSIValue.cpp
+++ b/Engine/Plugins/Experimental/DynamicMaterial/Source/DynamicMaterialEditor/Private/Components/MaterialStageInputs/DMMSIValue.cpp
@@ -45,14 +45,14 @@ FText UDMMaterialStageInputValue::GetComponentDescription() const
 	{
 		if (Value->IsLocal())
 		{
-			static const FText Template = LOCTEXT("ValueLocalNamed", "{0}");
+			static const FText Template = LOCTEXT("ComponentDescriptionLocalNamed", "{0}");
 
 			return FText::Format(Template, Value->GetTypeName());
 
 		}
 		else
 		{
-			static const FText Template = LOCTEXT("ValueXGlobalNamed", "{0} (Global)");
+			static const FText Template = LOCTEXT("ComponentDescriptionGlobalNamed", "{0} (Global)");
 
 			return FText::Format(Template, Value->GetTypeName());
 		}
@@ -61,14 +61,14 @@ FText UDMMaterialStageInputValue::GetComponentDescription() const
 	{
 		if (Value->IsLocal())
 		{
-			static const FText Template = LOCTEXT("ValueLocal", "{0}");
+			static const FText Template = LOCTEXT("ComponentDescriptionLocal", "{0}");
 
 			return FText::Format(Template, Value->GetTypeName());
 
 		}
 		else
 		{
-			static const FText Template = LOCTEXT("ValueXGlobal", "{0} (Global)");
+			static const FText Template = LOCTEXT("ComponentDescriptionGlobal", "{0} (Global)");
 
 			return FText::Format(Template, Value->GetTypeName());
 		}
@@ -89,7 +89,7 @@ FText UDMMaterialStageInputValue::GetChannelDescription(const FDMMaterialStageCo
 
 	if (Value->GetParameter())
 	{
-		static const FText FormatTemplateNamed = LOCTEXT("ValueXGlobalNamed", "{0} (Global)");
+		static const FText FormatTemplateNamed = LOCTEXT("ChannelDescriptionGlobalNamed", "{0} (Global)");
 
 		return FText::Format(
 			FormatTemplateNamed,
@@ -98,7 +98,7 @@ FText UDMMaterialStageInputValue::GetChannelDescription(const FDMMaterialStageCo
 	}
 	else
 	{
-		static const FText FormatTemplate = LOCTEXT("ValueXGlobal", "Value {0} (Global)");
+		static const FText FormatTemplate = LOCTEXT("ChannelDescriptionGlobal", "Value {0} (Global)");
 
 		return FText::Format(
 			FormatTemplate,
diff --git a/Engine/Plugins/Experimental/DynamicMaterial/Source/DynamicMaterialEditor/Private/Menus/DMMaterialStageInputMenus.cpp b/Engine/Plugins/Experimental/DynamicMaterial/Source/DynamicMaterialEditor/Private/Menus/DMMaterialStageInputMenus.cpp
--- a/Engine/Plugins/Experimental/DynamicMaterial/Source/DynamicMaterialEditor/Private/Menus/DMMaterialStageInputMenus.cpp
+++ b/Engine/Plugins/Experimental/DynamicMaterial/Source/DynamicMaterialEditor/Private/Menus/DMMaterialStageInputMenus.cpp
@@ -342,7 +342,7 @@ namespace UE::DynamicMaterialEditor::Private
 	{
 		InChildMenuBuilder.AddSubMenu(
 			InMenuName,
-			LOCTEXT("ChangeInputPreviousStageTooltip", "Change the source of this input to a previous Material Stage."),
+			LOCTEXT("ChangeInputPreviousStageOutputsTooltip", "Change the source of this input to a previous Material Stage."),
 			FNewMenuDelegate::CreateStatic(
 				&FDMMaterialStageInputMenus::GenerateChangeInputMenu_PreviousStage_Outputs,
 				InThroughput,
@@ -359,7 +359,7 @@ namespace UE::DynamicMaterialEditor::Private
 	{
 		InChildMenuBuilder.AddSubMenu(
 			InMenuName,
-			LOCTEXT("ChangeInputPreviousStageTooltip", "Change the source of this input to a previous Material Stage."),
+			LOCTEXT("ChangeInputPreviousStageOutputChannelsTooltip", "Change the source of this input to a previous Material Stage."),
 			FNewMenuDelegate::CreateStatic(
 				&FDMMaterialStageInputMenus::GenerateChangeInputMenu_PreviousStage_Output_Channels,
 				InThroughput,
@@ -377,7 +377,7 @@ namespace UE::DynamicMaterialEditor::Private
 	{
 		InChildMenuBuilder.AddMenuEntry(
 			InMenuName,
-			LOCTEXT("ChangeInputPreviousStageTooltip", "Change the source of this input to a previous Material Stage."),
+			LOCTEXT("ChangeInputPreviousStageOutputChannelTooltip", "Change the source of this input to a previous Material Stage."),
 			FSlateIcon(),
 			FUIAction(FExecuteAction::CreateWeakLambda(
 				InStage,
@@ -397,7 +397,7 @@ namespace UE::DynamicMaterialEditor::Private
 	{
 		InChildMenuBuilder.AddSubMenu(
 			InMenuName,
-			LOCTEXT("ChangeInputPreviousStageTooltip", "Change the source of this input to a previous Material Stage."),
+			LOCTEXT("ChangeInputPreviousStageOutputChannelChannelsTooltip", "Change the source of this input to a previous Material Stage."),
 			FNewMenuDelegate::CreateStatic(
 				&FDMMaterialStageInputMenus::GenerateChangeInputMenu_PreviousStage_Output_Channels,
 				InThroughput,
@@ -776,7 +776,7 @@ namespace UE::DynamicMaterialEditor::Private
 	{
 		InChildMenuBuilder.AddSubMenu(
 			InMenuName,
-			LOCTEXT("ChangeInputGlobalValueTooltip", "Change the source of this input to this Global Material Value."),
+			LOCTEXT("ChangeInputGlobalValueChannelsTooltip", "Change the source of this input to this Global Material Value."),
 			FNewMenuDelegate::CreateStatic(
 				&FDMMaterialStageInputMenus::GenerateChangeInputMenu_GlobalValue_Channels,
 				InThroughput,
@@ -793,7 +793,7 @@ namespace UE::DynamicMaterialEditor::Private
 	{		
 		InChildMenuBuilder.AddMenuEntry(
 			InMenuName,
-			LOCTEXT("ChangeInputGlobalValueTooltip", "Change the source of this input to this Global Material Value."),
+			LOCTEXT("ChangeInputGlobalValueChannelTooltip", "Change the source of this input to this Global Material Value."),
 			FSlateIcon(),
 			FUIAction(FExecuteAction::CreateWeakLambda(
 				InStage,
@@ -817,7 +817,7 @@ namespace UE::DynamicMaterialEditor::Private
 	{
 		InChildMenuBuilder.AddSubMenu(
 			InMenuName,
-			LOCTEXT("ChangeInputGlobalValueTooltip", "Change the source of this input to this Global Material Value."),
+			LOCTEXT("ChangeInputGlobalValueValueChannelsTooltip", "Change the source of this input to this Global Material Value."),
 			FNewMenuDelegate::CreateStatic(
 				&FDMMaterialStageInputMenus::GenerateChangeInputMenu_GlobalValue_Channels,
 				InThroughput,
@@ -1062,7 +1062,7 @@ namespace UE::DynamicMaterialEditor::Private
 	{
 		InChildMenuBuilder.AddSubMenu(
 			InMenuName,
-			LOCTEXT("ChangeInputSlotTooltip", "Change the source of this input to the output from this Material Slot."),
+			LOCTEXT("ChangeInputSlotPropertiesTooltip", "Change the source of this input to the output from this Material Slot."),
 			FNewMenuDelegate::CreateStatic(
 				&FDMMaterialStageInputMenus::GenerateChangeInputMenu_Slot_Properties,
 				InThroughput,
@@ -1079,7 +1079,7 @@ namespace UE::DynamicMaterialEditor::Private
 	{
 		InChildMenuBuilder.AddSubMenu(
 			InMenuName,
-			LOCTEXT("ChangeInputSlotTooltip", "Change the source of this input to the output from this Material Slot."),
+			LOCTEXT("ChangeInputSlotPropertyOutputsTooltip", "Change the source of this input to the output from this Material Slot."),
 			FNewMenuDelegate::CreateStatic(
 				&FDMMaterialStageInputMenus::GenerateChangeInputMenu_Slot_Property_Outputs,
 				InThroughput,
@@ -1097,7 +1097,7 @@ namespace UE::DynamicMaterialEditor::Private
 	{
 		InChildMenuBuilder.AddSubMenu(
 			InMenuName,
-			LOCTEXT("ChangeInputSlotTooltip", "Change the source of this input to the output from this Material Slot."),
+			LOCTEXT("ChangeInputSlotPropertyOutputChannelsTooltip", "Change the source of this input to the output from this Material Slot."),
 			FNewMenuDelegate::CreateStatic(
 				&FDMMaterialStageInputMenus::GenerateChangeInputMenu_Slot_Property_Output_Channels,
 				InThroughput,
@@ -1116,7 +1116,7 @@ namespace UE::DynamicMaterialEditor::Private
 	{
 		InChildMenuBuilder.AddMenuEntry(
 			InMenuName,
-			LOCTEXT("ChangeInputSlotTooltip", "Change the source of this input to the output from this Material Slot."),
+			LOCTEXT("ChangeInputSlotPropertyOutputChannelTooltip", "Change the source of this input to the output from this Material Slot."),
 			FSlateIcon(),
 			FUIAction(FExecuteAction::CreateWeakLambda(
 				InStage,
@@ -1140,7 +1140,7 @@ namespace UE::DynamicMaterialEditor::Private
 	{
 		InChildMenuBuilder.AddSubMenu(
 			InMenuName,
-			LOCTEXT("ChangeInputSlotTooltip", "Change the source of this input to the output from this Material Slot."),
+			LOCTEXT("ChangeInputSlotPropertyOutputChannelChannelsTooltip", "Change the source of this input to the output from this Material Slot."),
 			FNewMenuDelegate::CreateStatic(
 				&FDMMaterialStageInputMenus::GenerateChangeInputMenu_Slot_Property_Output_Channels,
 				InThroughput,
diff --git a/Engine/Plugins/Experimental/DynamicMaterial/Source/DynamicMaterialEditor/Private/Menus/DMMaterialStageSourceMenus.cpp b/Engine/Plugins/Experimental/DynamicMaterial/Source/DynamicMaterialEditor/Private/Menus/DMMaterialStageSourceMenus.cpp
--- a/Engine/Plugins/Experimental/DynamicMaterial/Source/DynamicMaterialEditor/Private/Menus/DMMaterialStageSourceMenus.cpp
+++ b/Engine/Plugins/Experimental/DynamicMaterial/Source/DynamicMaterialEditor/Private/Menus/DMMaterialStageSourceMenus.cpp
@@ -1061,7 +1061,7 @@ namespace UE::DynamicMaterialEditor::Private
 
 		InSection.AddMenuEntry("TextureSample_EdgeColor",
 			LOCTEXT("AddTextureSampleEgdeColor", "Texture Edge Color"),
-			LOCTEXT("ChangeSourceTextureSampleTooltip", "Change the source of this stage to the edge color of a texture."),
+			LOCTEXT("ChangeSourceTextureSampleEdgeColorTooltip", "Change the source of this stage to the edge color of a texture."),
 			FSlateIcon(),
 			FUIAction(
 				FExecuteAction::CreateStatic(
@@ -1144,8 +1144,8 @@ namespace UE::DynamicMaterialEditor::Private
 						[](FToolMenuSection& InSection)
 						{
 							InSection.AddSubMenu("SlotOutput",
-							LOCTEXT("ChangeSourceSlot", "Slot Output"),
-							LOCTEXT("ChangeSourceSlotTooltip", "Change the source of this stage to the output from another Material Slot."),
+							LOCTEXT("ChangeSourceSlotOuptut", "Slot Output"),
+							LOCTEXT("ChangeSourceSlotOutputTooltip", "Change the source of this stage to the output from another Material Slot."),
 							FNewToolMenuDelegate::CreateStatic(&GenerateChangeSourceMenu_Slots)
 							);
 						}));
diff --git a/Engine/Plugins/Experimental/DynamicMaterial/Source/DynamicMaterialEditor/Private/Slate/Layers/SDMLayerEffectsItem.cpp b/Engine/Plugins/Experimental/DynamicMaterial/Source/DynamicMaterialEditor/Private/Slate/Layers/SDMLayerEffectsItem.cpp
--- a/Engine/Plugins/Experimental/DynamicMaterial/Source/DynamicMaterialEditor/Private/Slate/Layers/SDMLayerEffectsItem.cpp
+++ b/Engine/Plugins/Experimental/DynamicMaterial/Source/DynamicMaterialEditor/Private/Slate/Layers/SDMLayerEffectsItem.cpp
@@ -106,7 +106,7 @@ TSharedRef<SWidget> SDMLayerEffectsItem::CreateLayerRemoveButton()
 		SNew(SButton)
 		.ContentPadding(2.0f)
 		.ButtonStyle(FDynamicMaterialEditorStyle::Get(), "HoverHintOnly")
-		.ToolTipText(LOCTEXT("RemoveEffect", "Remove Effect"))
+		.ToolTipText(LOCTEXT("RemoveEffectTooltip", "Remove Effect"))
 		.OnClicked(this, &SDMLayerEffectsItem::OnLayerRemoveButtonClick)
 		[
 			SNew(SImage)
diff --git a/Engine/Plugins/Experimental/DynamicMaterial/Source/DynamicMaterialEditor/Private/Slate/SDMComponentEdit.cpp b/Engine/Plugins/Experimental/DynamicMaterial/Source/DynamicMaterialEditor/Private/Slate/SDMComponentEdit.cpp
--- a/Engine/Plugins/Experimental/DynamicMaterial/Source/DynamicMaterialEditor/Private/Slate/SDMComponentEdit.cpp
+++ b/Engine/Plugins/Experimental/DynamicMaterial/Source/DynamicMaterialEditor/Private/Slate/SDMComponentEdit.cpp
@@ -216,7 +216,7 @@ FText SDMComponentEdit::GetInputChannelMapDescription(UDMMaterialStageThroughput
 	}
 
 	static const FText Separator = LOCTEXT("ChannelSeparator", ",");
-	static const FText MaskedChannelFormat = LOCTEXT("WholeChannelFormat", "{0}: {1}");
+	static const FText MaskedChannelFormat = LOCTEXT("MaskedChannelFormat", "{0}: {1}");
 	const FText ChannelName = FText::Join(Separator, Channels);
 
 	return FText::Format(MaskedChannelFormat, SourceName, ChannelName);
diff --git a/Engine/Plugins/Experimental/DynamicMaterial/Source/DynamicMaterialEditor/Private/Slate/SDMSlot.cpp b/Engine/Plugins/Experimental/DynamicMaterial/Source/DynamicMaterialEditor/Private/Slate/SDMSlot.cpp
--- a/Engine/Plugins/Experimental/DynamicMaterial/Source/DynamicMaterialEditor/Private/Slate/SDMSlot.cpp
+++ b/Engine/Plugins/Experimental/DynamicMaterial/Source/DynamicMaterialEditor/Private/Slate/SDMSlot.cpp
@@ -579,7 +579,7 @@ TSharedRef<SWidget> SDMSlot::CreateLayerButtonsRowWidget()
 			SNew(SButton)
 			.ContentPadding(4.0f)
 			.ButtonStyle(FDynamicMaterialEditorStyle::Get(), "HoverHintOnly.Bordered.Dark")
-			.ToolTipText(LOCTEXT("RemoveSlot", "Remove Slot\n\nRGB Slots cannot be removed."))
+			.ToolTipText(LOCTEXT("RemoveSlotTooltip", "Remove Slot\n\nRGB Slots cannot be removed."))
 			.IsEnabled(this, &SDMSlot::CanRemoveSlot)
 			.OnClicked(this, &SDMSlot::OnRemoveSlotClicked)
 			[
@@ -608,7 +608,7 @@ TSharedRef<SWidget> SDMSlot::CreateLayerButtonsRowWidget()
 			.IsFocusable(true)
 			.ContentPadding(4.0f)
 			.ButtonStyle(FDynamicMaterialEditorStyle::Get(), "HoverHintOnly.Bordered.Dark")
-			.ToolTipText(LOCTEXT("AddLayer", "Add New Layer"))
+			.ToolTipText(LOCTEXT("AddLayerTooltip", "Add New Layer"))
 			.OnGetMenuContent(this, &SDMSlot::GetLayerButtonsMenuContent)
 			.ButtonContent()
 			[
@@ -645,7 +645,7 @@ TSharedRef<SWidget> SDMSlot::CreateLayerButtonsRowWidget()
 			SNew(SButton)
 			.ContentPadding(4.0f)
 			.ButtonStyle(FDynamicMaterialEditorStyle::Get(), "HoverHintOnly.Bordered.Dark")
-			.ToolTipText(LOCTEXT("RemoveLayer", "Remove Selected Layer\n\nThe last layer cannot be removed."))
+			.ToolTipText(LOCTEXT("RemoveLayerTooltip", "Remove Selected Layer\n\nThe last layer cannot be removed."))
 			.IsEnabled(this, &SDMSlot::GetLayerRowsButtonsCanRemove)
 			.OnClicked(this, &SDMSlot::OnLayerRowButtonsRemoveClicked)
 			[
diff --git a/Engine/Plugins/Experimental/PropertyAnimatorCore/Source/PropertyAnimatorCore/Private/Subsystems/PropertyAnimatorCoreSubsystem.cpp b/Engine/Plugins/Experimental/PropertyAnimatorCore/Source/PropertyAnimatorCore/Private/Subsystems/PropertyAnimatorCoreSubsystem.cpp
--- a/Engine/Plugins/Experimental/PropertyAnimatorCore/Source/PropertyAnimatorCore/Private/Subsystems/PropertyAnimatorCoreSubsystem.cpp
+++ b/Engine/Plugins/Experimental/PropertyAnimatorCore/Source/PropertyAnimatorCore/Private/Subsystems/PropertyAnimatorCoreSubsystem.cpp
@@ -508,7 +508,7 @@ bool UPropertyAnimatorCoreSubsystem::LinkAnimatorProperties(UPropertyAnimatorCor
 			const FText AnimatorName = FText::FromName(InAnimator->GetAnimatorOriginalName());
 
 			TransactionText = FText::Format(
-				LOCTEXT("UnlinkAnimatorProperties", "Linking {0} properties to {1} animator")
+				LOCTEXT("LinkAnimatorProperties", "Linking {0} properties to {1} animator")
 				, PropertyCount
 				, AnimatorName
 			);
diff --git a/Engine/Plugins/Experimental/PropertyAnimatorCore/Source/PropertyAnimatorCoreEditor/Private/Menus/PropertyAnimatorCoreEditorMenu.cpp b/Engine/Plugins/Experimental/PropertyAnimatorCore/Source/PropertyAnimatorCoreEditor/Private/Menus/PropertyAnimatorCoreEditorMenu.cpp
--- a/Engine/Plugins/Experimental/PropertyAnimatorCore/Source/PropertyAnimatorCoreEditor/Private/Menus/PropertyAnimatorCoreEditorMenu.cpp
+++ b/Engine/Plugins/Experimental/PropertyAnimatorCore/Source/PropertyAnimatorCoreEditor/Private/Menus/PropertyAnimatorCoreEditorMenu.cpp
@@ -181,8 +181,8 @@ void UE::PropertyAnimatorCoreEditor::Menu::FillDeleteAnimatorSection(UToolMenu*
 
 	ActorAnimatorsSection.AddMenuEntry(
 		TEXT("DeleteActorAnimator")
-		, LOCTEXT("DeleteActorAnimator.Label", "Delete actor animators")
-		, LOCTEXT("DeleteActorAnimator.Tooltip", "Delete selected actor animators")
+		, LOCTEXT("DeleteActorAnimators.Label", "Delete actor animators")
+		, LOCTEXT("DeleteActorAnimators.Tooltip", "Delete selected actor animators")
 		, FSlateIcon()
 		, FUIAction(
 		  FExecuteAction::CreateLambda(&ExecuteDeleteActorAnimatorAction, InMenuData)
@@ -207,8 +207,8 @@ void UE::PropertyAnimatorCoreEditor::Menu::FillDeleteAnimatorSection(UToolMenu*
 
 		ActorAnimatorsSection.AddMenuEntry(
 			FName(Animator->GetAnimatorDisplayName())
-			, FText::Format(LOCTEXT("DeleteActorAnimator.Label", "Delete {0}"), FText::FromString(Animator->GetAnimatorDisplayName()))
-			, LOCTEXT("DeleteActorAnimator.Tooltip", "Delete selected animator")
+			, FText::Format(LOCTEXT("DeleteSingleActorAnimator.Label", "Delete {0}"), FText::FromString(Animator->GetAnimatorDisplayName()))
+			, LOCTEXT("DeleteSingleActorAnimator.Tooltip", "Delete selected animator")
 			, FSlateIcon()
 			, FUIAction(
 			  FExecuteAction::CreateLambda(&ExecuteDeleteAnimatorAction, Animator, InMenuData)
diff --git a/Engine/Plugins/Experimental/StormSync/Source/StormSyncDrives/Private/StormSyncDrivesModule.cpp b/Engine/Plugins/Experimental/StormSync/Source/StormSyncDrives/Private/StormSyncDrivesModule.cpp
--- a/Engine/Plugins/Experimental/StormSync/Source/StormSyncDrives/Private/StormSyncDrivesModule.cpp
+++ b/Engine/Plugins/Experimental/StormSync/Source/StormSyncDrives/Private/StormSyncDrivesModule.cpp
@@ -148,7 +148,7 @@ bool FStormSyncDrivesModule::ValidateMountPoint(const FStormSyncMountPointConfig
 	if (!FStormSyncDrivesUtils::ValidateMountPoint(InMountPoint, ErrorText))
 	{
 		const FText ErrorLog = FText::Format(
-			LOCTEXT("ValidationError_MountPoint", "Validation failed for MountPoint at index {0} (MountPoint: {1}, MountDirectory: {2}) - {3}"),
+			LOCTEXT("ValidationError_MountPointIndex", "Validation failed for MountPoint at index {0} (MountPoint: {1}, MountDirectory: {2}) - {3}"),
 			FText::AsNumber(Index),
 			FText::FromString(RootPath),
 			FText::FromString(Directory),
diff --git a/Engine/Plugins/Experimental/StormSync/Source/StormSyncEditor/Private/ContextMenu/StormSyncAssetFolderContextMenu.cpp b/Engine/Plugins/Experimental/StormSync/Source/StormSyncEditor/Private/ContextMenu/StormSyncAssetFolderContextMenu.cpp
--- a/Engine/Plugins/Experimental/StormSync/Source/StormSyncEditor/Private/ContextMenu/StormSyncAssetFolderContextMenu.cpp
+++ b/Engine/Plugins/Experimental/StormSync/Source/StormSyncEditor/Private/ContextMenu/StormSyncAssetFolderContextMenu.cpp
@@ -75,7 +75,7 @@ void FStormSyncAssetFolderContextMenu::BuildCompareWithMenuSection(FMenuBuilder&
 
 		const FText RemoteDescription = FText::FromString(FString::Printf(TEXT("%s (%s)"), *Connection.HostName, *Connection.ProjectDir));
 		FText LabelText = FText::Format(
-			LOCTEXT("PushPullAssetMenuAction", "Compare {0} selected asset(s) with {1}"),
+			LOCTEXT("CompareAssetMenuAction", "Compare {0} selected asset(s) with {1}"),
 			InPackageNames.Num(),
 			RemoteDescription
 		);
@@ -309,7 +309,7 @@ void FStormSyncAssetFolderContextMenu::AddFileMenuOptions(UToolMenu* InMenu, con
 		if (!bIsServerRunning)
 		{
 			MenuTooltipTexts.Add(LOCTEXT(
-				"PullAssetsMenuEntry_ServerDisabled",
+				"PullAssetsMenuEntry_ServerDisabledTooltip",
 				"\n\nThis action is disabled because Storm Sync Server is not running.\nRun `StormSync.Server.Start` or configure the server with auto-start setting enabled."
 			));
 		}
@@ -344,7 +344,7 @@ void FStormSyncAssetFolderContextMenu::AddFileMenuOptions(UToolMenu* InMenu, con
 		Section.AddSubMenu(
 			TEXT("StormSync_CompareAssets"),
 			LOCTEXT("CompareAssetsMenuEntry", "Compare Asset(s) with"),
-			FText::Format(LOCTEXT("CompareAssetsMenuEntry", "Compare asset(s) with a specific remote and see if files (and inner dependencies) are either missing or in mismatched state.{0}"), DisabledTooltipReason),
+			FText::Format(LOCTEXT("CompareAssetsMenuEntryTooltip", "Compare asset(s) with a specific remote and see if files (and inner dependencies) are either missing or in mismatched state.{0}"), DisabledTooltipReason),
 			FNewMenuDelegate::CreateRaw(
 				this,
 				&FStormSyncAssetFolderContextMenu::BuildCompareWithMenuSection,
@@ -364,8 +364,8 @@ void FStormSyncAssetFolderContextMenu::AddFileMenuOptions(UToolMenu* InMenu, con
 	{
 		Section.AddMenuEntry(
 			TEXT("StormSync_ExportAssets"),
-			LOCTEXT("ExportAsAssetsMenuEntry", "Export"),
-			FText::Format(LOCTEXT("ExportAsAssetsMenuEntryTooltip", "Export assets (and all inner dependencies) locally.{0}"), DisabledTooltipReason),
+			LOCTEXT("FolderExportAsAssetsMenuEntry", "Export"),
+			FText::Format(LOCTEXT("FolderExportAsAssetsMenuEntryTooltip", "Export assets (and all inner dependencies) locally.{0}"), DisabledTooltipReason),
 			FSlateIcon(FAppStyle::GetAppStyleSetName(), "Icons.Save"),
 			FUIAction(
 				FExecuteAction::CreateStatic(&FStormSyncAssetFolderContextMenu::ExecuteExportAction, InSelectedPackageNames),
@@ -402,8 +402,8 @@ void FStormSyncAssetFolderContextMenu::AddFolderMenuOptions(UToolMenu* InMenu, c
 	{
 		Section.AddMenuEntry(
 			TEXT("StormSync_ExportAssets"),
-			LOCTEXT("ExportAsAssetsMenuEntry", "Export"),
-			LOCTEXT("ExportAsAssetsMenuEntryTooltip", "Export assets (and all inner dependencies) locally"),
+			LOCTEXT("FileExportAsAssetsMenuEntry", "Export"),
+			LOCTEXT("FileExportAsAssetsMenuEntryTooltip", "Export assets (and all inner dependencies) locally"),
 			FSlateIcon(FAppStyle::GetAppStyleSetName(), "Icons.Save"),
 			FUIAction(
 				FExecuteAction::CreateStatic(&FStormSyncAssetFolderContextMenu::ExecuteExportAction, SelectedPackageNames)
diff --git a/Engine/Plugins/Experimental/StormSync/Source/StormSyncEditor/Private/Slate/SStormSyncImportWizard.cpp b/Engine/Plugins/Experimental/StormSync/Source/StormSyncEditor/Private/Slate/SStormSyncImportWizard.cpp
--- a/Engine/Plugins/Experimental/StormSync/Source/StormSyncEditor/Private/Slate/SStormSyncImportWizard.cpp
+++ b/Engine/Plugins/Experimental/StormSync/Source/StormSyncEditor/Private/Slate/SStormSyncImportWizard.cpp
@@ -99,8 +99,8 @@ void SStormSyncImportWizard::Construct(const FArguments& InArgs, const TArray<FS
 					.ToolTip(LOCTEXT("Tab_FilesToImport_Tooltip", "Click to see the list of files to import due to mismatched file state (File not existing locally or mismatched File Size or Hash)"))
 					+SSegmentedControl<EStormSyncImportWizardActiveTab>::Slot(EStormSyncImportWizardActiveTab::BufferFiles)
 					.Icon(FAppStyle::Get().GetBrush("Icons.Save"))
-					.Text(LOCTEXT("Tab_FilesToImport", "Full list of files from buffer"))
-					.ToolTip(LOCTEXT("Tab_FilesToImport_Tooltip", "Click to see the original list of files included in buffer"))
+					.Text(LOCTEXT("Tab_FilesToImportFullList", "Full list of files from buffer"))
+					.ToolTip(LOCTEXT("Tab_FilesToImportFullList_Tooltip", "Click to see the original list of files included in buffer"))
 				]
 
 				+SVerticalBox::Slot()
diff --git a/Engine/Plugins/Experimental/StormSync/Source/StormSyncEditor/Private/Slate/Status/SStormSyncStatusWidget.cpp b/Engine/Plugins/Experimental/StormSync/Source/StormSyncEditor/Private/Slate/Status/SStormSyncStatusWidget.cpp
--- a/Engine/Plugins/Experimental/StormSync/Source/StormSyncEditor/Private/Slate/Status/SStormSyncStatusWidget.cpp
+++ b/Engine/Plugins/Experimental/StormSync/Source/StormSyncEditor/Private/Slate/Status/SStormSyncStatusWidget.cpp
@@ -166,7 +166,7 @@ void SStormSyncStatusWidget::Construct(const FArguments& InArgs, const TSharedRe
 					[
 						SNew(STextBlock)
 						.Font(FAppStyle::Get().GetFontStyle("NormalFontBold"))
-						.Text(LOCTEXT("Remote_InstanceId", "Remote Instance Type:"))
+						.Text(LOCTEXT("Remote_InstanceType", "Remote Instance Type:"))
 					]
 					+SHorizontalBox::Slot()
 					.HAlign(HAlign_Left)
diff --git a/Engine/Plugins/Experimental/StormSync/Source/StormSyncImport/Private/Subsystems/StormSyncImportSubsystem.cpp b/Engine/Plugins/Experimental/StormSync/Source/StormSyncImport/Private/Subsystems/StormSyncImportSubsystem.cpp
--- a/Engine/Plugins/Experimental/StormSync/Source/StormSyncImport/Private/Subsystems/StormSyncImportSubsystem.cpp
+++ b/Engine/Plugins/Experimental/StormSync/Source/StormSyncImport/Private/Subsystems/StormSyncImportSubsystem.cpp
@@ -158,7 +158,7 @@ bool UStormSyncImportSubsystem::PerformImport(const FStormSyncPackageDescriptor&
 	FStormSyncPackageDescriptor PackageDescriptor = InPackageDescriptor;
 	
 	float InitialAmountOfWork = 1.f;
-	FScopedSlowTask SlowTask(InitialAmountOfWork, LOCTEXT("ImportPackages_Extracting", "Figuring out file states..."));
+	FScopedSlowTask SlowTask(InitialAmountOfWork, LOCTEXT("FiguringOutFileStates", "Figuring out file states..."));
 	SlowTask.MakeDialog();
 
 	FStormSyncCoreExtractArgs ExtractArgs;
@@ -210,7 +210,7 @@ bool UStormSyncImportSubsystem::PerformImport(const FStormSyncPackageDescriptor&
 			SplitImportFiles(FilesToImport, ExistingFiles, NewFiles);
 
 			FScopedSlowTask SlowTaskExtract(FilesToImport.Num() + ExistingFiles.Num() + NewFiles.Num()
-				, LOCTEXT("ImportPackages_Extracting", "Importing files to project..."));
+				, LOCTEXT("ImportingFilesToProject", "Importing files to project..."));
 
 			// Step 1: Prompt for Checkout (for existing assets)
 			SlowTaskExtract.EnterProgressFrame(ExistingFiles.Num());
@@ -308,7 +308,7 @@ bool UStormSyncImportSubsystem::FillFilesToImport(const FStormSyncImportFileInfo
 		FString ShortPath = InDestFilepath;
 		ShortPath.RemoveFromStart(FPaths::ProjectContentDir());
 		FileInfo.ImportReason = FText::Format(LOCTEXT("FileInfo_MissingFile", "Missing file in local project ({0})"), FText::FromString(ShortPath));
-		FileInfo.ImportReasonTooltip = FText::Format(LOCTEXT("FileInfo_MissingFile", "File {0} does not exist locally"), FText::FromString(InDestFilepath));
+		FileInfo.ImportReasonTooltip = FText::Format(LOCTEXT("FileInfo_MissingFile_Tooltip", "File {0} does not exist locally"), FText::FromString(InDestFilepath));
 		OutFilesToImport.Add(FileInfo);
 		return true;
 	}
@@ -328,7 +328,7 @@ bool UStormSyncImportSubsystem::FillFilesToImport(const FStormSyncImportFileInfo
 	if (!bSameSize)
 	{
 		FileInfo.ImportReason = FText::Format(
-			LOCTEXT("FileInfo_MismatchedSize_Tooltip", "Files have different sizes ({0} vs {1})"),
+			LOCTEXT("FileInfo_MismatchedSize", "Files have different sizes ({0} vs {1})"),
 			FText::FromString(FStormSyncCoreUtils::GetHumanReadableByteSize(DestFileSize)),
 			FText::FromString(FStormSyncCoreUtils::GetHumanReadableByteSize(FileDependency.FileSize))
 		);
