Material Designer
- Materials will now have the bUsedWithNiagaraMeshParticles flag enabled by default.

diff --git a/Engine/Plugins/Experimental/DynamicMaterial/Source/DynamicMaterialEditor/Private/Model/DynamicMaterialModelEditorOnlyData.cpp b/Engine/Plugins/Experimental/DynamicMaterial/Source/DynamicMaterialEditor/Private/Model/DynamicMaterialModelEditorOnlyData.cpp
--- a/Engine/Plugins/Experimental/DynamicMaterial/Source/DynamicMaterialEditor/Private/Model/DynamicMaterialModelEditorOnlyData.cpp
+++ b/Engine/Plugins/Experimental/DynamicMaterial/Source/DynamicMaterialEditor/Private/Model/DynamicMaterialModelEditorOnlyData.cpp
@@ -185,6 +185,10 @@ void UDynamicMaterialModelEditorOnlyData::CreateMaterial()
 
 	MaterialModel->DynamicMaterial->bOutputTranslucentVelocity = true;
 	MaterialModel->DynamicMaterial->bEnableResponsiveAA = true;
+
+	// Not setting this to true can cause the level associated with this material to dirty itself when it
+	// is used with Niagara. It doesn't negatively affect the material in any meaningful way.
+	MaterialModel->DynamicMaterial->bUsedWithNiagaraMeshParticles = true;
 }
 
 void UDynamicMaterialModelEditorOnlyData::BuildMaterial(bool bInDirtyAssets)
