[Material Designer]
- fix category name

diff --git a/Engine/Plugins/Experimental/DynamicMaterial/Source/DynamicMaterialEditor/Private/DynamicMaterialEditorSettings.cpp b/Engine/Plugins/Experimental/DynamicMaterial/Source/DynamicMaterialEditor/Private/DynamicMaterialEditorSettings.cpp
--- a/Engine/Plugins/Experimental/DynamicMaterial/Source/DynamicMaterialEditor/Private/DynamicMaterialEditorSettings.cpp
+++ b/Engine/Plugins/Experimental/DynamicMaterial/Source/DynamicMaterialEditor/Private/DynamicMaterialEditorSettings.cpp
@@ -16,7 +16,7 @@
 
 UDynamicMaterialEditorSettings::UDynamicMaterialEditorSettings()
 {
-	CategoryName = TEXT("Avalanche");
+	CategoryName = TEXT("Plugins");
 	SectionName = TEXT("Material Designer");
 
 	bFollowSelection = true;
