Material Designer
- Changed outer of generated materials to the material model instead of the editor-only data.

diff --git a/Engine/Plugins/Experimental/DynamicMaterial/Source/DynamicMaterialEditor/Private/Model/DynamicMaterialModelEditorOnlyData.cpp b/Engine/Plugins/Experimental/DynamicMaterial/Source/DynamicMaterialEditor/Private/Model/DynamicMaterialModelEditorOnlyData.cpp
--- a/Engine/Plugins/Experimental/DynamicMaterial/Source/DynamicMaterialEditor/Private/Model/DynamicMaterialModelEditorOnlyData.cpp
+++ b/Engine/Plugins/Experimental/DynamicMaterial/Source/DynamicMaterialEditor/Private/Model/DynamicMaterialModelEditorOnlyData.cpp
@@ -154,7 +154,7 @@ void UDynamicMaterialModelEditorOnlyData::CreateMaterial()
 
 		MaterialModel->DynamicMaterial = Cast<UMaterial>(MaterialFactory->FactoryCreateNew(
 			UMaterial::StaticClass(),
-			this,
+			MaterialModel,
 			NAME_None,
 			RF_DuplicateTransient | RF_TextExportTransient,
 			nullptr,
@@ -1113,6 +1113,18 @@ void UDynamicMaterialModelEditorOnlyData::PostLoad()
 {
 	Super::PostLoad();
 
+	// Backwards compatibility change - materials were originally parented to this object instead of the model.
+	if (IsValid(MaterialModel))
+	{
+		if (UMaterial* Material = MaterialModel->GetGeneratedMaterial())
+		{
+			if (Material->GetOuter() != MaterialModel)
+			{
+				Material->Rename(nullptr, MaterialModel, UE::DynamicMaterial::RenameFlags);
+			}
+		}
+	}
+
 	SetFlags(RF_Transactional);
 
 	ReinitComponents();
