Material Designer
- UI now uses the Custom Details View to create keyframe buttons.

diff --git a/Engine/Plugins/Experimental/DynamicMaterial/Source/DynamicMaterialEditor/Private/Slate/SDMComponentEdit.cpp b/Engine/Plugins/Experimental/DynamicMaterial/Source/DynamicMaterialEditor/Private/Slate/SDMComponentEdit.cpp
--- a/Engine/Plugins/Experimental/DynamicMaterial/Source/DynamicMaterialEditor/Private/Slate/SDMComponentEdit.cpp
+++ b/Engine/Plugins/Experimental/DynamicMaterial/Source/DynamicMaterialEditor/Private/Slate/SDMComponentEdit.cpp
@@ -14,6 +14,7 @@
 #include "Components/DMTextureUV.h"
 #include "CustomDetailsViewArgs.h"
 #include "CustomDetailsViewModule.h"
+#include "CustomDetailsViewSequencer.h"
 #include "DetailLayoutBuilder.h"
 #include "DMEDefs.h"
 #include "DMWorldSubsystem.h"
@@ -535,19 +536,44 @@ TSharedRef<SWidget> SDMComponentEdit::CreateExtensionButtons(const TSharedPtr<SW
 
 	ToolBarBuilder.AddWidget(ResetButtonWidget);
 
+	bool bAddedSequencerButtons = false;
+
 	// Sequencer relies on getting the Keyframe Handler via the Details View of the IDetailTreeNode, but it's null since
 	// there's no Details View here. Instead add it manually.
-	if (bInAllowKeyframe && InComponent->GetWorld() && PropertyHandle.IsValid())
-	{
-		ToolBarBuilder.AddToolBarButton(
-			FUIAction(FExecuteAction::CreateStatic(&SDMComponentEdit::CreateKeyFrame, PropertyHandle)	),
-			NAME_None,
-			FText::GetEmpty(),
-			LOCTEXT("CreateKeyToolTip", "Add a keyframe for this property."),
-			FSlateIcon(FAppStyle::Get().GetStyleSetName(), "Sequencer.AddKey.Details")
-		);
+	if (bInAllowKeyframe && PropertyHandle.IsValid())
+	{
+		if (UWorld* World = InComponent->GetWorld())
+		{
+			if (UDMWorldSubsystem* DMWorldSubsystem = World->GetSubsystem<UDMWorldSubsystem>())
+			{
+				if (const TSharedPtr<IDetailKeyframeHandler>& KeyframeHandler = DMWorldSubsystem->GetKeyframeHandler())
+				{
+					TArray<FPropertyRowExtensionButton> SequencerButtons;
+
+					FCustomDetailsViewSequencerUtils::CreateSequencerExtensionButton(
+						KeyframeHandler,
+						PropertyHandle,
+						SequencerButtons
+					);
+
+					for (const FPropertyRowExtensionButton& SequencerButton : SequencerButtons)
+					{
+						ToolBarBuilder.AddToolBarButton(
+							SequencerButton.UIAction,
+							NAME_None,
+							TAttribute<FText>(),
+							SequencerButton.ToolTip,
+							SequencerButton.Icon
+						);
+
+						bAddedSequencerButtons = true;
+					}
+				}
+			}
+		}
 	}
-	else
+
+	if (!bAddedSequencerButtons)
 	{
 		// Maintain space
 		ToolBarBuilder.AddWidget(
