Material Designer
- Changed slot selector buttons to be checkboxes with appropriate blue styles.

diff --git a/Engine/Plugins/Experimental/DynamicMaterial/Source/DynamicMaterialEditor/Private/Slate/SDMEditor.cpp b/Engine/Plugins/Experimental/DynamicMaterial/Source/DynamicMaterialEditor/Private/Slate/SDMEditor.cpp
--- a/Engine/Plugins/Experimental/DynamicMaterial/Source/DynamicMaterialEditor/Private/Slate/SDMEditor.cpp
+++ b/Engine/Plugins/Experimental/DynamicMaterial/Source/DynamicMaterialEditor/Private/Slate/SDMEditor.cpp
@@ -751,11 +751,14 @@ TSharedRef<SWidget> SDMEditor::CreateSlotPickerWidget()
 	if (Slots.IsValidIndex(0))
 	{
 		SlotSelector->AddSlot()
+			.Padding(3.f, 0.f, 3.f, 0.f)
 			[
-				SNew(SButton)
+				SNew(SCheckBox)
+				.Style(FAppStyle::Get(), "DetailsView.SectionButton")
 				.HAlign(EHorizontalAlignment::HAlign_Center)
-				.ButtonColorAndOpacity(this, &SDMEditor::GetRGBButtonColorAndOpacity_HasSlot)
-				.OnClicked(this, &SDMEditor::OnRGBButtonClicked_HasSlot)
+				.IsChecked(this, &SDMEditor::GetRGBSlotCheckState_HasSlot)
+				.OnCheckStateChanged(this, &SDMEditor::OnRGBSlotCheckStateChanged_HasSlot)
+				.Content()
 				[
 					SNew(STextBlock)
 					.Text(LOCTEXT("RGBSlot", "RGB"))
@@ -765,13 +768,17 @@ TSharedRef<SWidget> SDMEditor::CreateSlotPickerWidget()
 	else
 	{
 		SlotSelector->AddSlot()
+			.Padding(3.f, 0.f, 3.f, 0.f)
 			[
-				SNew(SButton)
+				SNew(SCheckBox)
+				.Style(FAppStyle::Get(), "DetailsView.SectionButton")
 				.HAlign(EHorizontalAlignment::HAlign_Center)
-				.OnClicked(this, &SDMEditor::OnRGBButtonClicked_NoSlot)
+				.IsChecked(this, &SDMEditor::GetRGBSlotCheckState_NoSlot)
+				.OnCheckStateChanged(this, &SDMEditor::OnRGBSlotCheckStateChanged_NoSlot)
+				.Content()
 				[
 					SNew(STextBlock)
-					.Text(LOCTEXT("AddRGBStack", "RGB"))
+					.Text(LOCTEXT("RGBSlot", "RGB"))
 				]
 			];
 	}
@@ -779,12 +786,15 @@ TSharedRef<SWidget> SDMEditor::CreateSlotPickerWidget()
 	if (Slots.IsValidIndex(1))
 	{
 		SlotSelector->AddSlot()
+			.Padding(0.f, 0.f, 3.f, 0.f)
 			[
-				SNew(SButton)
+				SNew(SCheckBox)
+				.Style(FAppStyle::Get(), "DetailsView.SectionButton")
 				.HAlign(EHorizontalAlignment::HAlign_Center)
 				.IsEnabled(this, &SDMEditor::GetOpacityButtonEnabled_HasSlot)
-				.ButtonColorAndOpacity(this, &SDMEditor::GetOpacityButtonColorAndOpacity_HasSlot)
-				.OnClicked(this, &SDMEditor::OnOpacityButtonClicked_HasSlot)
+				.IsChecked(this, &SDMEditor::GetOpacitySlotCheckState_HasSlot)
+				.OnCheckStateChanged(this, &SDMEditor::OnOpacitySlotCheckStateChanged_HasSlot)
+				.Content()
 				[
 					SNew(STextBlock)
 					.Text(LOCTEXT("OpacitySlot", "Opacity"))
@@ -794,11 +804,15 @@ TSharedRef<SWidget> SDMEditor::CreateSlotPickerWidget()
 	else
 	{
 		SlotSelector->AddSlot()
+			.Padding(0.f, 0.f, 3.f, 0.f)
 			[
-				SNew(SButton)
+				SNew(SCheckBox)
+				.Style(FAppStyle::Get(), "DetailsView.SectionButton")
 				.HAlign(EHorizontalAlignment::HAlign_Center)
-				.IsEnabled(this, &SDMEditor::GetOpacityButtonEnabled_NoSlot)
-				.OnClicked(this, &SDMEditor::OnOpacityButtonClicked_NoSlot)
+				.IsEnabled(this, &SDMEditor::GetOpacityButtonEnabled_HasSlot)
+				.IsChecked(this, &SDMEditor::GetOpacitySlotCheckState_NoSlot)
+				.OnCheckStateChanged(this, &SDMEditor::OnOpacitySlotCheckStateChanged_NoSlot)
+				.Content()
 				[
 					SNew(STextBlock)
 					.Text(LOCTEXT("AddOpacity", "Opacity"))
@@ -988,7 +1002,22 @@ bool SDMEditor::IsGlobalOpacityEnabled() const
 	return false;
 }
 
-FReply SDMEditor::OnRGBButtonClicked_NoSlot()
+ECheckBoxState SDMEditor::GetRGBSlotCheckState_HasSlot() const
+{
+	return GetActiveSlotIndex() == 0 ? ECheckBoxState::Checked : ECheckBoxState::Unchecked;
+}
+
+void SDMEditor::OnRGBSlotCheckStateChanged_HasSlot(ECheckBoxState InCheckState)
+{
+	SetActiveSlotIndex(0);
+}
+
+ECheckBoxState SDMEditor::GetRGBSlotCheckState_NoSlot() const
+{
+	return ECheckBoxState::Unchecked;
+}
+
+void SDMEditor::OnRGBSlotCheckStateChanged_NoSlot(ECheckBoxState InCheckState)
 {
 	if (UDynamicMaterialModelEditorOnlyData* ModelEditorOnlyData = UDynamicMaterialModelEditorOnlyData::Get(MaterialModelWeak))
 	{
@@ -998,28 +1027,45 @@ FReply SDMEditor::OnRGBButtonClicked_NoSlot()
 			SetActiveSlotIndex(0);
 		}
 	}
+}
 
-	return FReply::Handled();
+bool SDMEditor::GetOpacityButtonEnabled_NoSlot() const
+{
+	if (UDynamicMaterialModelEditorOnlyData* ModelEditorOnlyData = UDynamicMaterialModelEditorOnlyData::Get(MaterialModelWeak))
+	{
+		return ModelEditorOnlyData->GetDomain() != EMaterialDomain::MD_PostProcess
+			&& ModelEditorOnlyData->GetBlendMode() != EBlendMode::BLEND_Opaque;
+	}
+
+	return false;
 }
 
-FSlateColor SDMEditor::GetRGBButtonColorAndOpacity_HasSlot() const
+bool SDMEditor::GetOpacityButtonEnabled_HasSlot() const
 {
-	if (GetActiveSlotIndex() == 0)
+	if (UDynamicMaterialModelEditorOnlyData* ModelEditorOnlyData = UDynamicMaterialModelEditorOnlyData::Get(MaterialModelWeak))
 	{
-		return FSlateColor(FStyleColors::AccentBlue.GetSpecifiedColor());
+		return ModelEditorOnlyData->GetBlendMode() != EBlendMode::BLEND_Opaque;
 	}
 
-	return FSlateColor(FStyleColors::AccentGray.GetSpecifiedColor());
+	return false;
 }
 
-FReply SDMEditor::OnRGBButtonClicked_HasSlot()
+ECheckBoxState SDMEditor::GetOpacitySlotCheckState_HasSlot() const
 {
-	SetActiveSlotIndex(0);
+	return GetActiveSlotIndex() == 1 ? ECheckBoxState::Checked : ECheckBoxState::Unchecked;
+}
 
-	return FReply::Handled();
+void SDMEditor::OnOpacitySlotCheckStateChanged_HasSlot(ECheckBoxState InCheckState)
+{
+	SetActiveSlotIndex(1);
+}
+
+ECheckBoxState SDMEditor::GetOpacitySlotCheckState_NoSlot() const
+{
+	return ECheckBoxState::Unchecked;
 }
 
-FReply SDMEditor::OnOpacityButtonClicked_NoSlot()
+void SDMEditor::OnOpacitySlotCheckStateChanged_NoSlot(ECheckBoxState InCheckState)
 {
 	if (UDynamicMaterialModelEditorOnlyData* ModelEditorOnlyData = UDynamicMaterialModelEditorOnlyData::Get(MaterialModelWeak))
 	{
@@ -1046,46 +1092,6 @@ FReply SDMEditor::OnOpacityButtonClicked_NoSlot()
 			SetActiveSlotIndex(1);
 		}
 	}
-
-	return FReply::Handled();
-}
-
-bool SDMEditor::GetOpacityButtonEnabled_NoSlot() const
-{
-	if (UDynamicMaterialModelEditorOnlyData* ModelEditorOnlyData = UDynamicMaterialModelEditorOnlyData::Get(MaterialModelWeak))
-	{
-		return ModelEditorOnlyData->GetDomain() != EMaterialDomain::MD_PostProcess
-			&& ModelEditorOnlyData->GetBlendMode() != EBlendMode::BLEND_Opaque;
-	}
-
-	return false;
-}
-
-FSlateColor SDMEditor::GetOpacityButtonColorAndOpacity_HasSlot() const
-{
-	if (GetActiveSlotIndex() == 1)
-	{
-		return FSlateColor(FStyleColors::AccentBlue.GetSpecifiedColor());
-	}
-
-	return FSlateColor(FStyleColors::AccentGray.GetSpecifiedColor());
-}
-
-FReply SDMEditor::OnOpacityButtonClicked_HasSlot()
-{
-	SetActiveSlotIndex(1);
-
-	return FReply::Handled();
-}
-
-bool SDMEditor::GetOpacityButtonEnabled_HasSlot() const
-{
-	if (UDynamicMaterialModelEditorOnlyData* ModelEditorOnlyData = UDynamicMaterialModelEditorOnlyData::Get(MaterialModelWeak))
-	{
-		return ModelEditorOnlyData->GetBlendMode() != EBlendMode::BLEND_Opaque;
-	}
-
-	return false;
 }
 
 FReply SDMEditor::OnCreateMaterialButtonClicked(TWeakPtr<FDMObjectMaterialProperty> InMaterialProperty)
diff --git a/Engine/Plugins/Experimental/DynamicMaterial/Source/DynamicMaterialEditor/Private/Slate/SDMEditor.h b/Engine/Plugins/Experimental/DynamicMaterial/Source/DynamicMaterialEditor/Private/Slate/SDMEditor.h
--- a/Engine/Plugins/Experimental/DynamicMaterial/Source/DynamicMaterialEditor/Private/Slate/SDMEditor.h
+++ b/Engine/Plugins/Experimental/DynamicMaterial/Source/DynamicMaterialEditor/Private/Slate/SDMEditor.h
@@ -151,17 +151,19 @@ protected:
 
 	bool IsGlobalOpacityEnabled() const;
 
-	FReply OnRGBButtonClicked_NoSlot();
+	ECheckBoxState GetRGBSlotCheckState_HasSlot() const;
+	void OnRGBSlotCheckStateChanged_HasSlot(ECheckBoxState InCheckState);
 
-	FSlateColor GetRGBButtonColorAndOpacity_HasSlot() const;
-	FReply OnRGBButtonClicked_HasSlot();
+	ECheckBoxState GetRGBSlotCheckState_NoSlot() const;
+	void OnRGBSlotCheckStateChanged_NoSlot(ECheckBoxState InCheckState);
 
-	FReply OnOpacityButtonClicked_NoSlot();
-	bool GetOpacityButtonEnabled_NoSlot() const;
-
-	FSlateColor GetOpacityButtonColorAndOpacity_HasSlot() const;
-	FReply OnOpacityButtonClicked_HasSlot();
 	bool GetOpacityButtonEnabled_HasSlot() const;
+	ECheckBoxState GetOpacitySlotCheckState_HasSlot() const;
+	void OnOpacitySlotCheckStateChanged_HasSlot(ECheckBoxState InCheckState);
+
+	bool GetOpacityButtonEnabled_NoSlot() const;
+	ECheckBoxState GetOpacitySlotCheckState_NoSlot() const;
+	void OnOpacitySlotCheckStateChanged_NoSlot(ECheckBoxState InCheckState);
 
 	FReply OnCreateMaterialButtonClicked(TWeakPtr<FDMObjectMaterialProperty> InMaterialProperty);
 
