Material Designer
- "Add new effect" text removed.
- Add effect button added next to add new layer button.
- Effects can now be dragged and dropped onto layers.

diff --git a/Engine/Plugins/Experimental/DynamicMaterial/Source/DynamicMaterialEditor/DynamicMaterialEditor.Build.cs b/Engine/Plugins/Experimental/DynamicMaterial/Source/DynamicMaterialEditor/DynamicMaterialEditor.Build.cs
--- a/Engine/Plugins/Experimental/DynamicMaterial/Source/DynamicMaterialEditor/DynamicMaterialEditor.Build.cs
+++ b/Engine/Plugins/Experimental/DynamicMaterial/Source/DynamicMaterialEditor/DynamicMaterialEditor.Build.cs
@@ -24,6 +24,7 @@ public class DynamicMaterialEditor : ModuleRules
 				"ApplicationCore",
 				"AssetDefinition",
 				"ContentBrowser",
+				"ContentBrowserData",
 				"CustomDetailsView",
 				"DeveloperSettings",
 				"EditorWidgets",
diff --git a/Engine/Plugins/Experimental/DynamicMaterial/Source/DynamicMaterialEditor/Private/DynamicMaterialEditorStyle.cpp b/Engine/Plugins/Experimental/DynamicMaterial/Source/DynamicMaterialEditor/Private/DynamicMaterialEditorStyle.cpp
--- a/Engine/Plugins/Experimental/DynamicMaterial/Source/DynamicMaterialEditor/Private/DynamicMaterialEditorStyle.cpp
+++ b/Engine/Plugins/Experimental/DynamicMaterial/Source/DynamicMaterialEditor/Private/DynamicMaterialEditorStyle.cpp
@@ -492,6 +492,7 @@ void FDynamicMaterialEditorStyle::SetupEffectsViewStyles(const TSharedRef<FSlate
 
 	Style->Set("EffectsView.Row.Fx.Closed", new IMAGE_BRUSH_SVG("Icons/Fx_Closed", Icon24x24));
 	Style->Set("EffectsView.Row.Fx.Opened", new IMAGE_BRUSH_SVG("Icons/Fx_Opened", Icon24x24));
+	Style->Set("EffectsView.Row.Fx", new IMAGE_BRUSH_SVG("Icons/Fx", Icon24x24));
 }
 
 void FDynamicMaterialEditorStyle::SetupTextStyles(const TSharedRef<FSlateStyleSet>& Style)
diff --git a/Engine/Plugins/Experimental/DynamicMaterial/Source/DynamicMaterialEditor/Private/Slate/Layers/SDMSlotLayerItem.cpp b/Engine/Plugins/Experimental/DynamicMaterial/Source/DynamicMaterialEditor/Private/Slate/Layers/SDMSlotLayerItem.cpp
--- a/Engine/Plugins/Experimental/DynamicMaterial/Source/DynamicMaterialEditor/Private/Slate/Layers/SDMSlotLayerItem.cpp
+++ b/Engine/Plugins/Experimental/DynamicMaterial/Source/DynamicMaterialEditor/Private/Slate/Layers/SDMSlotLayerItem.cpp
@@ -1,6 +1,8 @@
 // Copyright Epic Games, Inc. All Rights Reserved.
 
 #include "Slate/Layers/SDMSlotLayerItem.h"
+
+#include "ContentBrowserDataDragDropOp.h"
 #include "Components/DMMaterialLayer.h"
 #include "Components/DMMaterialProperty.h"
 #include "Components/DMMaterialSlot.h"
@@ -12,6 +14,7 @@
 #include "DynamicMaterialEditorSettings.h"
 #include "DynamicMaterialEditorStyle.h"
 #include "Framework/Application/SlateApplication.h"
+#include "Materials/MaterialFunctionInterface.h"
 #include "Menus/DMMaterialSlotLayerAddEffectMenus.h"
 #include "Model/DynamicMaterialModelEditorOnlyData.h"
 #include "ScopedTransaction.h"
@@ -19,6 +22,8 @@
 #include "Slate/SDMSlot.h"
 #include "Slate/SDMStage.h"
 #include "SlateOptMacros.h"
+#include "Components/DMMaterialEffectFunction.h"
+#include "Components/DMMaterialEffectStack.h"
 #include "Styling/StyleColors.h"
 #include "Widgets/Images/SImage.h"
 #include "Widgets/Input/SButton.h"
@@ -132,20 +137,6 @@ TSharedRef<SWidget> SDMSlotLayerItem::CreateMainContent()
 			[
 				SAssignNew(EffectsList, SDMLayerEffectsView, SlotWidgetWeak.Pin(), EffectStack)
 			]
-		]
-
-		+ SVerticalBox::Slot()
-		.AutoHeight()
-		.HAlign(HAlign_Fill)
-		.VAlign(VAlign_Center)
-		[
-			SNew(SBox)
-			.HAlign(EHorizontalAlignment::HAlign_Center)
-			.Padding(10.f, 3.f, 5.f, 3.f)
-			.Visibility(this, &SDMSlotLayerItem::GetNewEffectVisibility)
-			[
-				CreateNewEffectButton()
-			]
 		];
 }
 
@@ -592,38 +583,49 @@ TOptional<EItemDropZone> SDMSlotLayerItem::OnLayerItemCanAcceptDrop(const FDragD
 		return TOptional<EItemDropZone>();
 	}
 
-	TSharedPtr<FDMSlotLayerDragDropOperation> SlotLayerDragDropOp = InDragDropEvent.GetOperationAs<FDMSlotLayerDragDropOperation>();
-	if (!SlotLayerDragDropOp.IsValid())
-	{
-		return TOptional<EItemDropZone>();
-	}
-
-	UDMMaterialLayerObject* DraggedSlotLayer = SlotLayerDragDropOp->GetLayer();
-	if (!DraggedSlotLayer)
+	if (TSharedPtr<FDMSlotLayerDragDropOperation> SlotLayerDragDropOp = InDragDropEvent.GetOperationAs<FDMSlotLayerDragDropOperation>())
 	{
-		return TOptional<EItemDropZone>();
-	}
+		UDMMaterialLayerObject* DraggedSlotLayer = SlotLayerDragDropOp->GetLayer();
+		if (!DraggedSlotLayer)
+		{
+			return TOptional<EItemDropZone>();
+		}
 
-	SlotLayerDragDropOp->SetToInvalidDropLocation();
+		SlotLayerDragDropOp->SetToInvalidDropLocation();
 
-	switch (InDropZone)
-	{
-	case EItemDropZone::AboveItem:
-		if (Layer->CanMoveLayerAbove(DraggedSlotLayer))
+		switch (InDropZone)
 		{
-			SlotLayerDragDropOp->SetToValidDropLocation();
-			return InDropZone;
+			case EItemDropZone::AboveItem:
+				if (Layer->CanMoveLayerAbove(DraggedSlotLayer))
+				{
+					SlotLayerDragDropOp->SetToValidDropLocation();
+					return InDropZone;
+				}
+				break;
+
+			case EItemDropZone::OntoItem:
+			case EItemDropZone::BelowItem:
+				if (Layer->CanMoveLayerBelow(DraggedSlotLayer))
+				{
+					SlotLayerDragDropOp->SetToValidDropLocation();
+					return EItemDropZone::BelowItem;
+				}
+				break;
 		}
-		break;
+	}
+	else if (TSharedPtr<FContentBrowserDataDragDropOp> ContentBrowserDragDropOp = InDragDropEvent.GetOperationAs<FContentBrowserDataDragDropOp>())
+	{
+		bool bHasValidItem = false;
 
-	case EItemDropZone::OntoItem:
-	case EItemDropZone::BelowItem:
-		if (Layer->CanMoveLayerBelow(DraggedSlotLayer))
+		for (const FAssetData& DraggedAsset : ContentBrowserDragDropOp->GetAssets())
 		{
-			SlotLayerDragDropOp->SetToValidDropLocation();
-			return EItemDropZone::BelowItem;
+			UClass* AssetClass = DraggedAsset.GetClass(EResolveClass::Yes);
+
+			if (AssetClass && AssetClass->IsChildOf(UMaterialFunctionInterface::StaticClass()))
+			{
+				return EItemDropZone::OntoItem;
+			}
 		}
-		break;
 	}
 
 	return TOptional<EItemDropZone>();
@@ -651,29 +653,41 @@ FReply SDMSlotLayerItem::OnLayerItemDragDetected(const FGeometry& InMyGeometry,
 FReply SDMSlotLayerItem::OnLayerItemAcceptDrop(const FDragDropEvent& InDragDropEvent, EItemDropZone InDropZone, 
 	TSharedPtr<FDMMaterialLayerReference> InSlotLayer)
 {
-	TSharedPtr<FDMSlotLayerDragDropOperation> SlotLayerDragDropOp = InDragDropEvent.GetOperationAs<FDMSlotLayerDragDropOperation>();
+	if (!InSlotLayer.IsValid())
+	{
+		return FReply::Handled();
+	}
+
+	const UDMMaterialLayerObject* DraggedOverLayer = GetLayer();
+
+	if (!DraggedOverLayer)
+	{
+		return FReply::Handled();
+	}
 
-	if (SlotLayerDragDropOp.IsValid())
+	if (TSharedPtr<FDMSlotLayerDragDropOperation> SlotLayerDragDropOp = InDragDropEvent.GetOperationAs<FDMSlotLayerDragDropOperation>())
 	{
-		const UDMMaterialLayerObject* DraggedOverLayer = InSlotLayer.IsValid() ? InSlotLayer->GetLayer() : nullptr;
-		if (!DraggedOverLayer || !DraggedOverLayer->GetStage(EDMMaterialLayerStage::Base))
+		if (!DraggedOverLayer->GetStage(EDMMaterialLayerStage::Base))
 		{
 			return FReply::Handled();
 		}
 
 		UDMMaterialLayerObject* DraggedLayer = SlotLayerDragDropOp->GetLayer();
+
 		if (!DraggedLayer || !DraggedLayer->GetStage(EDMMaterialLayerStage::Base))
 		{
 			return FReply::Handled();
 		}
 
 		const int32 ThisLayerIndex = GetLayerItemIndex();
+
 		if (ThisLayerIndex == INDEX_NONE)
 		{
 			return FReply::Handled();
 		}
 		
 		UDMMaterialSlot* const Slot = DraggedLayer->GetSlot();
+
 		if (!Slot)
 		{
 			return FReply::Handled();
@@ -697,6 +711,39 @@ FReply SDMSlotLayerItem::OnLayerItemAcceptDrop(const FDragDropEvent& InDragDropE
 
 		return FReply::Handled();
 	}
+
+	if (TSharedPtr<FContentBrowserDataDragDropOp> ContentBrowserDragDropOp = InDragDropEvent.GetOperationAs<FContentBrowserDataDragDropOp>())
+	{
+		UDMMaterialEffectStack* EffectStack = DraggedOverLayer->GetEffectStack();
+
+		if (!EffectStack)
+		{
+			return FReply::Handled();
+		}
+
+		for (const FAssetData& DraggedAsset : ContentBrowserDragDropOp->GetAssets())
+		{
+			UClass* AssetClass = DraggedAsset.GetClass(EResolveClass::Yes);
+
+			if (AssetClass && AssetClass->IsChildOf(UMaterialFunctionInterface::StaticClass()))
+			{
+				if (UMaterialFunctionInterface* MaterialFunction = Cast<UMaterialFunctionInterface>(DraggedAsset.GetAsset()))
+				{
+					UDMMaterialEffectFunction* EffectFunction = UDMMaterialEffect::CreateEffect<UDMMaterialEffectFunction>(EffectStack);
+					EffectFunction->SetMaterialFunction(MaterialFunction);
+
+					// Will successfully set the function if it's valid
+					if (EffectFunction->GetMaterialFunction() == MaterialFunction)
+					{
+						EffectStack->AddEffect(EffectFunction);
+					}
+				}
+			}
+		}
+
+		return FReply::Handled();
+	}
+
 	return FReply::Unhandled();
 }
 
@@ -969,40 +1016,6 @@ EVisibility SDMSlotLayerItem::GetEffectsListVisibility() const
 	return EffectsList->GetLayerItemCount() > 0 && bDisplayEffectsList ? EVisibility::Visible : EVisibility::Collapsed;
 }
 
-EVisibility SDMSlotLayerItem::GetNewEffectVisibility() const
-{
-	return bDisplayEffectsList ? EVisibility::Visible : EVisibility::Collapsed;
-}
-
-TSharedRef<SWidget> SDMSlotLayerItem::CreateNewEffectButton()
-{
-	TSharedRef<SWidget> TextBlock = SNew(STextBlock)
-		.Font(FAppStyle::GetFontStyle("PropertyWindow.NormalFont"))
-		.Text(LOCTEXT("AddNewEffect", "Add New Effect"))
-		.ColorAndOpacity(FSlateColor(EStyleColor::PrimaryHover));
-
-	TextBlock->SetOnMouseButtonDown(FPointerEventHandler::CreateSPLambda(
-		TextBlock,
-		[this, TextBlockWeak = TextBlock->AsShared().ToWeakPtr()](const FGeometry& InGeometry, const FPointerEvent& InPointerEvent)
-		{
-			if (TSharedPtr<SWidget> TextBlock = TextBlockWeak.Pin())
-			{
-				FSlateApplication::Get().PushMenu(
-					TextBlock.ToSharedRef(),
-					FWidgetPath(),
-					CreateNewEffectMenu(),
-					GetTickSpaceGeometry().GetAbsolutePosition(),
-					FPopupTransitionEffect::TopMenu
-				);
-			}
-
-			return FReply::Handled();
-		}
-	));
-
-	return TextBlock;
-}
-
 TSharedRef<SWidget> SDMSlotLayerItem::CreateNewEffectMenu()
 {
 	if (LayerItem.IsValid())
diff --git a/Engine/Plugins/Experimental/DynamicMaterial/Source/DynamicMaterialEditor/Private/Slate/Layers/SDMSlotLayerItem.h b/Engine/Plugins/Experimental/DynamicMaterial/Source/DynamicMaterialEditor/Private/Slate/Layers/SDMSlotLayerItem.h
--- a/Engine/Plugins/Experimental/DynamicMaterial/Source/DynamicMaterialEditor/Private/Slate/Layers/SDMSlotLayerItem.h
+++ b/Engine/Plugins/Experimental/DynamicMaterial/Source/DynamicMaterialEditor/Private/Slate/Layers/SDMSlotLayerItem.h
@@ -139,9 +139,6 @@ protected:
 
 	EVisibility GetEffectsListVisibility() const;
 
-	EVisibility GetNewEffectVisibility() const;
-	TSharedRef<SWidget> CreateNewEffectButton();
-
 	TSharedRef<SWidget> CreateNewEffectMenu();
 
 	FReply OnCreateLayerBypassButtonClicked();
diff --git a/Engine/Plugins/Experimental/DynamicMaterial/Source/DynamicMaterialEditor/Private/Slate/SDMSlot.cpp b/Engine/Plugins/Experimental/DynamicMaterial/Source/DynamicMaterialEditor/Private/Slate/SDMSlot.cpp
--- a/Engine/Plugins/Experimental/DynamicMaterial/Source/DynamicMaterialEditor/Private/Slate/SDMSlot.cpp
+++ b/Engine/Plugins/Experimental/DynamicMaterial/Source/DynamicMaterialEditor/Private/Slate/SDMSlot.cpp
@@ -22,6 +22,7 @@
 #include "DynamicMaterialEditorSettings.h"
 #include "DynamicMaterialEditorStyle.h"
 #include "DynamicMaterialModule.h"
+#include "Menus/DMMaterialSlotLayerAddEffectMenus.h"
 #include "Menus/DMMaterialSlotMenus.h"
 #include "Model/DynamicMaterialModel.h"
 #include "Model/DynamicMaterialModelEditorOnlyData.h"
@@ -165,6 +166,21 @@ TSharedRef<SWidget> SDMSlot::GetLayerButtonsMenuContent()
 	return FDMMaterialSlotMenus::MakeAddLayerButtonMenu(SharedThis(this));
 }
 
+bool SDMSlot::GetLayerCanAddEffect() const
+{
+	return !!GetSelectedLayer();
+}
+
+TSharedRef<SWidget> SDMSlot::GetLayerEffectsMenuContent()
+{
+	if (UDMMaterialLayerObject* LayerObject = GetSelectedLayer())
+	{
+		return FDMMaterialSlotLayerAddEffectMenus::OpenAddEffectMenu(LayerObject);
+	}
+
+	return SNullWidget::NullWidget;
+}
+
 bool SDMSlot::GetLayerRowsButtonsCanDuplicate() const
 {
 	return LayerView->GetSelectedItems().IsEmpty() == false;
@@ -585,7 +601,7 @@ TSharedRef<SWidget> SDMSlot::CreateLayerButtonsRowWidget()
 			[
 				SNew(SImage)
 				.Image(FAppStyle::Get().GetBrush("GenericCommands.Delete"))
-				.DesiredSizeOverride(FVector2D(22.0f))
+				.DesiredSizeOverride(FVector2D(16.0f))
 			]
 		]
 		+ SHorizontalBox::Slot()
@@ -602,6 +618,26 @@ TSharedRef<SWidget> SDMSlot::CreateLayerButtonsRowWidget()
 		.AutoWidth()
 		.HAlign(HAlign_Right)
 		.VAlign(VAlign_Center)
+		[
+			SNew(SComboButton)
+			.HasDownArrow(false)
+			.IsFocusable(true)
+			.ContentPadding(4.0f)
+			.ButtonStyle(FDynamicMaterialEditorStyle::Get(), "HoverHintOnly.Bordered.Dark")
+			.ToolTipText(LOCTEXT("AddLayerEffecTooltip", "Add Layer Effect"))
+			.IsEnabled(this, &SDMSlot::GetLayerCanAddEffect)
+			.OnGetMenuContent(this, &SDMSlot::GetLayerEffectsMenuContent)
+			.ButtonContent()
+			[
+				SNew(SImage)
+				.Image(FDynamicMaterialEditorStyle::GetBrush("EffectsView.Row.Fx"))
+				.DesiredSizeOverride(FVector2D(16.0f))
+			]
+		]
+		+ SHorizontalBox::Slot()
+		.AutoWidth()
+		.HAlign(HAlign_Right)
+		.VAlign(VAlign_Center)
 		[
 			SNew(SComboButton)
 			.HasDownArrow(false)
@@ -615,7 +651,7 @@ TSharedRef<SWidget> SDMSlot::CreateLayerButtonsRowWidget()
 				SNew(SImage)
 				.Image(FAppStyle::Get().GetBrush("Icons.Plus"))
 				.ColorAndOpacity(FDynamicMaterialEditorStyle::Get().GetColor("Color.Stage.Enabled"))
-				.DesiredSizeOverride(FVector2D(22.0f))
+				.DesiredSizeOverride(FVector2D(16.0f))
 			]
 		]
 		+ SHorizontalBox::Slot()
@@ -633,7 +669,7 @@ TSharedRef<SWidget> SDMSlot::CreateLayerButtonsRowWidget()
 			[
 				SNew(SImage)
 				.Image(FDynamicMaterialEditorStyle::GetBrush("LayerView.DuplicateIcon"))
-				.DesiredSizeOverride(FVector2D(22.0f))
+				.DesiredSizeOverride(FVector2D(16.0f))
 			]
 		]
 		+ SHorizontalBox::Slot()
@@ -651,7 +687,7 @@ TSharedRef<SWidget> SDMSlot::CreateLayerButtonsRowWidget()
 			[
 				SNew(SImage)
 				.Image(FAppStyle::Get().GetBrush("Icons.Delete"))
-				.DesiredSizeOverride(FVector2D(22.0f))
+				.DesiredSizeOverride(FVector2D(16.0f))
 			]
 		];
 }
diff --git a/Engine/Plugins/Experimental/DynamicMaterial/Source/DynamicMaterialEditor/Private/Slate/SDMSlot.h b/Engine/Plugins/Experimental/DynamicMaterial/Source/DynamicMaterialEditor/Private/Slate/SDMSlot.h
--- a/Engine/Plugins/Experimental/DynamicMaterial/Source/DynamicMaterialEditor/Private/Slate/SDMSlot.h
+++ b/Engine/Plugins/Experimental/DynamicMaterial/Source/DynamicMaterialEditor/Private/Slate/SDMSlot.h
@@ -155,6 +155,9 @@ protected:
 	FText GetLayerButtonsDescription() const;
 	TSharedRef<SWidget> GetLayerButtonsMenuContent();
 
+	bool GetLayerCanAddEffect() const;
+	TSharedRef<SWidget> GetLayerEffectsMenuContent();
+
 	bool GetLayerRowsButtonsCanDuplicate() const;
 	FReply OnLayerRowButtonsDuplicateClicked();
 
