Material Designer
- Swapped caption and description in material effects.

#virtualized

diff --git a/Engine/Plugins/Experimental/DynamicMaterial/Source/DynamicMaterialEditor/Private/Components/DMMaterialEffectFunction.cpp b/Engine/Plugins/Experimental/DynamicMaterial/Source/DynamicMaterialEditor/Private/Components/DMMaterialEffectFunction.cpp
--- a/Engine/Plugins/Experimental/DynamicMaterial/Source/DynamicMaterialEditor/Private/Components/DMMaterialEffectFunction.cpp
+++ b/Engine/Plugins/Experimental/DynamicMaterial/Source/DynamicMaterialEditor/Private/Components/DMMaterialEffectFunction.cpp
@@ -78,11 +78,11 @@ FText UDMMaterialEffectFunction::GetEffectName() const
 {
 	if (UMaterialFunctionInterface* MaterialFunction = MaterialFunctionPtr.Get())
 	{
-		const FString& Description = MaterialFunction->GetDescription();
+		const FString Caption = MaterialFunction->GetUserExposedCaption();
 
-		if (!Description.IsEmpty())
+		if (!Caption.IsEmpty())
 		{
-			return FText::FromString(Description);
+			return FText::FromString(Caption);
 		}
 	}
 
@@ -94,11 +94,11 @@ FText UDMMaterialEffectFunction::GetEffectDescription() const
 {
 	if (UMaterialFunctionInterface* MaterialFunction = MaterialFunctionPtr.Get())
 	{
-		const FString Caption = MaterialFunction->GetUserExposedCaption();
+		const FString& Description = MaterialFunction->GetDescription();
 
-		if (!Caption.IsEmpty())
+		if (!Description.IsEmpty())
 		{
-			return FText::FromString(Caption);
+			return FText::FromString(Description);
 		}
 	}
 
