Material Desiger
- Fix for crash when changing maps.

diff --git a/Engine/Plugins/Experimental/DynamicMaterial/Source/DynamicMaterialEditor/Private/Slate/Previews/SDMMaterialValuePreview.cpp b/Engine/Plugins/Experimental/DynamicMaterial/Source/DynamicMaterialEditor/Private/Slate/Previews/SDMMaterialValuePreview.cpp
--- a/Engine/Plugins/Experimental/DynamicMaterial/Source/DynamicMaterialEditor/Private/Slate/Previews/SDMMaterialValuePreview.cpp
+++ b/Engine/Plugins/Experimental/DynamicMaterial/Source/DynamicMaterialEditor/Private/Slate/Previews/SDMMaterialValuePreview.cpp
@@ -53,6 +53,11 @@ void SDMMaterialValuePreview::Tick(const FGeometry& AllottedGeometry, const doub
 {
 	SCompoundWidget::Tick(AllottedGeometry, InCurrentTime, InDeltaTime);
 
+	if (!PreviewMaterialWeak.IsValid())
+	{
+		Brush.SetMaterial(nullptr);
+	}
+
 	if (UDMMaterialComponent::CanClean())
 	{
 		if (UDMMaterialValue* Value = ValueWeak.Get())
@@ -77,6 +82,7 @@ void SDMMaterialValuePreview::OnValueUpdated(UDMMaterialComponent* InComponent,
 			if (Brush.GetResourceObject() != PreviewMaterial)
 			{
 				Brush.SetMaterial(PreviewMaterial);
+				PreviewMaterialWeak = PreviewMaterial;
 			}
 		}
 	}
diff --git a/Engine/Plugins/Experimental/DynamicMaterial/Source/DynamicMaterialEditor/Private/Slate/Previews/SDMMaterialValuePreview.h b/Engine/Plugins/Experimental/DynamicMaterial/Source/DynamicMaterialEditor/Private/Slate/Previews/SDMMaterialValuePreview.h
--- a/Engine/Plugins/Experimental/DynamicMaterial/Source/DynamicMaterialEditor/Private/Slate/Previews/SDMMaterialValuePreview.h
+++ b/Engine/Plugins/Experimental/DynamicMaterial/Source/DynamicMaterialEditor/Private/Slate/Previews/SDMMaterialValuePreview.h
@@ -4,6 +4,8 @@
 
 #include "DMEDefs.h"
 #include "SlateMaterialBrush.h"
+#include "UObject/WeakObjectPtr.h"
+#include "UObject/WeakObjectPtrTemplates.h"
 #include "Widgets/SCompoundWidget.h"
 
 class UDMMaterialComponent;
@@ -19,7 +21,7 @@ class SDMMaterialValuePreview : public SCompoundWidget
 
 public:
 	SDMMaterialValuePreview();
-	virtual ~SDMMaterialValuePreview();
+	virtual ~SDMMaterialValuePreview() override;
 
 	void Construct(const FArguments& InArgs, UDMMaterialValue* InValue);
 
@@ -27,6 +29,7 @@ public:
 
 protected:
 	TWeakObjectPtr<UDMMaterialValue> ValueWeak;
+	TWeakObjectPtr<UMaterialInterface> PreviewMaterialWeak;
 	FSlateMaterialBrush Brush;
 
 	void OnValueUpdated(UDMMaterialComponent* InComponent, EDMUpdateType InUpdateType);
diff --git a/Engine/Plugins/Experimental/DynamicMaterial/Source/DynamicMaterialEditor/Private/Slate/Previews/SDMSlotPreview.cpp b/Engine/Plugins/Experimental/DynamicMaterial/Source/DynamicMaterialEditor/Private/Slate/Previews/SDMSlotPreview.cpp
--- a/Engine/Plugins/Experimental/DynamicMaterial/Source/DynamicMaterialEditor/Private/Slate/Previews/SDMSlotPreview.cpp
+++ b/Engine/Plugins/Experimental/DynamicMaterial/Source/DynamicMaterialEditor/Private/Slate/Previews/SDMSlotPreview.cpp
@@ -97,6 +97,7 @@ void SDMSlotPreview::OnSlotUpdated(UDMMaterialComponent* InComponent, EDMUpdateT
 			if (Brush.GetResourceObject() != PreviewMaterial)
 			{
 				Brush.SetMaterial(PreviewMaterial);
+				PreviewMaterialWeak = PreviewMaterial;
 			}
 		}
 	}
@@ -116,6 +117,11 @@ void SDMSlotPreview::Tick(const FGeometry& InAllottedGeometry, const double InCu
 {
 	SCompoundWidget::Tick(InAllottedGeometry, InCurrentTime, InDeltaTime);
 
+	if (!PreviewMaterialWeak.IsValid())
+	{
+		Brush.SetMaterial(nullptr);
+	}
+
 	if (UDMMaterialComponent::CanClean())
 	{
 		if (UDMMaterialSlot* Slot = SlotWeak.Get())
diff --git a/Engine/Plugins/Experimental/DynamicMaterial/Source/DynamicMaterialEditor/Private/Slate/Previews/SDMSlotPreview.h b/Engine/Plugins/Experimental/DynamicMaterial/Source/DynamicMaterialEditor/Private/Slate/Previews/SDMSlotPreview.h
--- a/Engine/Plugins/Experimental/DynamicMaterial/Source/DynamicMaterialEditor/Private/Slate/Previews/SDMSlotPreview.h
+++ b/Engine/Plugins/Experimental/DynamicMaterial/Source/DynamicMaterialEditor/Private/Slate/Previews/SDMSlotPreview.h
@@ -5,6 +5,8 @@
 #include "DMDefs.h"
 #include "SlateMaterialBrush.h"
 #include "Types/WidgetMouseEventsDelegate.h"
+#include "UObject/WeakObjectPtr.h"
+#include "UObject/WeakObjectPtrTemplates.h"
 #include "Widgets/SCompoundWidget.h"
 
 class UDMMaterialComponent;
@@ -22,7 +24,7 @@ class SDMSlotPreview : public SCompoundWidget
 
 public:
 	SDMSlotPreview();
-	virtual ~SDMSlotPreview();
+	virtual ~SDMSlotPreview() override;
 
 	void Construct(const FArguments& InArgs, UDMMaterialSlot* InSlot, EDMMaterialLayerStage InLayerStage);
 
@@ -35,6 +37,7 @@ public:
 
 protected:
 	TWeakObjectPtr<UDMMaterialSlot> SlotWeak;
+	TWeakObjectPtr<UMaterialInterface> PreviewMaterialWeak;
 	EDMMaterialLayerStage LayerStage;
 	FSlateMaterialBrush Brush;
 
diff --git a/Engine/Plugins/Experimental/DynamicMaterial/Source/DynamicMaterialEditor/Private/Slate/Previews/SDMStagePreview.cpp b/Engine/Plugins/Experimental/DynamicMaterial/Source/DynamicMaterialEditor/Private/Slate/Previews/SDMStagePreview.cpp
--- a/Engine/Plugins/Experimental/DynamicMaterial/Source/DynamicMaterialEditor/Private/Slate/Previews/SDMStagePreview.cpp
+++ b/Engine/Plugins/Experimental/DynamicMaterial/Source/DynamicMaterialEditor/Private/Slate/Previews/SDMStagePreview.cpp
@@ -52,6 +52,11 @@ void SDMStagePreview::Tick(const FGeometry& AllottedGeometry, const double InCur
 {
 	SCompoundWidget::Tick(AllottedGeometry, InCurrentTime, InDeltaTime);
 
+	if (!PreviewMaterialWeak.IsValid())
+	{
+		Brush.SetMaterial(nullptr);
+	}
+
 	if (UDMMaterialComponent::CanClean())
 	{
 		if (UDMMaterialStage* Stage = StageWeak.Get())
@@ -76,6 +81,7 @@ void SDMStagePreview::OnStageUpdated(UDMMaterialComponent* InComponent, EDMUpdat
 			if (Brush.GetResourceObject() != PreviewMaterial)
 			{
 				Brush.SetMaterial(PreviewMaterial);
+				PreviewMaterialWeak = PreviewMaterial;
 			}
 		}
 	}
diff --git a/Engine/Plugins/Experimental/DynamicMaterial/Source/DynamicMaterialEditor/Private/Slate/Previews/SDMStagePreview.h b/Engine/Plugins/Experimental/DynamicMaterial/Source/DynamicMaterialEditor/Private/Slate/Previews/SDMStagePreview.h
--- a/Engine/Plugins/Experimental/DynamicMaterial/Source/DynamicMaterialEditor/Private/Slate/Previews/SDMStagePreview.h
+++ b/Engine/Plugins/Experimental/DynamicMaterial/Source/DynamicMaterialEditor/Private/Slate/Previews/SDMStagePreview.h
@@ -4,6 +4,8 @@
 
 #include "DMEDefs.h"
 #include "SlateMaterialBrush.h"
+#include "UObject/WeakObjectPtr.h"
+#include "UObject/WeakObjectPtrTemplates.h"
 #include "Widgets/SCompoundWidget.h"
 
 class UDMMaterialComponent;
@@ -29,6 +31,7 @@ public:
 
 protected:
 	TWeakObjectPtr<UDMMaterialStage> StageWeak;
+	TWeakObjectPtr<UMaterialInterface> PreviewMaterialWeak;
 	FSlateMaterialBrush Brush;
 	TAttribute<FVector2D> PreviewSize;
 
diff --git a/Engine/Plugins/Experimental/DynamicMaterial/Source/DynamicMaterialEditor/Private/Slate/Previews/SDMStageSourcePreview.cpp b/Engine/Plugins/Experimental/DynamicMaterial/Source/DynamicMaterialEditor/Private/Slate/Previews/SDMStageSourcePreview.cpp
--- a/Engine/Plugins/Experimental/DynamicMaterial/Source/DynamicMaterialEditor/Private/Slate/Previews/SDMStageSourcePreview.cpp
+++ b/Engine/Plugins/Experimental/DynamicMaterial/Source/DynamicMaterialEditor/Private/Slate/Previews/SDMStageSourcePreview.cpp
@@ -53,6 +53,11 @@ void SDMStageSourcePreview::Tick(const FGeometry& AllottedGeometry, const double
 {
 	SCompoundWidget::Tick(AllottedGeometry, InCurrentTime, InDeltaTime);
 
+	if (!PreviewMaterialWeak.IsValid())
+	{
+		Brush.SetMaterial(nullptr);
+	}
+
 	if (UDMMaterialComponent::CanClean())
 	{
 		if (UDMMaterialStageSource* Source = StageSourceWeak.Get())
@@ -77,6 +82,7 @@ void SDMStageSourcePreview::OnStageSourceUpdated(UDMMaterialComponent* InCompone
 			if (Brush.GetResourceObject() != PreviewMaterial)
 			{
 				Brush.SetMaterial(PreviewMaterial);
+				PreviewMaterialWeak = PreviewMaterial;
 			}
 		}
 	}
diff --git a/Engine/Plugins/Experimental/DynamicMaterial/Source/DynamicMaterialEditor/Private/Slate/Previews/SDMStageSourcePreview.h b/Engine/Plugins/Experimental/DynamicMaterial/Source/DynamicMaterialEditor/Private/Slate/Previews/SDMStageSourcePreview.h
--- a/Engine/Plugins/Experimental/DynamicMaterial/Source/DynamicMaterialEditor/Private/Slate/Previews/SDMStageSourcePreview.h
+++ b/Engine/Plugins/Experimental/DynamicMaterial/Source/DynamicMaterialEditor/Private/Slate/Previews/SDMStageSourcePreview.h
@@ -4,6 +4,8 @@
 
 #include "DMEDefs.h"
 #include "SlateMaterialBrush.h"
+#include "UObject/WeakObjectPtr.h"
+#include "UObject/WeakObjectPtrTemplates.h"
 #include "Widgets/SCompoundWidget.h"
 
 class UDMMaterialComponent;
@@ -20,7 +22,7 @@ class SDMStageSourcePreview : public SCompoundWidget
 public:
 
 	SDMStageSourcePreview();
-	virtual ~SDMStageSourcePreview();
+	virtual ~SDMStageSourcePreview() override;
 
 	void Construct(const FArguments& InArgs, UDMMaterialStageSource* InStageSource);
 
@@ -28,6 +30,7 @@ public:
 
 protected:
 	TWeakObjectPtr<UDMMaterialStageSource> StageSourceWeak;
+	TWeakObjectPtr<UMaterialInterface> PreviewMaterialWeak;
 	FSlateMaterialBrush Brush;
 
 	void OnStageSourceUpdated(UDMMaterialComponent* InComponent, EDMUpdateType InUpdateType);
