Material Designer
- Fixed add effect menu

diff --git a/Engine/Plugins/Experimental/DynamicMaterial/Source/DynamicMaterialEditor/Private/Menus/DMMaterialSlotLayerAddEffectMenus.cpp b/Engine/Plugins/Experimental/DynamicMaterial/Source/DynamicMaterialEditor/Private/Menus/DMMaterialSlotLayerAddEffectMenus.cpp
--- a/Engine/Plugins/Experimental/DynamicMaterial/Source/DynamicMaterialEditor/Private/Menus/DMMaterialSlotLayerAddEffectMenus.cpp
+++ b/Engine/Plugins/Experimental/DynamicMaterial/Source/DynamicMaterialEditor/Private/Menus/DMMaterialSlotLayerAddEffectMenus.cpp
@@ -151,7 +151,7 @@ namespace UE::DynamicMaterialEditor::Private
 			Section.AddMenuEntry(
 				FName(*Description),
 				FText::FromString(Description),
-				FText::FromString(MaterialFunction->GetDescription()),
+				FText::FromString(MaterialFunction->GetUserExposedCaption()),
 				FSlateIcon(),
 				FToolUIActionChoice(Action)
 			);
