Material Designer
- Actual fix to the add effect menu

diff --git a/Engine/Plugins/Experimental/DynamicMaterial/Source/DynamicMaterialEditor/Private/Menus/DMMaterialSlotLayerAddEffectMenus.cpp b/Engine/Plugins/Experimental/DynamicMaterial/Source/DynamicMaterialEditor/Private/Menus/DMMaterialSlotLayerAddEffectMenus.cpp
--- a/Engine/Plugins/Experimental/DynamicMaterial/Source/DynamicMaterialEditor/Private/Menus/DMMaterialSlotLayerAddEffectMenus.cpp
+++ b/Engine/Plugins/Experimental/DynamicMaterial/Source/DynamicMaterialEditor/Private/Menus/DMMaterialSlotLayerAddEffectMenus.cpp
@@ -146,12 +146,13 @@ namespace UE::DynamicMaterialEditor::Private
 			Action.ExecuteAction = FToolMenuExecuteAction::CreateStatic(&AddEffect, Effect);
 			Action.CanExecuteAction = FToolMenuCanExecuteAction::CreateStatic(&CanAddEffect, Effect);
 
-			const FString& Description = MaterialFunction->GetDescription();
+			const FString& Description = MaterialFunction->GetUserExposedCaption();
+			const FString ToolTip = MaterialFunction->GetDescription();
 
 			Section.AddMenuEntry(
 				FName(*Description),
 				FText::FromString(Description),
-				FText::FromString(MaterialFunction->GetUserExposedCaption()),
+				FText::FromString(ToolTip),
 				FSlateIcon(),
 				FToolUIActionChoice(Action)
 			);
