Material Designer
- "Animated" is now disabled when in translucent view modes.

diff --git a/Engine/Plugins/Experimental/DynamicMaterial/Source/DynamicMaterialEditor/Private/Model/DynamicMaterialModelEditorOnlyData.cpp b/Engine/Plugins/Experimental/DynamicMaterial/Source/DynamicMaterialEditor/Private/Model/DynamicMaterialModelEditorOnlyData.cpp
--- a/Engine/Plugins/Experimental/DynamicMaterial/Source/DynamicMaterialEditor/Private/Model/DynamicMaterialModelEditorOnlyData.cpp
+++ b/Engine/Plugins/Experimental/DynamicMaterial/Source/DynamicMaterialEditor/Private/Model/DynamicMaterialModelEditorOnlyData.cpp
@@ -755,6 +755,11 @@ void UDynamicMaterialModelEditorOnlyData::SetBlendMode(TEnumAsByte<EBlendMode> I
 		}
 	}
 
+	if (BlendMode == EBlendMode::BLEND_Translucent)
+	{
+		SetPixelAnimationFlag(false);
+	}
+
 	RequestMaterialBuild();
 }
 
diff --git a/Engine/Plugins/Experimental/DynamicMaterial/Source/DynamicMaterialEditor/Private/Slate/SDMEditor.cpp b/Engine/Plugins/Experimental/DynamicMaterial/Source/DynamicMaterialEditor/Private/Slate/SDMEditor.cpp
--- a/Engine/Plugins/Experimental/DynamicMaterial/Source/DynamicMaterialEditor/Private/Slate/SDMEditor.cpp
+++ b/Engine/Plugins/Experimental/DynamicMaterial/Source/DynamicMaterialEditor/Private/Slate/SDMEditor.cpp
@@ -598,7 +598,7 @@ TSharedRef<SWidget> SDMEditor::CreateMaterialSettingsRow()
 		.Padding(10.0f, 0.0f, 0.0f, 2.0f)
 		[
 			SNew(SHorizontalBox)
-			.ToolTipText(LOCTEXT("MaterialDesignerInstanceTypeTooltip", "Enables the material's TSR pixel animation flag."))
+			.ToolTipText(LOCTEXT("MaterialDesignerInstanceTypeTooltip", "Enables the material's TSR pixel animation flag. Not available in translucent blend modes."))
 			+ SHorizontalBox::Slot()
 			.AutoWidth()
 			.HAlign(HAlign_Left)
@@ -618,6 +618,7 @@ TSharedRef<SWidget> SDMEditor::CreateMaterialSettingsRow()
 			[
 				SNew(SCheckBox)
 				.IsChecked(this, &SDMEditor::IsMaterialAnimated)
+				.IsEnabled(this, &SDMEditor::CanMaterialBeAnimated)
 				.OnCheckStateChanged(this, &SDMEditor::OnMaterialAnimatedChanged)
 			]
 		]
@@ -1463,6 +1464,16 @@ void SDMEditor::OnMaterialUnlitChanged(const ECheckBoxState InNewCheckState)
 	}
 }
 
+bool SDMEditor::CanMaterialBeAnimated() const
+{
+	if (UDynamicMaterialModelEditorOnlyData* ModelEditorOnlyData = UDynamicMaterialModelEditorOnlyData::Get(MaterialModelWeak))
+	{
+		return ModelEditorOnlyData->GetBlendMode() != EBlendMode::BLEND_Translucent;
+	}
+
+	return false;
+}
+
 bool SDMEditor::CanChangeMaterialShadingModel() const
 {
 	if (UDynamicMaterialModelEditorOnlyData* ModelEditorOnlyData = UDynamicMaterialModelEditorOnlyData::Get(MaterialModelWeak))
diff --git a/Engine/Plugins/Experimental/DynamicMaterial/Source/DynamicMaterialEditor/Private/Slate/SDMEditor.h b/Engine/Plugins/Experimental/DynamicMaterial/Source/DynamicMaterialEditor/Private/Slate/SDMEditor.h
--- a/Engine/Plugins/Experimental/DynamicMaterial/Source/DynamicMaterialEditor/Private/Slate/SDMEditor.h
+++ b/Engine/Plugins/Experimental/DynamicMaterial/Source/DynamicMaterialEditor/Private/Slate/SDMEditor.h
@@ -183,6 +183,7 @@ protected:
 	ECheckBoxState IsMaterialUnlit() const;
 	void OnMaterialUnlitChanged(const ECheckBoxState InNewCheckState);
 
+	bool CanMaterialBeAnimated() const;
 	ECheckBoxState IsMaterialAnimated() const;
 	void OnMaterialAnimatedChanged(const ECheckBoxState InNewCheckState);
 
