Material Designer
- Removed material properties as options on the toolbar. Only slots remain.

diff --git a/Engine/Plugins/Experimental/DynamicMaterial/Source/DynamicMaterialEditor/Private/DMBlueprintFunctionLibrary.cpp b/Engine/Plugins/Experimental/DynamicMaterial/Source/DynamicMaterialEditor/Private/DMBlueprintFunctionLibrary.cpp
--- a/Engine/Plugins/Experimental/DynamicMaterial/Source/DynamicMaterialEditor/Private/DMBlueprintFunctionLibrary.cpp
+++ b/Engine/Plugins/Experimental/DynamicMaterial/Source/DynamicMaterialEditor/Private/DMBlueprintFunctionLibrary.cpp
@@ -194,57 +194,6 @@ TArray<FDMObjectMaterialProperty> UDMBlueprintFunctionLibrary::GetActorMaterialP
 			}
 		});
 
-	auto IterateProperies = [&ActorProperties](UObject* Outer)
-	{
-		if (!Outer || !Outer->GetClass())
-		{
-			return;
-		}
-
-		for (FProperty* OuterProperty : TFieldRange<FProperty>(Outer->GetClass()))
-		{
-			const EPropertyFlags Flags = OuterProperty->GetPropertyFlags();
-			const bool bEditable = !!(Flags & CPF_Edit);
-			const bool bEditConst = !!(Flags & CPF_EditConst);
-			const bool bNoEditInstance = !!(Flags & CPF_DisableEditOnInstance);
-
-			if (!bEditable || bEditConst || bNoEditInstance)
-			{
-				continue;;
-			}
-
-			if (FObjectPropertyBase* ObjectProperty = CastField<FObjectPropertyBase>(OuterProperty))
-			{
-				if (ObjectProperty->PropertyClass && ObjectProperty->PropertyClass->IsChildOf(UMaterialInterface::StaticClass()))
-				{
-					ActorProperties.Add(FDMObjectMaterialProperty(Outer, ObjectProperty));
-				}
-			}
-			else if (FArrayProperty* ArrayProperty = CastField<FArrayProperty>(OuterProperty))
-			{
-				if (FObjectPropertyBase* ArrayObjectProperty = CastField<FObjectPropertyBase>(ArrayProperty->Inner))
-				{
-					if (ArrayObjectProperty->PropertyClass && ArrayObjectProperty->PropertyClass->IsChildOf(UMaterialInterface::StaticClass()))
-					{
-						FScriptArrayHelper ArrayHelper(ArrayProperty, ArrayProperty->ContainerPtrToValuePtr<void>(Outer));
-
-						for (int Idx = 0, Count = ArrayHelper.Num(); Idx < Count; ++Idx)
-						{
-							ActorProperties.Add(FDMObjectMaterialProperty(Outer, OuterProperty, Idx));
-						}
-					}
-				}
-			}
-		}
-	};
-
-	IterateProperies(InActor);
-
-	InActor->ForEachComponent<UActorComponent>(false, [&IterateProperies](UActorComponent* InComp)
-		{
-			IterateProperies(InComp);
-		});
-
 	return ActorProperties;
 }
 
