Material Designer:
- General public api / unused func cleanup

#changelist validated
#virtualized

diff --git a/Engine/Plugins/Experimental/Avalanche/Source/AvalancheEditor/Private/AvaEditorModule.cpp b/Engine/Plugins/Experimental/Avalanche/Source/AvalancheEditor/Private/AvaEditorModule.cpp
--- a/Engine/Plugins/Experimental/Avalanche/Source/AvalancheEditor/Private/AvaEditorModule.cpp
+++ b/Engine/Plugins/Experimental/Avalanche/Source/AvalancheEditor/Private/AvaEditorModule.cpp
@@ -35,7 +35,6 @@
 #include "DetailView/Customizations/AvaSceneSettingsCustomization.h"
 #include "DetailView/Customizations/AvaVectorPropertyTypeCustomization.h"
 #include "DetailView/Customizations/AvaViewportQualitySettingsPropertyTypeCustomization.h"
-#include "DetailsPanel/DMMaterialInterfaceTypeCustomizer.h"
 
 DEFINE_LOG_CATEGORY(AvaLog);
 
@@ -263,11 +262,6 @@ void FAvaEditorModule::RegisterCustomLayouts()
 	PropertyModule.RegisterCustomClassLayout(UAvaSceneSettings::StaticClass()->GetFName(), FOnGetDetailCustomizationInstance::CreateStatic(
 		&FAvaSceneSettingsCustomization::MakeInstance));
 
-	const TSharedRef<FDMMaterialInterfaceTypeIdentifier> MaterialPropertyTypeIdentifier = MakeShared<FDMMaterialInterfaceTypeIdentifier>();
-	
-	PropertyModule.RegisterCustomPropertyTypeLayout(UMaterialInterface::StaticClass()->GetFName(), FOnGetPropertyTypeCustomizationInstance::CreateStatic(
-		&FDMMaterialInterfaceTypeCustomizer::MakeInstance), MaterialPropertyTypeIdentifier);
-
 	const TSharedRef<FAvaVectorPropertyTypeIdentifier> VectorPropertyTypeIdentifier = MakeShared<FAvaVectorPropertyTypeIdentifier>();
 
 	PropertyModule.RegisterCustomPropertyTypeLayout(TEXT("Vector"), FOnGetPropertyTypeCustomizationInstance::CreateStatic(
diff --git a/Engine/Plugins/Experimental/Avalanche/Source/AvalancheEditor/Private/MaterialDesigner/AvaMaterialDesignerExtension.cpp b/Engine/Plugins/Experimental/Avalanche/Source/AvalancheEditor/Private/MaterialDesigner/AvaMaterialDesignerExtension.cpp
--- a/Engine/Plugins/Experimental/Avalanche/Source/AvalancheEditor/Private/MaterialDesigner/AvaMaterialDesignerExtension.cpp
+++ b/Engine/Plugins/Experimental/Avalanche/Source/AvalancheEditor/Private/MaterialDesigner/AvaMaterialDesignerExtension.cpp
@@ -2,10 +2,12 @@
 
 #include "MaterialDesigner/AvaMaterialDesignerExtension.h"
 #include "AvaShapeActor.h"
-#include "DetailView/AvaDetailsExtension.h"
+#include "DMObjectMaterialProperty.h"
 #include "DMWorldSubsystem.h"
+#include "DetailView/AvaDetailsExtension.h"
 #include "DynamicMeshes/AvaShapeDynMeshBase.h"
 #include "Engine/World.h"
+#include "IDynamicMaterialEditorModule.h"
 #include "Material/DynamicMaterialInstance.h"
 #include "Materials/Material.h"
 #include "Model/DynamicMaterialModel.h"
@@ -137,7 +139,7 @@ void FAvaMaterialDesignerExtension::OpenEditor()
 		return;
 	}
 
-	FDynamicMaterialEditorModule::Get().OpenEditor(World);
+	IDynamicMaterialEditorModule::Get().OpenEditor(World);
 }
 
 #undef LOCTEXT_NAMESPACE
diff --git a/Engine/Plugins/Experimental/Avalanche/Source/AvalancheTextEditor/Private/AvaTextEditorModule.cpp b/Engine/Plugins/Experimental/Avalanche/Source/AvalancheTextEditor/Private/AvaTextEditorModule.cpp
--- a/Engine/Plugins/Experimental/Avalanche/Source/AvalancheTextEditor/Private/AvaTextEditorModule.cpp
+++ b/Engine/Plugins/Experimental/Avalanche/Source/AvalancheTextEditor/Private/AvaTextEditorModule.cpp
@@ -8,15 +8,16 @@
 #include "AvaTextEditorCommands.h"
 #include "ColorPicker/AvaViewportColorPickerActorClassRegistry.h"
 #include "ColorPicker/AvaViewportColorPickerAdapter.h"
+#include "DMObjectMaterialProperty.h"
 #include "DetailsView/AvaLinearGradientSettingsCustomization.h"
 #include "DetailsView/AvaTextAlignmentCustomization.h"
 #include "DetailsView/AvaTextComponentCustomization.h"
 #include "DetailsView/AvaTextFieldCustomization.h"
-#include "DynamicMaterialEditorModule.h"
 #include "Font/AvaFontDetailsCustomization.h"
 #include "Font/AvaFontManagerSubsystem.h"
 #include "Framework/Application/SlateApplication.h"
 #include "IAvalancheComponentVisualizersModule.h"
+#include "IDynamicMaterialEditorModule.h"
 #include "Misc/CoreDelegates.h"
 #include "Modules/ModuleManager.h"
 #include "PropertyEditorModule.h"
@@ -115,7 +116,7 @@ void FAvaTextEditorModule::UnregisterCustomLayouts()
 
 void FAvaTextEditorModule::RegisterDynamicMaterialPropertyGenerator()
 {
-	FDynamicMaterialEditorModule::RegisterCustomMaterialPropertyGenerator(
+	IDynamicMaterialEditorModule::Get().RegisterCustomMaterialPropertyGenerator(
 		AAvaTextActor::StaticClass(),
 		FDMGetObjectMaterialPropertiesDelegate::CreateLambda(
 			[](UObject* InObject)
diff --git a/Engine/Plugins/Experimental/DynamicMaterial/Source/DynamicMaterialEditor/Private/Components/DMMaterialProperty.cpp b/Engine/Plugins/Experimental/DynamicMaterial/Source/DynamicMaterialEditor/Private/Components/DMMaterialProperty.cpp
--- a/Engine/Plugins/Experimental/DynamicMaterial/Source/DynamicMaterialEditor/Private/Components/DMMaterialProperty.cpp
+++ b/Engine/Plugins/Experimental/DynamicMaterial/Source/DynamicMaterialEditor/Private/Components/DMMaterialProperty.cpp
@@ -2,12 +2,13 @@
 
 #include "Components/DMMaterialProperty.h"
 #include "Components/DMMaterialSlot.h"
+#include "DMPrivate.h"
 #include "DMValueDefinition.h"
-#include "Materials/MaterialExpressionMaterialFunctionCall.h"
 #include "Materials/MaterialExpressionConstant.h"
 #include "Materials/MaterialExpressionConstant2Vector.h"
 #include "Materials/MaterialExpressionConstant3Vector.h"
 #include "Materials/MaterialExpressionConstant4Vector.h"
+#include "Materials/MaterialExpressionMaterialFunctionCall.h"
 #include "Materials/MaterialFunctionInterface.h"
 #include "Math/Vector.h"
 #include "Math/Vector2D.h"
@@ -37,7 +38,7 @@ FString UDMMaterialProperty::GetComponentPathComponent() const
 UDMMaterialProperty* UDMMaterialProperty::CreateCustomMaterialPropertyDefaultSubobject(UDynamicMaterialModelEditorOnlyData* InModelEditorOnlyData, EDMMaterialPropertyType InMaterialProperty, const FName& InSubObjName)
 {
 	check(InModelEditorOnlyData);
-	check(UE::DynamicMaterialEditor::IsCustomMaterialProperty(InMaterialProperty));
+	check(UE::DynamicMaterialEditor::Private::IsCustomMaterialProperty(InMaterialProperty));
 
 	UDMMaterialProperty* NewMaterialProperty = InModelEditorOnlyData->CreateDefaultSubobject<UDMMaterialProperty>(InSubObjName);
 	NewMaterialProperty->MaterialProperty = InMaterialProperty;
diff --git a/Engine/Plugins/Experimental/DynamicMaterial/Source/DynamicMaterialEditor/Private/Components/DMMaterialSlot.cpp b/Engine/Plugins/Experimental/DynamicMaterial/Source/DynamicMaterialEditor/Private/Components/DMMaterialSlot.cpp
--- a/Engine/Plugins/Experimental/DynamicMaterial/Source/DynamicMaterialEditor/Private/Components/DMMaterialSlot.cpp
+++ b/Engine/Plugins/Experimental/DynamicMaterial/Source/DynamicMaterialEditor/Private/Components/DMMaterialSlot.cpp
@@ -10,13 +10,14 @@
 #include "CoreGlobals.h"
 #include "DMComponentPath.h"
 #include "DMDefs.h"
+#include "DMPrivate.h"
 #include "DynamicMaterialEditorModule.h"
 #include "Factories/MaterialFactoryNew.h"
+#include "MaterialValueType.h"
 #include "Materials/Material.h"
 #include "Materials/MaterialExpression.h"
 #include "Materials/MaterialExpressionMax.h"
 #include "Materials/MaterialExpressionMultiply.h"
-#include "MaterialValueType.h"
 #include "Model/DMMaterialBuildState.h"
 #include "Model/DMMaterialBuildUtils.h"
 #include "Model/DynamicMaterialModel.h"
@@ -854,7 +855,7 @@ void UDMMaterialSlot::UpdateBasePreviewMaterialFull(const TSharedRef<FDMMaterial
 
 	for (const TPair<EDMMaterialPropertyType, UDMMaterialProperty*>& Pair : MaterialProperties)
 	{
-		if (UE::DynamicMaterialEditor::IsCustomMaterialProperty(Pair.Key))
+		if (UE::DynamicMaterialEditor::Private::IsCustomMaterialProperty(Pair.Key))
 		{
 			continue;
 		}
diff --git a/Engine/Plugins/Experimental/DynamicMaterial/Source/DynamicMaterialEditor/Private/Components/DMMaterialStageThroughput.cpp b/Engine/Plugins/Experimental/DynamicMaterial/Source/DynamicMaterialEditor/Private/Components/DMMaterialStageThroughput.cpp
--- a/Engine/Plugins/Experimental/DynamicMaterial/Source/DynamicMaterialEditor/Private/Components/DMMaterialStageThroughput.cpp
+++ b/Engine/Plugins/Experimental/DynamicMaterial/Source/DynamicMaterialEditor/Private/Components/DMMaterialStageThroughput.cpp
@@ -691,7 +691,7 @@ int32 UDMMaterialStageThroughput::ResolveInputChannel(const TSharedRef<FDMMateri
 
 			for (int32 FloatIdx = 1; FloatIdx <= InputFloatCount; ++FloatIdx)
 			{
-				ComponentMask += UE::DynamicMaterialEditor::ChannelIndexToChannelBit(FloatIdx);
+				ComponentMask += UE::DynamicMaterialEditor::Private::ChannelIndexToChannelBit(FloatIdx);
 			}
 
 			OutChannel.OutputChannel = ComponentMask;
diff --git a/Engine/Plugins/Experimental/DynamicMaterial/Source/DynamicMaterialEditor/Private/DMBlueprintFunctionLibrary.cpp b/Engine/Plugins/Experimental/DynamicMaterial/Source/DynamicMaterialEditor/Private/DMBlueprintFunctionLibrary.cpp
--- a/Engine/Plugins/Experimental/DynamicMaterial/Source/DynamicMaterialEditor/Private/DMBlueprintFunctionLibrary.cpp
+++ b/Engine/Plugins/Experimental/DynamicMaterial/Source/DynamicMaterialEditor/Private/DMBlueprintFunctionLibrary.cpp
@@ -131,40 +131,6 @@ void UDMBlueprintFunctionLibrary::SetDefaultStageSourceTexture(UDMMaterialStage*
 	SetTextureValue(NewInputValue);
 }
 
-TArray<FDMObjectMaterialProperty> UDMBlueprintFunctionLibrary::GetActorDynamicMaterialSlots(AActor* InActor, const bool bInIncludeEmptySlots)
-{
-	TArray<FDMObjectMaterialProperty> OutSlots;
-
-	if (!IsValid(InActor))
-	{
-		return OutSlots;
-	}
-
-	InActor->ForEachComponent<UPrimitiveComponent>(false, [bInIncludeEmptySlots, &OutSlots](UPrimitiveComponent* InPrimComp)
-		{
-			const int32 PrimitiveMaterialCount = InPrimComp->GetNumMaterials();
-
-			for (int32 MaterialIndex = 0; MaterialIndex < PrimitiveMaterialCount; ++MaterialIndex)
-			{
-				if (UDynamicMaterialInstance* const MDI = Cast<UDynamicMaterialInstance>(InPrimComp->GetMaterial(MaterialIndex)))
-				{
-					if (UDynamicMaterialModel* const ThisMaterialModel = MDI->GetMaterialModel())
-					{
-						OutSlots.Add(FDMObjectMaterialProperty(InPrimComp, MaterialIndex));
-						continue;
-					}
-				}
-
-				if (bInIncludeEmptySlots)
-				{
-					OutSlots.Add(FDMObjectMaterialProperty(InPrimComp, MaterialIndex));
-				}
-			}
-		});
-
-	return OutSlots;
-}
-
 TArray<FDMObjectMaterialProperty> UDMBlueprintFunctionLibrary::GetActorMaterialProperties(AActor* InActor)
 {
 	TArray<FDMObjectMaterialProperty> ActorProperties;
@@ -197,58 +163,6 @@ TArray<FDMObjectMaterialProperty> UDMBlueprintFunctionLibrary::GetActorMaterialP
 	return ActorProperties;
 }
 
-UDynamicMaterialModel* UDMBlueprintFunctionLibrary::FindFirstValidActorMaterialSlot(const TArray<AActor*>& InActors)
-{
-	// Loop through selected actors to find first one with a valid material model on a primitive component.
-	for (AActor* Actor : InActors)
-	{
-		if (!IsValid(Actor))
-		{
-			continue;
-		}
-
-		const TArray<FDMObjectMaterialProperty> ActorMaterialProperties = GetActorMaterialProperties(Actor);
-		for (const FDMObjectMaterialProperty& ActorMaterialProperty : ActorMaterialProperties)
-		{
-			if (UDynamicMaterialModel* const ThisMaterialModel = ActorMaterialProperty.GetMaterialModel())
-			{
-				return ThisMaterialModel;
-			}
-		}
-	}
-
-	return nullptr;
-}
-
-FDMObjectMaterialProperty UDMBlueprintFunctionLibrary::MakePrimitiveComponentSlot(UPrimitiveComponent* InComponent, int32 InSlot)
-{
-	return FDMObjectMaterialProperty(InComponent, InSlot);
-}
-
-FDMObjectMaterialProperty UDMBlueprintFunctionLibrary::MakeObjectMaterialProperty(UObject* InObject, FName PropertyName)
-{
-	FProperty* Property = nullptr;
-
-	if (IsValid(InObject))
-	{
-		Property = InObject->GetClass()->FindPropertyByName(PropertyName);
-	}
-
-	return FDMObjectMaterialProperty(InObject, Property, INDEX_NONE);
-}
-
-FDMObjectMaterialProperty UDMBlueprintFunctionLibrary::MakeObjectMaterialPropertyArray(UObject* InObject, FName PropertyName, int32 ArrayIndex)
-{
-	FProperty* Property = nullptr;
-
-	if (IsValid(InObject))
-	{
-		Property = InObject->GetClass()->FindPropertyByName(PropertyName);
-	}
-
-	return FDMObjectMaterialProperty(InObject, Property, ArrayIndex);
-}
-
 UDynamicMaterialModel* UDMBlueprintFunctionLibrary::CreateDynamicMaterialInObject(FDMObjectMaterialProperty& InMaterialProperty)
 {
 	if (!InMaterialProperty.IsValid())
@@ -291,31 +205,6 @@ UDynamicMaterialModel* UDMBlueprintFunctionLibrary::CreateDynamicMaterialInObjec
 	return NewInstance->GetMaterialModel();
 }
 
-UDynamicMaterialModel* UDMBlueprintFunctionLibrary::GetObjectPropertyMaterialModel(const FDMObjectMaterialProperty& InMaterialProperty) const
-{
-	return InMaterialProperty.GetMaterialModel();
-}
-
-UDynamicMaterialInstance* UDMBlueprintFunctionLibrary::GetObjectPropertyMaterial(const FDMObjectMaterialProperty& InMaterialProperty) const
-{
-	return InMaterialProperty.GetMaterial();
-}
-
-void UDMBlueprintFunctionLibrary::SetObjectPropertyMaterial(FDMObjectMaterialProperty& InMaterialProperty, UDynamicMaterialInstance* InDynamicMaterial)
-{
-	InMaterialProperty.SetMaterial(InDynamicMaterial);
-}
-
-bool UDMBlueprintFunctionLibrary::IsObjectPropertyValid(const FDMObjectMaterialProperty& InMaterialProperty) const
-{
-	return InMaterialProperty.IsValid();
-}
-
-FText UDMBlueprintFunctionLibrary::GetObjectPropertyName(const FDMObjectMaterialProperty& InMaterialProperty, bool bInIgnoreNewStatus) const
-{
-	return InMaterialProperty.GetPropertyName(bInIgnoreNewStatus);
-}
-
 bool UDMBlueprintFunctionLibrary::ExportMaterialInstance(UDynamicMaterialModel* InMaterialModel, const FString& InSavePath)
 {
 	if (!IsValid(InMaterialModel))
diff --git a/Engine/Plugins/Experimental/DynamicMaterial/Source/DynamicMaterialEditor/Private/DMBlueprintFunctionLibrary.h b/Engine/Plugins/Experimental/DynamicMaterial/Source/DynamicMaterialEditor/Private/DMBlueprintFunctionLibrary.h
new file mode 100644
--- /dev/null
+++ b/Engine/Plugins/Experimental/DynamicMaterial/Source/DynamicMaterialEditor/Private/DMBlueprintFunctionLibrary.h
@@ -0,0 +1,37 @@
+// Copyright Epic Games, Inc. All Rights Reserved.
+
+#pragma once
+
+#include "DMDefs.h"
+#include "DMEDefs.h"
+#include "Kismet/BlueprintFunctionLibrary.h"
+#include "DMBlueprintFunctionLibrary.generated.h"
+
+class AActor;
+class UDMMaterialStage;
+class UDMMaterialStageInputValue;
+class UDynamicMaterialModel;
+class UTexture;
+struct FDMObjectMaterialProperty;
+
+/**
+ * Material Designer Blueprint Function Library
+ */
+UCLASS()
+class UDMBlueprintFunctionLibrary : public UBlueprintFunctionLibrary
+{
+	GENERATED_BODY()
+
+public:
+	static UDMMaterialStageInputValue* FindDefaultStageOpacityInputValue(UDMMaterialStage* InStage);
+
+	static void SetDefaultStageSourceTexture(UDMMaterialStage* InStage, UTexture* InTexture);
+
+	static TArray<FDMObjectMaterialProperty> GetActorMaterialProperties(AActor* InActor);
+
+	static UDynamicMaterialModel* CreateDynamicMaterialInObject(FDMObjectMaterialProperty& InMaterialProperty);
+
+	static bool ExportMaterialInstance(UDynamicMaterialModel* InMaterialModel, const FString& InSavePath);
+
+	static bool ExportGeneratedMaterial(UDynamicMaterialModel* InMaterialModel, const FString& InSavePath);
+};
diff --git a/Engine/Plugins/Experimental/DynamicMaterial/Source/DynamicMaterialEditor/Private/DMEDefs.cpp b/Engine/Plugins/Experimental/DynamicMaterial/Source/DynamicMaterialEditor/Private/DMEDefs.cpp
deleted file mode 100644
diff --git a/Engine/Plugins/Experimental/DynamicMaterial/Source/DynamicMaterialEditor/Private/DMMaterialFunctionLibrary.h b/Engine/Plugins/Experimental/DynamicMaterial/Source/DynamicMaterialEditor/Private/DMMaterialFunctionLibrary.h
new file mode 100644
--- /dev/null
+++ b/Engine/Plugins/Experimental/DynamicMaterial/Source/DynamicMaterialEditor/Private/DMMaterialFunctionLibrary.h
@@ -0,0 +1,35 @@
+// Copyright Epic Games, Inc. All Rights Reserved.
+
+#pragma once
+
+#include "Materials/MaterialFunctionInterface.h"
+#include "UObject/StrongObjectPtr.h"
+
+class UMaterial;
+class UMaterialExpressionMaterialFunctionCall;
+
+class FDMMaterialFunctionLibrary
+{
+public:
+	static FDMMaterialFunctionLibrary& Get();
+
+	UMaterialFunctionInterface* GetFunction(const FName& Name, const FString& Path);
+
+	UMaterialExpressionMaterialFunctionCall* MakeExpression(UMaterial* Parent, UMaterialFunctionInterface* Function, const FString& InComment);
+	UMaterialExpressionMaterialFunctionCall* MakeExpression(UMaterial* Parent, const FName& Name, const FString& Path, const FString& InComment);
+
+	UMaterialExpressionMaterialFunctionCall* GetBreakOutFloat2Components(UMaterial* Parent, const FString& InComment);
+	UMaterialExpressionMaterialFunctionCall* GetBreakOutFloat3Components(UMaterial* Parent, const FString& InComment);
+	UMaterialExpressionMaterialFunctionCall* GetBreakOutFloat4Components(UMaterial* Parent, const FString& InComment);
+	UMaterialExpressionMaterialFunctionCall* GetRGVtoHSV(UMaterial* Parent, const FString& InComment);
+	UMaterialExpressionMaterialFunctionCall* GetMakeFloat2(UMaterial* Parent, const FString& InComment);
+	UMaterialExpressionMaterialFunctionCall* GetMakeFloat3(UMaterial* Parent, const FString& InComment);
+	UMaterialExpressionMaterialFunctionCall* GetMakeFloat4(UMaterial* Parent, const FString& InComment);
+	UMaterialExpressionMaterialFunctionCall* GetCustomRotator(UMaterial* Parent, const FString& InComment); /** Possibly promote this to a full expression? */
+	UMaterialExpressionMaterialFunctionCall* GetRadialGradientExponential(UMaterial* Parent, const FString& InComment); /** Possibly promote this to a full expression? */
+
+protected:
+	TMap<FName, TStrongObjectPtr<UMaterialFunctionInterface>> LoadedFunctions;
+
+	UMaterialFunctionInterface* LoadFunction(const FString& Path);
+};
diff --git a/Engine/Plugins/Experimental/DynamicMaterial/Source/DynamicMaterialEditor/Private/DMPrivate.cpp b/Engine/Plugins/Experimental/DynamicMaterial/Source/DynamicMaterialEditor/Private/DMPrivate.cpp
--- a/Engine/Plugins/Experimental/DynamicMaterial/Source/DynamicMaterialEditor/Private/DMPrivate.cpp
+++ b/Engine/Plugins/Experimental/DynamicMaterial/Source/DynamicMaterialEditor/Private/DMPrivate.cpp
@@ -35,6 +35,36 @@ namespace UE::DynamicMaterialEditor::Private
 			}
 		}
 	}
+
+	int32 ChannelIndexToChannelBit(int32 InChannelIndex)
+	{
+		switch (InChannelIndex)
+		{
+		case 0:
+			return FDMMaterialStageConnectorChannel::WHOLE_CHANNEL;
+
+		case 1:
+			return FDMMaterialStageConnectorChannel::FIRST_CHANNEL;
+
+		case 2:
+			return FDMMaterialStageConnectorChannel::SECOND_CHANNEL;
+
+		case 3:
+			return FDMMaterialStageConnectorChannel::THIRD_CHANNEL;
+
+		case 4:
+			return FDMMaterialStageConnectorChannel::FOURTH_CHANNEL;
+
+		default:
+			checkNoEntry();
+			return 0;
+		}
+	}
+
+	bool IsCustomMaterialProperty(EDMMaterialPropertyType InMaterialProperty)
+	{
+		return (InMaterialProperty >= EDMMaterialPropertyType::Custom1 && InMaterialProperty <= EDMMaterialPropertyType::Custom4);
+	}
 }
 
 FDMMaterialLayerReference::FDMMaterialLayerReference()
diff --git a/Engine/Plugins/Experimental/DynamicMaterial/Source/DynamicMaterialEditor/Private/DMPrivate.h b/Engine/Plugins/Experimental/DynamicMaterial/Source/DynamicMaterialEditor/Private/DMPrivate.h
--- a/Engine/Plugins/Experimental/DynamicMaterial/Source/DynamicMaterialEditor/Private/DMPrivate.h
+++ b/Engine/Plugins/Experimental/DynamicMaterial/Source/DynamicMaterialEditor/Private/DMPrivate.h
@@ -23,6 +23,11 @@ namespace UE::DynamicMaterialEditor::Private
 	};
 
 	void SetMask(FExpressionInput& InInputConnector, const FExpressionOutput& InOutputConnector, int32 InChannelOverride);
+
+	/** Converts 0,1,2,3,4 to 0,1,2,4,8 */
+	int32 ChannelIndexToChannelBit(int32 InChannelIndex);
+
+	bool IsCustomMaterialProperty(EDMMaterialPropertyType InMaterialProperty);
 }
 
 struct FDMMaterialLayerReference
diff --git a/Engine/Plugins/Experimental/DynamicMaterial/Source/DynamicMaterialEditor/Private/DMWorldSubsystem.cpp b/Engine/Plugins/Experimental/DynamicMaterial/Source/DynamicMaterialEditor/Private/DMWorldSubsystem.cpp
--- a/Engine/Plugins/Experimental/DynamicMaterial/Source/DynamicMaterialEditor/Private/DMWorldSubsystem.cpp
+++ b/Engine/Plugins/Experimental/DynamicMaterial/Source/DynamicMaterialEditor/Private/DMWorldSubsystem.cpp
@@ -10,8 +10,3 @@ UDMWorldSubsystem::UDMWorldSubsystem()
 	// Default fallback implementation
 	InvokeTabDelegate.BindWeakLambda(this, [this]() { FDMLevelEditorIntegration::InvokeTabForWorld(GetWorld()); });
 }
- 
-void UDMWorldSubsystem::SetKeyframeHandler(const TSharedPtr<IDetailKeyframeHandler>& InKeyframeHandler)
-{
-	KeyframeHandler = InKeyframeHandler;
-}
diff --git a/Engine/Plugins/Experimental/DynamicMaterial/Source/DynamicMaterialEditor/Private/DetailsPanel/DMMaterialInterfaceTypeCustomizer.h b/Engine/Plugins/Experimental/DynamicMaterial/Source/DynamicMaterialEditor/Private/DetailsPanel/DMMaterialInterfaceTypeCustomizer.h
new file mode 100644
--- /dev/null
+++ b/Engine/Plugins/Experimental/DynamicMaterial/Source/DynamicMaterialEditor/Private/DetailsPanel/DMMaterialInterfaceTypeCustomizer.h
@@ -0,0 +1,33 @@
+// Copyright Epic Games, Inc. All Rights Reserved.
+
+#pragma once
+
+#include "IPropertyTypeCustomization.h"
+#include "PropertyEditorModule.h"
+#include "PropertyHandle.h"
+#include "Templates/SharedPointer.h"
+
+class FDMMaterialInterfaceTypeIdentifier : public IPropertyTypeIdentifier
+{
+public:
+	virtual bool IsPropertyTypeCustomized(const IPropertyHandle& InPropertyHandle) const override
+	{
+		return InPropertyHandle.GetNumOuterObjects() > 0;
+	}
+};
+
+class FDMMaterialInterfaceTypeCustomizer : public IPropertyTypeCustomization
+{
+public:
+	static TSharedRef<IPropertyTypeCustomization> MakeInstance();
+
+	FDMMaterialInterfaceTypeCustomizer();
+
+	/** BEGIN IPropertyTypeCustomization interface */
+	virtual void CustomizeHeader(TSharedRef<IPropertyHandle> EditorPropertyHandle, class FDetailWidgetRow& HeaderRow,
+		IPropertyTypeCustomizationUtils& CustomizationUtils) override;
+
+	virtual void CustomizeChildren(TSharedRef<IPropertyHandle> EditorPropertyHandle, class IDetailChildrenBuilder& ChildBuilder,
+		IPropertyTypeCustomizationUtils& CustomizationUtils) override;
+	/** END IPropertyTypeCustomization interface */
+};
diff --git a/Engine/Plugins/Experimental/DynamicMaterial/Source/DynamicMaterialEditor/Private/DetailsPanel/Slate/SDMDetailsPanelMaterialInterfaceWidget.h b/Engine/Plugins/Experimental/DynamicMaterial/Source/DynamicMaterialEditor/Private/DetailsPanel/Slate/SDMDetailsPanelMaterialInterfaceWidget.h
new file mode 100644
--- /dev/null
+++ b/Engine/Plugins/Experimental/DynamicMaterial/Source/DynamicMaterialEditor/Private/DetailsPanel/Slate/SDMDetailsPanelMaterialInterfaceWidget.h
@@ -0,0 +1,36 @@
+// Copyright Epic Games, Inc. All Rights Reserved.
+
+#pragma once
+
+#include "Widgets/SCompoundWidget.h"
+#include "Templates/SharedPointer.h"
+
+class FAssetThumbnailPool;
+class IPropertyHandle;
+class UDynamicMaterialInstance;
+
+class SDMDetailsPanelMaterialInterfaceWidget : public SCompoundWidget
+{
+public:
+	SLATE_BEGIN_ARGS(SDMDetailsPanelMaterialInterfaceWidget) {}
+	SLATE_ARGUMENT(TSharedPtr<FAssetThumbnailPool>, ThumbnailPool)
+	SLATE_END_ARGS()
+
+	void Construct(const FArguments& InArgs, const TSharedRef<IPropertyHandle>& InPropertyHandle);
+
+protected:
+	TSharedPtr<IPropertyHandle> PropertyHandle;
+
+	UObject* GetAsset() const;
+	void SetAsset(UObject* NewAsset);
+
+	UDynamicMaterialInstance* GetDynamicMaterialInstance() const;
+	void SetDynamicMaterialInstance(UDynamicMaterialInstance* NewInstance);
+
+	FReply OnButtonClicked();
+	FReply CreateDynamicMaterialInstance();
+	FReply ClearDynamicMaterialInstance();
+	FReply OpenDynamicMaterialInstanceTab();
+
+	FText GetButtonText() const;
+};
diff --git a/Engine/Plugins/Experimental/DynamicMaterial/Source/DynamicMaterialEditor/Private/DynamicMaterialEditorModule.cpp b/Engine/Plugins/Experimental/DynamicMaterial/Source/DynamicMaterialEditor/Private/DynamicMaterialEditorModule.cpp
--- a/Engine/Plugins/Experimental/DynamicMaterial/Source/DynamicMaterialEditor/Private/DynamicMaterialEditorModule.cpp
+++ b/Engine/Plugins/Experimental/DynamicMaterial/Source/DynamicMaterialEditor/Private/DynamicMaterialEditorModule.cpp
@@ -17,11 +17,12 @@
 #include "Components/MaterialValues/DMMaterialValueFloat4.h"
 #include "Components/MaterialValues/DMMaterialValueTexture.h"
 #include "Components/PrimitiveComponent.h"
+#include "DMMaterialFunctionLibrary.h"
+#include "DMWorldSubsystem.h"
 #include "DetailLayoutBuilder.h"
+#include "DetailsPanel/DMMaterialInterfaceTypeCustomizer.h"
 #include "DetailsPanel/DMPropertyTypeCustomizer.h"
 #include "DetailsPanel/Slate/SDMMaterialListExtensionWidget.h"
-#include "DMMaterialFunctionLibrary.h"
-#include "DMWorldSubsystem.h"
 #include "DynamicMaterialEditorCommands.h"
 #include "DynamicMaterialEditorSettings.h"
 #include "DynamicMaterialEditorStyle.h"
@@ -197,7 +198,7 @@ FDMGetObjectMaterialPropertiesDelegate FDynamicMaterialEditorModule::GetCustomMa
 
 FDynamicMaterialEditorModule& FDynamicMaterialEditorModule::Get()
 {
-	return FModuleManager::LoadModuleChecked<FDynamicMaterialEditorModule>(TEXT("DynamicMaterialEditor"));
+	return FModuleManager::LoadModuleChecked<FDynamicMaterialEditorModule>(ModuleName);
 }
 
 FDynamicMaterialEditorModule::FDynamicMaterialEditorModule()
@@ -218,6 +219,9 @@ void FDynamicMaterialEditorModule::StartupModule()
 	PropertyModule.RegisterCustomPropertyTypeLayout(UDynamicMaterialModelEditorOnlyData::StaticClass()->GetFName()
 		, FOnGetPropertyTypeCustomizationInstance::CreateStatic(&FDMPropertyTypeCustomizer::MakeInstance));
 
+	PropertyModule.RegisterCustomPropertyTypeLayout(UMaterialInterface::StaticClass()->GetFName()
+		, FOnGetPropertyTypeCustomizationInstance::CreateStatic(&FDMMaterialInterfaceTypeCustomizer::MakeInstance), MakeShared<FDMMaterialInterfaceTypeIdentifier>());
+
 	using namespace UE::DynamicMaterialEditor::Private;
 
 	MateriaListWidgetsDelegate = FMaterialList::OnAddMaterialItemViewExtraBottomWidget.AddStatic(&AddMaterialListWidgets);
diff --git a/Engine/Plugins/Experimental/DynamicMaterial/Source/DynamicMaterialEditor/Private/DynamicMaterialEditorModule.h b/Engine/Plugins/Experimental/DynamicMaterial/Source/DynamicMaterialEditor/Private/DynamicMaterialEditorModule.h
new file mode 100644
--- /dev/null
+++ b/Engine/Plugins/Experimental/DynamicMaterial/Source/DynamicMaterialEditor/Private/DynamicMaterialEditorModule.h
@@ -0,0 +1,155 @@
+// Copyright Epic Games, Inc. All Rights Reserved.
+
+#pragma once
+
+#include "DMEDefs.h"
+#include "DMObjectMaterialProperty.h"
+#include "IDynamicMaterialEditorModule.h"
+#include "Templates/SharedPointer.h"
+#include "TickableEditorObject.h"
+
+class AActor;
+class FDMMaterialFunctionLibrary;
+class FUICommandList;
+class IAssetTypeActions;
+class ILevelEditor;
+class SDMEditor;
+class SDMComponentEdit;
+class SDockTab;
+class SWidget;
+class UDMMaterialComponent;
+class UDMMaterialStageSource;
+class UDynamicMaterialInstance;
+class UDynamicMaterialModel;
+struct FDMBuildRequestList;
+
+DECLARE_LOG_CATEGORY_EXTERN(LogDynamicMaterialEditor, Log, All);
+
+namespace UE::DynamicMaterialEditor
+{
+	constexpr bool bMultipleSlotPropertiesEnabled = false;
+	constexpr bool bGlobalValuesEnabled = false;
+	constexpr bool bAdvancedSlotsEnabled = false;
+}
+
+DECLARE_MULTICAST_DELEGATE(FDMOnUIValueUpdate);
+
+/** Takes a UMaterialValue and returns the widget used to edit it. */
+DECLARE_DELEGATE_RetVal_TwoParams(TSharedPtr<SWidget>, FDMCreateValueEditWidgetDelegate, const TSharedPtr<SDMComponentEdit>&, UDMMaterialValue*);
+
+/** Creates property rows in the edit widget. */
+DECLARE_DELEGATE_FourParams(FDMComponentPropertyRowGeneratorDelegate, const TSharedRef<SDMComponentEdit>&, UDMMaterialComponent*,
+	TArray<FDMPropertyHandle>&, TSet<UDMMaterialComponent*>&)
+
+struct FDMBuildRequestEntry
+{
+	FString AssetPath;
+	bool bDirtyAssets;
+
+	bool operator==(const FDMBuildRequestEntry& Other) const
+	{
+		return AssetPath == Other.AssetPath;
+	}
+
+	friend uint32 GetTypeHash(const FDMBuildRequestEntry& InEntry)
+	{
+		return GetTypeHash(InEntry.AssetPath);
+	}
+};
+
+/**
+ * Material Designer - Build your own materials in a slimline editor!
+ */
+class FDynamicMaterialEditorModule : public IDynamicMaterialEditorModule, public FTickableEditorObject
+{
+public:
+	static const FName TabId;
+	static FDMOnUIValueUpdate& GetOnUIValueUpdate() { return OnUIValueUpdate; }
+
+	static FDynamicMaterialEditorModule& Get();
+
+	static void RegisterValueEditWidgetDelegate(UClass* InClass, FDMCreateValueEditWidgetDelegate InValueEditBodyDelegate);
+	template<class InValClass, class InEditClass> static void RegisterValueEditWidgetDelegate();
+	static FDMCreateValueEditWidgetDelegate GetValueEditWidgetDelegate(UClass* InClass);
+	static TSharedPtr<SWidget> CreateEditWidgetForValue(const TSharedPtr<SDMComponentEdit>& InComponentEditWidget, UDMMaterialValue* InValue);
+
+	static void RegisterComponentPropertyRowGeneratorDelegate(UClass* InClass, FDMComponentPropertyRowGeneratorDelegate InComponentPropertyRowGeneratorDelegate);
+	template<class InObjClass, class InGenClass> static void RegisterComponentPropertyRowGeneratorDelegate();
+	static FDMComponentPropertyRowGeneratorDelegate GetComponentPropertyRowGeneratorDelegate(UClass* InClass);
+	static void GeneratorComponentPropertyRows(const TSharedRef<SDMComponentEdit>& InComponentEditWidget, UDMMaterialComponent* InComponent, 
+		TArray<FDMPropertyHandle>& InOutPropertyRows, TSet<UDMMaterialComponent*>& InOutProcessedObjects);
+
+	static FDMGetObjectMaterialPropertiesDelegate GetCustomMaterialPropertyGenerator(UClass* InClass);
+
+	/** With a provided world, the editor will bind to the MD world subsystem to receive model changes. */
+	static TSharedRef<SWidget> CreateEditor(UDynamicMaterialModel* InMaterialModel, UWorld* InAssetEditorWorld);
+	static TSharedRef<SWidget> CreateEmptyTabContent();
+
+	FDynamicMaterialEditorModule();
+
+	//~ Begin IDynamicMaterialEditorModule
+	virtual void OpenEditor(UWorld* InWorld) override;
+	virtual void RegisterCustomMaterialPropertyGenerator(UClass* InClass, FDMGetObjectMaterialPropertiesDelegate InGenerator) override;
+	//~ End IDynamicMaterialEditorModule
+
+	//~ Begin IModuleInterface
+	virtual void StartupModule() override;
+	virtual void ShutdownModule() override;
+	//~ End IModuleInterface
+
+	const TSharedPtr<FDMMaterialFunctionLibrary>& GetFunctionLibrary();
+
+	void SetDynamicMaterialModel(UDynamicMaterialModel* InMaterialModel, UWorld* InWorld, bool bInInvokeTab);
+	void SetDynamicMaterialObjectProperty(const FDMObjectMaterialProperty& InObjectProperty, UWorld* InWorld, bool bInInvokeTab);
+	void SetDynamicMaterialInstance(UDynamicMaterialInstance* InInstance, UWorld* InWorld, bool bInInvokeTab);
+	void SetDynamicMaterialActor(AActor* InActor, UWorld* InWorld, bool bInInvokeTab);
+	void ClearDynamicMaterialModel(UWorld* InWorld);
+
+	//~ Begin FTickableEditorObject
+	virtual void Tick(float DeltaTime) override;
+	virtual ETickableTickType GetTickableTickType() const override { return ETickableTickType::Always; }
+	virtual TStatId GetStatId() const override;
+	//~ End FTickableEditorObject
+
+	void AddBuildRequest(UObject* InToBuild, bool bInDirtyAssets);
+
+	const TSharedRef<FUICommandList>& GetCommandList() const { return CommandList; }
+
+protected:
+	static TMap<UClass*, FDMCreateValueEditWidgetDelegate> ValueEditWidgetDelegates;
+	static TMap<UClass*, FDMComponentPropertyRowGeneratorDelegate> ComponentPropertyRowGenerators;
+	static TMap<UClass*, FDMGetObjectMaterialPropertiesDelegate> CustomMaterialPropertyGenerators;
+	static FDMOnUIValueUpdate OnUIValueUpdate;
+
+	TSharedPtr<FDMMaterialFunctionLibrary> FunctionLibrary;
+	TSet<FDMBuildRequestEntry> BuildRequestList;
+	TSharedRef<FUICommandList> CommandList;
+
+	void ProcessBuildRequest(UObject* InToBuild, bool bInDirtyAssets);
+
+	void MapCommands();
+	void UnmapCommands();
+};
+
+template <class InValClass, class InEditClass>
+void FDynamicMaterialEditorModule::RegisterValueEditWidgetDelegate()
+{
+	RegisterValueEditWidgetDelegate(
+		InValClass::StaticClass(),
+		FDMCreateValueEditWidgetDelegate::CreateStatic(
+			&InEditClass::CreateEditWidget
+		)
+	);
+}
+
+template <class InObjClass, class InGenClass>
+void FDynamicMaterialEditorModule::RegisterComponentPropertyRowGeneratorDelegate()
+{
+	RegisterComponentPropertyRowGeneratorDelegate(
+		InObjClass::StaticClass(),
+		FDMComponentPropertyRowGeneratorDelegate::CreateSP(
+			InGenClass::Get(),
+			&InGenClass::AddComponentProperties
+		)
+	);
+}
diff --git a/Engine/Plugins/Experimental/DynamicMaterial/Source/DynamicMaterialEditor/Private/DynamicMaterialEditorSettings.h b/Engine/Plugins/Experimental/DynamicMaterial/Source/DynamicMaterialEditor/Private/DynamicMaterialEditorSettings.h
new file mode 100644
--- /dev/null
+++ b/Engine/Plugins/Experimental/DynamicMaterial/Source/DynamicMaterialEditor/Private/DynamicMaterialEditorSettings.h
@@ -0,0 +1,152 @@
+// Copyright Epic Games, Inc. All Rights Reserved.
+
+#pragma once
+
+#include "Engine/DeveloperSettings.h"
+#include "Containers/Array.h"
+#include "Engine/Texture.h"
+#include "PropertyEditorDelegates.h"
+#include "UObject/SoftObjectPtr.h"
+#include "DynamicMaterialEditorSettings.generated.h"
+
+class FModifierKeysState;
+class UMaterialFunctionInterface;
+struct FPropertyChangedEvent;
+enum EOrientation : int;
+
+USTRUCT(BlueprintType)
+struct FDMMaterialEffectList
+{
+	GENERATED_BODY()
+
+	UPROPERTY(Config, EditAnywhere, BlueprintReadWrite, Category = "Effects")
+	FString Name;
+
+	UPROPERTY(Config, EditAnywhere, BlueprintReadWrite, Category = "Effects")
+	TArray<TSoftObjectPtr<UMaterialFunctionInterface>> Effects;
+};
+
+/**
+ * Material Designer Settings
+ */
+UCLASS(Config=EditorPerProjectUserSettings, meta = (DisplayName = "Material Designer"))
+class UDynamicMaterialEditorSettings : public UDeveloperSettings
+{
+	GENERATED_BODY()
+
+public:
+	UDynamicMaterialEditorSettings();
+	virtual ~UDynamicMaterialEditorSettings() override {}
+
+	static UDynamicMaterialEditorSettings* Get();
+
+	/** Changes the currently active material in the designer following actor/object selection. */
+	UPROPERTY(Config, EditAnywhere, BlueprintReadWrite, Category = "Preview")
+	bool bFollowSelection;
+
+	/** Adjusts the the spin box value sensitivity. */
+	UPROPERTY(Config, EditAnywhere, BlueprintReadWrite, AdvancedDisplay, Category = "Layout", meta = (
+		ClampMin = "0.01", UIMin = "0.01", ClampMax = "1.0", UIMax = "1.0"))
+	float SpinBoxValueMultiplier_Default;
+
+	/** Adjusts the the spin box value sensitivity. */
+	UPROPERTY(Config, EditAnywhere, BlueprintReadWrite, AdvancedDisplay, Category = "Layout", meta = (
+		ClampMin = "0.01", UIMin = "0.01", ClampMax = "1.0", UIMax = "1.0"))
+	float SpinBoxValueMultiplier_Shift;
+
+	/** Adjusts the the spin box value sensitivity. */
+	UPROPERTY(Config, EditAnywhere, BlueprintReadWrite, AdvancedDisplay, Category = "Layout", meta = (
+		ClampMin = "0.01", UIMin = "0.01", ClampMax = "1.0", UIMax = "1.0"))
+	float SpinBoxValueMultiplier_AltShift;
+
+	/** Adjusts the the spin box value sensitivity. */
+	UPROPERTY(Config, EditAnywhere, BlueprintReadWrite, AdvancedDisplay, Category = "Layout", meta = (
+		ClampMin = "0.01", UIMin = "0.01", ClampMax = "1.0", UIMax = "1.0"))
+	float SpinBoxValueMultiplier_Cmd;
+
+	/** Adjusts the the spin box value sensitivity. */
+	UPROPERTY(Config, EditAnywhere, BlueprintReadWrite, AdvancedDisplay, Category = "Layout", meta = (
+		ClampMin = "0.01", UIMin = "0.01", ClampMax = "1.0", UIMax = "1.0"))
+	float SpinBoxValueMultiplier_AltCmd;
+
+	/** Adjusts the vertical size of the material layer view. */
+	UPROPERTY(Config, EditAnywhere, BlueprintReadWrite, Category = "Layout", meta = (
+		ClampMin = "0.05", UIMin = "0.05", ClampMax = "0.95", UIMax = "0.95"))
+	float SplitterLocation;
+
+	/** Sets the maximum width for float-based sliders in the editor */
+	UPROPERTY(Config, EditAnywhere, BlueprintReadWRite, Category = "Layout", meta = (
+		ClampMin = "100", UIMin = "100", ClampMax = "1000", UIMax = "1000"))
+	float MaxFloatSliderWidth;
+
+	/** The size of the material layer preview images. */
+	UPROPERTY(Config, EditAnywhere, BlueprintReadWrite, Category = "Preview", meta = (ClampMin = "32", UIMin = "32", ClampMax = "128", UIMax = "128"))
+	int32 LayerPreviewSize;
+
+	/** The size of the material slot preview images. */
+	UPROPERTY(Config, EditAnywhere, BlueprintReadWrite, Category = "Preview", meta = (ClampMin = "40", UIMin = "40", ClampMax = "128", UIMax = "128"))
+	int32 SlotPreviewSize;
+
+	/** The size of the material layer details preview images. */
+	UPROPERTY(Config, EditAnywhere, BlueprintReadWrite, Category = "Preview", meta = (ClampMin = "32", UIMin = "32", ClampMax = "128", UIMax = "128"))
+	int32 DetailsPreviewSize;
+
+	UPROPERTY(Config, EditAnywhere, BlueprintReadWrite, Category = "Preview")
+	bool bPreviewImagesUseTextureUVs;
+
+	/**
+	 * If true, will display the hovered preview image in an enlarged format inside the tooltip.
+	 * 
+	 * NOTE: The Material Designer must be re-opened for changes to take effect.
+	 */
+	UPROPERTY(Config, EditAnywhere, BlueprintReadWrite, Category = "ToolTips")
+	bool bShowTooltipPreview;
+
+	/** The size of the preview image when displayed in the tooltip. */
+	UPROPERTY(Config, EditAnywhere, BlueprintReadWrite, Category = "ToolTips", meta = (EditCondition = "bShowTooltipPreview == true", ClampMin = "128", UIMin = "128", ClampMax = "2048", UIMax = "2048"))
+	int32 TooltipTextureSize;
+
+	UPROPERTY(Config, EditAnywhere, BlueprintReadWrite, Category = "Preview", meta=(DisplayName = "Default RGB Texture"))
+	TSoftObjectPtr<UTexture> DefaultRGBTexture;
+
+	UPROPERTY(Config, EditAnywhere, BlueprintReadWrite, Category = "Preview")
+	TSoftObjectPtr<UTexture> DefaultOpaqueTexture;
+
+	UPROPERTY(Config, EditAnywhere, BlueprintReadWrite, Category = "Preview")
+	TSoftObjectPtr<UTexture> DefaultPreviewCanvasTexture;
+
+	UPROPERTY(Config, EditAnywhere, BlueprintReadWrite, Category = "Preview")
+	TSoftObjectPtr<UTexture> DefaultOpacitySlotMask;
+
+	/*
+	 * Add paths to search for custom effects.
+	 *
+	 * Format examples:
+	 * - /Game/Some/Path
+	 * - /Plugin/Some/Path
+	 *
+	 * The assets must be in a sub-folder of the base path. The sub-folder
+	 * will be used as the category name.
+	 *
+	 * Asset Examples:
+	 * - /Game/Some/Path/UV/Asset.Asset -> Category: UV
+	 * - /Plugin/Some/Path/Color/OtherAsset.OtherAsset -> Category: Color
+	 */
+	UPROPERTY(Config, EditAnywhere, BlueprintReadWrite, Category = "Effects")
+	TArray<FName> CustomEffectsFolders;
+
+	FOnFinishedChangingProperties OnSettingsChanged;
+
+	//~ Begin UObject
+	virtual void PostInitProperties() override;
+	virtual void PostEditChangeProperty(FPropertyChangedEvent& InPropertyChangedEvent) override;
+	//~ End UObject
+
+	void OpenEditorSettingsWindow() const;
+
+	void ResetAllLayoutSettings();
+
+	float GetSpinboxValueChangeMultiplier(const FModifierKeysState& InModifierKeys) const;
+
+	TArray<FDMMaterialEffectList> GetEffectList() const;
+};
diff --git a/Engine/Plugins/Experimental/DynamicMaterial/Source/DynamicMaterialEditor/Private/LevelEditor/DMLevelEditorIntegrationInstance.cpp b/Engine/Plugins/Experimental/DynamicMaterial/Source/DynamicMaterialEditor/Private/LevelEditor/DMLevelEditorIntegrationInstance.cpp
--- a/Engine/Plugins/Experimental/DynamicMaterial/Source/DynamicMaterialEditor/Private/LevelEditor/DMLevelEditorIntegrationInstance.cpp
+++ b/Engine/Plugins/Experimental/DynamicMaterial/Source/DynamicMaterialEditor/Private/LevelEditor/DMLevelEditorIntegrationInstance.cpp
@@ -3,6 +3,7 @@
 #include "DMLevelEditorIntegrationInstance.h"
 #include "Components/PrimitiveComponent.h"
 #include "DMWorldSubsystem.h"
+#include "DynamicMaterialEditorModule.h"
 #include "EditorModeManager.h"
 #include "Elements/Framework/TypedElementSelectionSet.h"
 #include "Engine/World.h"
diff --git a/Engine/Plugins/Experimental/DynamicMaterial/Source/DynamicMaterialEditor/Private/Material/DynamicMaterialInstanceFactory.cpp b/Engine/Plugins/Experimental/DynamicMaterial/Source/DynamicMaterialEditor/Private/Material/DynamicMaterialInstanceFactory.cpp
--- a/Engine/Plugins/Experimental/DynamicMaterial/Source/DynamicMaterialEditor/Private/Material/DynamicMaterialInstanceFactory.cpp
+++ b/Engine/Plugins/Experimental/DynamicMaterial/Source/DynamicMaterialEditor/Private/Material/DynamicMaterialInstanceFactory.cpp
@@ -12,9 +12,6 @@
 
 #define LOCTEXT_NAMESPACE "MaterialDesignerInstanceFactory"
 
-const FString UDynamicMaterialInstanceFactory::BaseDirectory = FString("/Game/DynamicMaterials");
-const FString UDynamicMaterialInstanceFactory::BaseName = FString("M_DynMatInst_");
-
 UDynamicMaterialInstanceFactory::UDynamicMaterialInstanceFactory()
 {
 	SupportedClass = UDynamicMaterialInstance::StaticClass();
diff --git a/Engine/Plugins/Experimental/DynamicMaterial/Source/DynamicMaterialEditor/Private/Menus/DMMaterialStageInputMenus.cpp b/Engine/Plugins/Experimental/DynamicMaterial/Source/DynamicMaterialEditor/Private/Menus/DMMaterialStageInputMenus.cpp
--- a/Engine/Plugins/Experimental/DynamicMaterial/Source/DynamicMaterialEditor/Private/Menus/DMMaterialStageInputMenus.cpp
+++ b/Engine/Plugins/Experimental/DynamicMaterial/Source/DynamicMaterialEditor/Private/Menus/DMMaterialStageInputMenus.cpp
@@ -11,6 +11,7 @@
 #include "Components/DMMaterialValue.h"
 #include "Components/MaterialStageExpressions/DMMSETextureSample.h"
 #include "Components/MaterialStageExpressions/DMMSETextureSampleEdgeColor.h"
+#include "DMPrivate.h"
 #include "DMValueDefinition.h"
 #include "DynamicMaterialEditorModule.h"
 #include "Framework/MultiBox/MultiBoxBuilder.h"
@@ -711,7 +712,7 @@ void FDMMaterialStageInputMenus::GenerateChangeInputMenu_PreviousStage_Output_Ch
 		for (int32 ChannelIndex = 0; ChannelIndex < FloatSize; ++ChannelIndex)
 		{
 			const FText ChannelName = UDMValueDefinitionLibrary::GetValueDefinition(PreviousStageOutputConnectors[OutputIndex].Type).GetChannelName(ChannelIndex + 1);
-			const int32 OutputChannel = UE::DynamicMaterialEditor::ChannelIndexToChannelBit(ChannelIndex + 1);
+			const int32 OutputChannel = UE::DynamicMaterialEditor::Private::ChannelIndexToChannelBit(ChannelIndex + 1);
 
 			UE::DynamicMaterialEditor::Private::GenerateChangeInputMenu_PreviousStage_Output_Channel_Impl(InChildMenuBuilder, ChannelName, Stage, InThroughput,
 				InInputIndex, InInputChannel, InMaterialProperty, OutputIndex, OutputChannel);
@@ -997,7 +998,7 @@ void FDMMaterialStageInputMenus::GenerateChangeInputMenu_GlobalValue_Channels(FM
 		for (int32 ChannelIndex = 0; ChannelIndex < FloatSize; ++ChannelIndex)
 		{
 			const FText ChannelName = UDMValueDefinitionLibrary::GetValueDefinition(Values[InValueIndex]->GetType()).GetChannelName(ChannelIndex + 1);
-			const int32 OutputChannel = UE::DynamicMaterialEditor::ChannelIndexToChannelBit(ChannelIndex + 1);
+			const int32 OutputChannel = UE::DynamicMaterialEditor::Private::ChannelIndexToChannelBit(ChannelIndex + 1);
 
 			UE::DynamicMaterialEditor::Private::GenerateChangeInputMenu_GlobalValue_Channel_Impl(InChildMenuBuilder,
 				ChannelName, Stage, InThroughput, InInputIndex, InInputChannel, InValueIndex, Values[InValueIndex], OutputChannel);
@@ -1633,7 +1634,7 @@ void FDMMaterialStageInputMenus::GenerateChangeInputMenu_Slot_Property_Output_Ch
 		for (int32 ChannelIndex = 0; ChannelIndex < FloatSize; ++ChannelIndex)
 		{
 			const FText ChannelName = UDMValueDefinitionLibrary::GetValueDefinition(LastPropertySourceOutputConnectors[OutputIndex].Type).GetChannelName(ChannelIndex + 1);
-			const int32 OutputChannel = UE::DynamicMaterialEditor::ChannelIndexToChannelBit(ChannelIndex + 1);
+			const int32 OutputChannel = UE::DynamicMaterialEditor::Private::ChannelIndexToChannelBit(ChannelIndex + 1);
 
 			UE::DynamicMaterialEditor::Private::GenerateChangeInputMenu_Slot_Property_Output_Channel_Impl(InChildMenuBuilder, ChannelName,
 				Stage, InThroughput, InInputIndex, InInputChannel, InSlot, InMaterialProperty, OutputIndex, OutputChannel);
@@ -2021,7 +2022,7 @@ void FDMMaterialStageInputMenus::GenerateChangeInputMenu_Expression_Output_Chann
 		for (int32 ChannelIndex = 0; ChannelIndex < FloatSize; ++ChannelIndex)
 		{
 			const FText ChannelName = UDMValueDefinitionLibrary::GetValueDefinition(OutputConnectors[InOutputIndex].Type).GetChannelName(ChannelIndex + 1);
-			const int32 OutputChannel = UE::DynamicMaterialEditor::ChannelIndexToChannelBit(ChannelIndex + 1);
+			const int32 OutputChannel = UE::DynamicMaterialEditor::Private::ChannelIndexToChannelBit(ChannelIndex + 1);
 
 			UE::DynamicMaterialEditor::Private::GenerateChangeInputMenu_Expression_Output_Channel_Impl(InChildMenuBuilder, ChannelName, Stage, InThroughput,
 				InInputIndex, InInputChannel, InExpressionClass, InOutputIndex, OutputChannel);
diff --git a/Engine/Plugins/Experimental/DynamicMaterial/Source/DynamicMaterialEditor/Private/Menus/DMMenuContext.h b/Engine/Plugins/Experimental/DynamicMaterial/Source/DynamicMaterialEditor/Private/Menus/DMMenuContext.h
--- a/Engine/Plugins/Experimental/DynamicMaterial/Source/DynamicMaterialEditor/Private/Menus/DMMenuContext.h
+++ b/Engine/Plugins/Experimental/DynamicMaterial/Source/DynamicMaterialEditor/Private/Menus/DMMenuContext.h
@@ -15,8 +15,8 @@ class UDMMaterialStageSource;
 class UDynamicMaterialModel;
 class UToolMenu;
 
-UCLASS(BlueprintType)
-class DYNAMICMATERIALEDITOR_API UDMMenuContext : public UObject
+UCLASS()
+class UDMMenuContext : public UObject
 {
 	GENERATED_BODY()
 
@@ -34,19 +34,14 @@ public:
 	const TWeakPtr<SDMSlot>& GetSlotWidget() const { return SlotWidgetWeak; }
 	const TWeakPtr<SDMStage>& GetStageWidget() const { return StageWidgetWeak; }
 
-	UFUNCTION(BlueprintCallable, Category = "Material Designer Extensions")
 	UDMMaterialSlot* GetSlot() const;
 
-	UFUNCTION(BlueprintCallable, Category = "Material Designer Extensions")
 	UDynamicMaterialModel* GetModel() const;
 
-	UFUNCTION(BlueprintCallable, Category = "Material Designer Extensions")
 	UDMMaterialStage* GetStage() const;
 
-	UFUNCTION(BlueprintCallable, Category = "Material Designer Extensions")
 	UDMMaterialStageSource* GetStageSource() const;
 
-	UFUNCTION(BlueprintCallable, Category = "Material Designer Extensions")
 	UDMMaterialStageBlend* GetStageSourceAsBlend() const;
 
 	const UDMMaterialLayerObject* GetLayer() const;
diff --git a/Engine/Plugins/Experimental/DynamicMaterial/Source/DynamicMaterialEditor/Private/Menus/DMToolBarMenus.h b/Engine/Plugins/Experimental/DynamicMaterial/Source/DynamicMaterialEditor/Private/Menus/DMToolBarMenus.h
--- a/Engine/Plugins/Experimental/DynamicMaterial/Source/DynamicMaterialEditor/Private/Menus/DMToolBarMenus.h
+++ b/Engine/Plugins/Experimental/DynamicMaterial/Source/DynamicMaterialEditor/Private/Menus/DMToolBarMenus.h
@@ -2,26 +2,13 @@
 
 #pragma once
 
-#include "Layout/Visibility.h"
-#include "Misc/Attribute.h"
-#include "Templates/SharedPointer.h"
+#include "Templates/SharedPointerFwd.h"
 
 class SDMEditor;
 class SWidget;
-class UToolMenu;
-struct FToolMenuSection;
-struct FUIAction;
 
 class FDMToolBarMenus
 {
 public:
 	static TSharedRef<SWidget> MakeEditorLayoutMenu(const TSharedPtr<SDMEditor>& InEditorWidget = nullptr);
-
-	static void AddPreviewOptionsSection(UToolMenu* InMenu);
-	static void AddTooltipOptionsSection(UToolMenu* InMenu);
-	static void AddAdvancedSection(UToolMenu* InMenu);
-
-	/** Helper functions to add property menu entries. */
-	static void AddBoolOptionMenuEntry(FToolMenuSection& InSection, const FName& InPropertyName, const FUIAction InAction);
-	static void AddIntOptionMenuEntry(FToolMenuSection& InSection, const FName& InPropertyName, TAttribute<bool> InIsEnabledAttribute = TAttribute<bool>(), TAttribute<EVisibility> InVisibilityAttribute = TAttribute<EVisibility>());
 };
diff --git a/Engine/Plugins/Experimental/DynamicMaterial/Source/DynamicMaterialEditor/Private/Model/DMMaterialNodeArranger.cpp b/Engine/Plugins/Experimental/DynamicMaterial/Source/DynamicMaterialEditor/Private/Model/DMMaterialNodeArranger.cpp
new file mode 100644
--- /dev/null
+++ b/Engine/Plugins/Experimental/DynamicMaterial/Source/DynamicMaterialEditor/Private/Model/DMMaterialNodeArranger.cpp
@@ -0,0 +1,141 @@
+// Copyright Epic Games, Inc. All Rights Reserved.
+
+#include "Model/DMMaterialNodeArranger.h"
+#include "Containers/ArrayView.h"
+#include "Materials/Material.h"
+#include "Materials/MaterialExpression.h"
+
+namespace UE::DynamicMaterialEditor::BuildState::Private
+{
+	static constexpr int32 SpaceBetweenNodes = 50;
+}
+
+FDMMaterialNodeArranger::FDMMaterialNodeArranger(UMaterial* InDynamicMaterial)
+	: DynamicMaterial(InDynamicMaterial)
+	, OffsetStart({0, 0})
+{
+}
+
+void FDMMaterialNodeArranger::ArrangeNodes()
+{
+	ArrangeMaterialInputNodes(DynamicMaterial->GetEditorOnlyData()->BaseColor.Expression);
+	ArrangeMaterialInputNodes(DynamicMaterial->GetEditorOnlyData()->EmissiveColor.Expression);
+	ArrangeMaterialInputNodes(DynamicMaterial->GetEditorOnlyData()->Opacity.Expression);
+	ArrangeMaterialInputNodes(DynamicMaterial->GetEditorOnlyData()->OpacityMask.Expression);
+	ArrangeMaterialInputNodes(DynamicMaterial->GetEditorOnlyData()->Metallic.Expression);
+	ArrangeMaterialInputNodes(DynamicMaterial->GetEditorOnlyData()->Specular.Expression);
+	ArrangeMaterialInputNodes(DynamicMaterial->GetEditorOnlyData()->Roughness.Expression);
+	ArrangeMaterialInputNodes(DynamicMaterial->GetEditorOnlyData()->Anisotropy.Expression);
+	ArrangeMaterialInputNodes(DynamicMaterial->GetEditorOnlyData()->Normal.Expression);
+	ArrangeMaterialInputNodes(DynamicMaterial->GetEditorOnlyData()->Tangent.Expression);
+	ArrangeMaterialInputNodes(DynamicMaterial->GetEditorOnlyData()->WorldPositionOffset.Expression);
+	ArrangeMaterialInputNodes(DynamicMaterial->GetEditorOnlyData()->Refraction.Expression);
+	ArrangeMaterialInputNodes(DynamicMaterial->GetEditorOnlyData()->AmbientOcclusion.Expression);
+	ArrangeMaterialInputNodes(DynamicMaterial->GetEditorOnlyData()->PixelDepthOffset.Expression);
+
+	FInt32Interval Vertical = {0, 0};
+
+	for (const TPair<UMaterialExpression*, FIntPoint>& Pair : NodePositions)
+	{
+		Vertical.Min = FMath::Min(Vertical.Min, Pair.Value.Y);
+		Vertical.Max = FMath::Max(Vertical.Max, Pair.Value.Y);
+	}
+
+	const int32 Offset = (Vertical.Max - Vertical.Min) / -2;
+
+	for (const TPair<UMaterialExpression*, FIntPoint>& Pair : NodePositions)
+	{
+		Pair.Key->MaterialExpressionEditorY += Offset;
+	}
+}
+
+void FDMMaterialNodeArranger::ArrangeMaterialInputNodes(UMaterialExpression* MaterialInputExpression)
+{
+	if (!MaterialInputExpression)
+	{
+		return;
+	}
+
+	FIntPoint NodeSize = {0, 0};
+	ArrangeNode(NodePositions, OffsetStart, MaterialInputExpression, NodeSize);
+
+	OffsetStart.X += NodeSize.X;
+	OffsetStart.Y += NodeSize.Y + UE::DynamicMaterialEditor::BuildState::Private::SpaceBetweenNodes;
+}
+
+void FDMMaterialNodeArranger::ArrangeNode(TMap<UMaterialExpression*, FIntPoint>& InOutNodePositions, const FIntPoint& InOffsetStart, 
+	UMaterialExpression* InNode, FIntPoint& InOutNodeSize)
+{
+	if (!InNode)
+	{
+		return;
+	}
+
+	TConstArrayView<FExpressionInput*> Inputs = InNode->GetInputsView();
+	InOutNodeSize = {0, 0};
+	const FIntPoint ThisNodeSize = {InNode->GetWidth() * 2, InNode->GetHeight()};
+	FIntPoint ChildOffsetStart = InOffsetStart;
+	ChildOffsetStart.X += ThisNodeSize.X + UE::DynamicMaterialEditor::BuildState::Private::SpaceBetweenNodes;
+
+	for (FExpressionInput* Input : Inputs)
+	{
+		if (!Input->Expression)
+		{
+			continue;
+		}
+
+		FIntPoint ChildNodeSize;
+		const FIntPoint* NodePosition = InOutNodePositions.Find(Input->Expression);
+
+		if (NodePosition && (-NodePosition->X) > InOffsetStart.X)
+		{
+			ChildNodeSize = {Input->Expression->GetWidth(), Input->Expression->GetHeight()};
+		}
+		else
+		{
+			ArrangeNode(InOutNodePositions, ChildOffsetStart, Input->Expression, ChildNodeSize);
+		}
+
+		if (ChildNodeSize.X > 0)
+		{
+			InOutNodeSize.X = FMath::Max(InOutNodeSize.X, ChildNodeSize.X);
+		}
+
+		if (ChildNodeSize.Y > 0)
+		{
+			if (InOutNodeSize.Y > 0)
+			{
+				InOutNodeSize.Y += UE::DynamicMaterialEditor::BuildState::Private::SpaceBetweenNodes;
+			}
+
+			InOutNodeSize.Y += ChildNodeSize.Y;
+		}
+
+		ChildOffsetStart.Y = InOffsetStart.Y + InOutNodeSize.Y + UE::DynamicMaterialEditor::BuildState::Private::SpaceBetweenNodes;
+	}
+
+	if (InOutNodeSize.X > 0)
+	{
+		InOutNodeSize.X += UE::DynamicMaterialEditor::BuildState::Private::SpaceBetweenNodes;
+	}
+
+	InOutNodeSize.X += ThisNodeSize.X;
+	InNode->MaterialExpressionEditorX = -InOffsetStart.X - ThisNodeSize.X;
+
+	if (InOutNodeSize.Y == 0)
+	{
+		InOutNodeSize.Y = ThisNodeSize.Y;
+		InNode->MaterialExpressionEditorY = InOffsetStart.Y;
+	}
+	else if (InOutNodeSize.Y <= ThisNodeSize.Y)
+	{
+		InOutNodeSize = ThisNodeSize.Y;
+		InNode->MaterialExpressionEditorY = InOffsetStart.Y;
+	}
+	else
+	{
+		InNode->MaterialExpressionEditorY = InOffsetStart.Y + (InOutNodeSize.Y - ThisNodeSize.Y) / 2;
+	}
+
+	InOutNodePositions.Emplace(InNode, FIntPoint(InNode->MaterialExpressionEditorX, InNode->MaterialExpressionEditorY));
+}
diff --git a/Engine/Plugins/Experimental/DynamicMaterial/Source/DynamicMaterialEditor/Private/Model/DMMaterialNodeArranger.h b/Engine/Plugins/Experimental/DynamicMaterial/Source/DynamicMaterialEditor/Private/Model/DMMaterialNodeArranger.h
new file mode 100644
--- /dev/null
+++ b/Engine/Plugins/Experimental/DynamicMaterial/Source/DynamicMaterialEditor/Private/Model/DMMaterialNodeArranger.h
@@ -0,0 +1,30 @@
+// Copyright Epic Games, Inc. All Rights Reserved.
+
+#pragma once
+
+#include "Containers/Map.h"
+#include "Math/IntPoint.h"
+
+class UMaterial;
+class UMaterialExpression;
+
+/**
+ * Algorithm for arranging material nodes.
+ */
+class FDMMaterialNodeArranger
+{
+public:
+	FDMMaterialNodeArranger(UMaterial* InDynamicMaterial);
+
+	void ArrangeNodes();
+
+protected:
+	UMaterial* DynamicMaterial;
+	TMap<UMaterialExpression*, FIntPoint> NodePositions;
+	FIntPoint OffsetStart = {0, 0};
+
+	void ArrangeMaterialInputNodes(UMaterialExpression* MaterialInputExpression);
+
+	void ArrangeNode(TMap<UMaterialExpression*, FIntPoint>& InOutNodePositions, const FIntPoint& InOffsetStart,
+		UMaterialExpression* InNode, FIntPoint& InOutNodeSize);
+};
diff --git a/Engine/Plugins/Experimental/DynamicMaterial/Source/DynamicMaterialEditor/Private/Model/DynamicMaterialModelEditorOnlyData.cpp b/Engine/Plugins/Experimental/DynamicMaterial/Source/DynamicMaterialEditor/Private/Model/DynamicMaterialModelEditorOnlyData.cpp
--- a/Engine/Plugins/Experimental/DynamicMaterial/Source/DynamicMaterialEditor/Private/Model/DynamicMaterialModelEditorOnlyData.cpp
+++ b/Engine/Plugins/Experimental/DynamicMaterial/Source/DynamicMaterialEditor/Private/Model/DynamicMaterialModelEditorOnlyData.cpp
@@ -3,6 +3,7 @@
 #include "Model/DynamicMaterialModelEditorOnlyData.h"
 #include "AssetRegistry/AssetRegistryModule.h"
 #include "AssetToolsModule.h"
+#include "Components/DMMaterialLayer.h"
 #include "Components/DMMaterialParameter.h"
 #include "Components/DMMaterialProperty.h"
 #include "Components/DMMaterialSlot.h"
@@ -36,13 +37,13 @@
 #include "DMComponentPath.h"
 #include "DMDefs.h"
 #include "DMMaterialFunctionLibrary.h"
+#include "DMPrivate.h"
 #include "DynamicMaterialEditorModule.h"
 #include "DynamicMaterialEditorSettings.h"
 #include "DynamicMaterialModule.h"
 #include "Factories/MaterialFactoryNew.h"
 #include "FileHelpers.h"
 #include "IAssetTools.h"
-#include "Components/DMMaterialLayer.h"
 #include "Material/DynamicMaterialInstance.h"
 #include "Materials/Material.h"
 #include "Materials/MaterialExpressionAppendVector.h"
@@ -51,9 +52,9 @@
 #include "Materials/MaterialExpressionMultiply.h"
 #include "Materials/MaterialFunction.h"
 #include "Misc/Guid.h"
-#include "Model/DynamicMaterialModel.h"
 #include "Model/DMMaterialBuildState.h"
 #include "Model/DMMaterialBuildUtils.h"
+#include "Model/DynamicMaterialModel.h"
 #include "UObject/Package.h"
 
 #define LOCTEXT_NAMESPACE "MaterialDesignerModel"
@@ -241,7 +242,7 @@ void UDynamicMaterialModelEditorOnlyData::BuildMaterial(bool bInDirtyAssets)
 	 */
 	for (const TPair<EDMMaterialPropertyType, TObjectPtr<UDMMaterialProperty>>& Pair : Properties)
 	{
-		if (UE::DynamicMaterialEditor::IsCustomMaterialProperty(Pair.Key))
+		if (UE::DynamicMaterialEditor::Private::IsCustomMaterialProperty(Pair.Key))
 		{
 			continue;
 		}
@@ -406,7 +407,7 @@ void UDynamicMaterialModelEditorOnlyData::BuildMaterial(bool bInDirtyAssets)
 	 */
 	for (const TPair<EDMMaterialPropertyType, TObjectPtr<UDMMaterialProperty>>& Pair : Properties)
 	{
-		if (UE::DynamicMaterialEditor::IsCustomMaterialProperty(Pair.Key))
+		if (UE::DynamicMaterialEditor::Private::IsCustomMaterialProperty(Pair.Key))
 		{
 			continue;
 		}
diff --git a/Engine/Plugins/Experimental/DynamicMaterial/Source/DynamicMaterialEditor/Private/Model/DynamicMaterialModelFactory.h b/Engine/Plugins/Experimental/DynamicMaterial/Source/DynamicMaterialEditor/Private/Model/DynamicMaterialModelFactory.h
new file mode 100644
--- /dev/null
+++ b/Engine/Plugins/Experimental/DynamicMaterial/Source/DynamicMaterialEditor/Private/Model/DynamicMaterialModelFactory.h
@@ -0,0 +1,23 @@
+// Copyright Epic Games, Inc. All Rights Reserved.
+
+#pragma once
+
+#include "Factories/Factory.h"
+#include "Containers/UnrealString.h"
+#include "DynamicMaterialModelFactory.generated.h"
+
+UCLASS()
+class UDynamicMaterialModelFactory : public UFactory
+{
+	GENERATED_BODY()
+
+public:		
+	static const FString BaseDirectory;
+	static const FString BaseName;
+
+	UDynamicMaterialModelFactory();
+
+	virtual UObject* FactoryCreateNew(UClass* Class, UObject* InParent, FName Name, EObjectFlags Flags, UObject* Context, FFeedbackContext* Warn) override;
+	virtual FText GetDisplayName() const override;
+	virtual FText GetToolTip() const override;
+};
\ No newline at end of file
diff --git a/Engine/Plugins/Experimental/DynamicMaterial/Source/DynamicMaterialEditor/Public/DMBlueprintFunctionLibrary.h b/Engine/Plugins/Experimental/DynamicMaterial/Source/DynamicMaterialEditor/Public/DMBlueprintFunctionLibrary.h
deleted file mode 100644
diff --git a/Engine/Plugins/Experimental/DynamicMaterial/Source/DynamicMaterialEditor/Public/DMEDefs.h b/Engine/Plugins/Experimental/DynamicMaterial/Source/DynamicMaterialEditor/Public/DMEDefs.h
--- a/Engine/Plugins/Experimental/DynamicMaterial/Source/DynamicMaterialEditor/Public/DMEDefs.h
+++ b/Engine/Plugins/Experimental/DynamicMaterial/Source/DynamicMaterialEditor/Public/DMEDefs.h
@@ -24,25 +24,14 @@ class UDynamicMaterialInstance;
 class UDynamicMaterialModel;
 class UPrimitiveComponent;
 
-namespace UE::DynamicMaterialEditor
-{
-	/** Converts 0,1,2,3,4 to 0,1,2,4,8 */
-	int32 ChannelIndexToChannelBit(int32 InChannelIndex);
-
-	/** Converts 0,1,2,4,8 to 0,1,2,3,4 */
-	int32 ChannelBitToChannelIndex(int32 InChannelBit);
-
-	bool IsCustomMaterialProperty(EDMMaterialPropertyType InMaterialProperty);
-}
-
-UINTERFACE(Blueprintable, BlueprintType)
-class DYNAMICMATERIALEDITOR_API UDMBuildable : public UInterface
+UINTERFACE(MinimalAPI, Blueprintable, BlueprintType)
+class UDMBuildable : public UInterface
 {
 public:
 	GENERATED_BODY()
 };
 
-class DYNAMICMATERIALEDITOR_API IDMBuildable
+class IDMBuildable
 {
 public:
 	GENERATED_BODY()
@@ -57,7 +46,7 @@ public:
  * An input or output form a material source/stage (e.g. RGB out.)
  */
 USTRUCT(BlueprintType, Category = "Material Designer", meta = (DisplayName = "Material Designer Stage Connector"))
-struct DYNAMICMATERIALEDITOR_API FDMMaterialStageConnector
+struct FDMMaterialStageConnector
 {
 	GENERATED_BODY()
 
@@ -101,7 +90,7 @@ struct DYNAMICMATERIALEDITOR_API FDMMaterialStageConnector
  * input channels (e.g. T1.R, T2.B, T3.G -> {T2.B, T1.R, T3.G, T2.B})
  */
 USTRUCT(BlueprintType, Category = "Material Designer", meta = (DisplayName = "Material Designer Stage Connection"))
-struct DYNAMICMATERIALEDITOR_API FDMMaterialStageConnection
+struct FDMMaterialStageConnection
 {
 	GENERATED_BODY()
 
@@ -115,7 +104,7 @@ struct DYNAMICMATERIALEDITOR_API FDMMaterialStageConnection
 };
 
 USTRUCT(BlueprintType, Category = "Material Designer", meta = (DisplayName = "Material Designer Slot Output Connector Types"))
-struct DYNAMICMATERIALEDITOR_API FDMMaterialSlotOutputConnectorTypes
+struct FDMMaterialSlotOutputConnectorTypes
 {
 	GENERATED_BODY()
 
diff --git a/Engine/Plugins/Experimental/DynamicMaterial/Source/DynamicMaterialEditor/Public/DMMaterialFunctionLibrary.h b/Engine/Plugins/Experimental/DynamicMaterial/Source/DynamicMaterialEditor/Public/DMMaterialFunctionLibrary.h
deleted file mode 100644
diff --git a/Engine/Plugins/Experimental/DynamicMaterial/Source/DynamicMaterialEditor/Public/DMWorldSubsystem.h b/Engine/Plugins/Experimental/DynamicMaterial/Source/DynamicMaterialEditor/Public/DMWorldSubsystem.h
--- a/Engine/Plugins/Experimental/DynamicMaterial/Source/DynamicMaterialEditor/Public/DMWorldSubsystem.h
+++ b/Engine/Plugins/Experimental/DynamicMaterial/Source/DynamicMaterialEditor/Public/DMWorldSubsystem.h
@@ -2,7 +2,7 @@
  
 #pragma once
 
-#include "DynamicMaterialEditorModule.h"
+#include "Delegates/Delegate.h"
 #include "Subsystems/WorldSubsystem.h"
 #include "Templates/SharedPointer.h"
 #include "DMWorldSubsystem.generated.h"
@@ -13,12 +13,15 @@ class UDynamicMaterialInstance;
 class UDynamicMaterialModel;
 struct FDMObjectMaterialProperty;
 
+DECLARE_DELEGATE_OneParam(FDMSetMaterialModelDelegate, UDynamicMaterialModel*)
+DECLARE_DELEGATE_OneParam(FDMSetMaterialObjectPropertyDelegate, const FDMObjectMaterialProperty&)
+DECLARE_DELEGATE_OneParam(FDMSetMaterialActorDelegate, AActor*)
 DECLARE_DELEGATE_RetVal_OneParam(bool, FDMIsValidDelegate, UDynamicMaterialModel*)
 DECLARE_DELEGATE_RetVal_TwoParams(bool, FDMSetMaterialValueDelegate, const FDMObjectMaterialProperty&, UDynamicMaterialInstance*)
 DECLARE_DELEGATE(FDMInvokeTabDelegate)
- 
-UCLASS()
-class DYNAMICMATERIALEDITOR_API UDMWorldSubsystem : public UWorldSubsystem
+
+UCLASS(MinimalAPI)
+class UDMWorldSubsystem : public UWorldSubsystem
 {
 	GENERATED_BODY()
  
@@ -26,7 +29,8 @@ public:
 	UDMWorldSubsystem();
  
 	const TSharedPtr<IDetailKeyframeHandler>& GetKeyframeHandler() const { return KeyframeHandler; }
-	void SetKeyframeHandler(const TSharedPtr<IDetailKeyframeHandler>& InKeyframeHandler);
+
+	void SetKeyframeHandler(const TSharedPtr<IDetailKeyframeHandler>& InKeyframeHandler) { KeyframeHandler = InKeyframeHandler; }
 
 	/** Sets material in a custom editor tab. */
 	FDMSetMaterialModelDelegate& GetSetCustomEditorModelDelegate() { return CustomModelEditorDelegate; }
diff --git a/Engine/Plugins/Experimental/DynamicMaterial/Source/DynamicMaterialEditor/Public/DetailsPanel/DMMaterialInterfaceTypeCustomizer.h b/Engine/Plugins/Experimental/DynamicMaterial/Source/DynamicMaterialEditor/Public/DetailsPanel/DMMaterialInterfaceTypeCustomizer.h
deleted file mode 100644
diff --git a/Engine/Plugins/Experimental/DynamicMaterial/Source/DynamicMaterialEditor/Public/DetailsPanel/Slate/SDMDetailsPanelMaterialInterfaceWidget.h b/Engine/Plugins/Experimental/DynamicMaterial/Source/DynamicMaterialEditor/Public/DetailsPanel/Slate/SDMDetailsPanelMaterialInterfaceWidget.h
deleted file mode 100644
diff --git a/Engine/Plugins/Experimental/DynamicMaterial/Source/DynamicMaterialEditor/Public/DynamicMaterialEditorModule.h b/Engine/Plugins/Experimental/DynamicMaterial/Source/DynamicMaterialEditor/Public/DynamicMaterialEditorModule.h
deleted file mode 100644
diff --git a/Engine/Plugins/Experimental/DynamicMaterial/Source/DynamicMaterialEditor/Public/DynamicMaterialEditorSettings.h b/Engine/Plugins/Experimental/DynamicMaterial/Source/DynamicMaterialEditor/Public/DynamicMaterialEditorSettings.h
deleted file mode 100644
diff --git a/Engine/Plugins/Experimental/DynamicMaterial/Source/DynamicMaterialEditor/Public/IDynamicMaterialEditorModule.h b/Engine/Plugins/Experimental/DynamicMaterial/Source/DynamicMaterialEditor/Public/IDynamicMaterialEditorModule.h
new file mode 100644
--- /dev/null
+++ b/Engine/Plugins/Experimental/DynamicMaterial/Source/DynamicMaterialEditor/Public/IDynamicMaterialEditorModule.h
@@ -0,0 +1,35 @@
+// Copyright Epic Games, Inc. All Rights Reserved.
+
+#pragma once
+
+#include "Containers/ContainersFwd.h"
+#include "Delegates/Delegate.h"
+#include "Modules/ModuleManager.h"
+
+class UClass;
+class UObject;
+struct FDMObjectMaterialProperty;
+
+DECLARE_DELEGATE_RetVal_OneParam(TArray<FDMObjectMaterialProperty>, FDMGetObjectMaterialPropertiesDelegate, UObject* InObject)
+
+/** Material Designer - Build your own materials in a slimline editor! */
+class IDynamicMaterialEditorModule : public IModuleInterface
+{
+protected:
+	static constexpr const TCHAR* ModuleName = TEXT("DynamicMaterialEditor");
+
+public:
+	static bool IsLoaded()
+	{
+		return FModuleManager::Get().IsModuleLoaded(ModuleName);
+	}
+
+	static IDynamicMaterialEditorModule& Get()
+	{
+		return FModuleManager::LoadModuleChecked<IDynamicMaterialEditorModule>(ModuleName);
+	}
+
+	virtual void OpenEditor(UWorld* InWorld) = 0;
+
+	virtual void RegisterCustomMaterialPropertyGenerator(UClass* InClass, FDMGetObjectMaterialPropertiesDelegate InGenerator) = 0;
+};
diff --git a/Engine/Plugins/Experimental/DynamicMaterial/Source/DynamicMaterialEditor/Public/Material/DynamicMaterialInstanceFactory.h b/Engine/Plugins/Experimental/DynamicMaterial/Source/DynamicMaterialEditor/Public/Material/DynamicMaterialInstanceFactory.h
--- a/Engine/Plugins/Experimental/DynamicMaterial/Source/DynamicMaterialEditor/Public/Material/DynamicMaterialInstanceFactory.h
+++ b/Engine/Plugins/Experimental/DynamicMaterial/Source/DynamicMaterialEditor/Public/Material/DynamicMaterialInstanceFactory.h
@@ -3,23 +3,19 @@
 #pragma once
 
 #include "Factories/Factory.h"
-#include "Containers/UnrealString.h"
 #include "DynamicMaterialInstanceFactory.generated.h"
 
-class UDynamicMaterialInstance;
-
-UCLASS(MinimalAPI, hidecategories = Object)
+UCLASS(MinimalAPI)
 class UDynamicMaterialInstanceFactory : public UFactory
 {
 	GENERATED_BODY()
 
 public:
-	static const FString BaseDirectory;
-	static const FString BaseName;
-
 	UDynamicMaterialInstanceFactory();
 
+	//~ Begin UFactory
 	virtual UObject* FactoryCreateNew(UClass* Class, UObject* InParent, FName Name, EObjectFlags Flags, UObject* Context, FFeedbackContext* Warn) override;
 	virtual FText GetDisplayName() const override;
 	virtual FText GetToolTip() const override;
+	//~ End UFactory
 };
diff --git a/Engine/Plugins/Experimental/DynamicMaterial/Source/DynamicMaterialEditor/Public/Model/DMMaterialNodeArranger.cpp b/Engine/Plugins/Experimental/DynamicMaterial/Source/DynamicMaterialEditor/Public/Model/DMMaterialNodeArranger.cpp
deleted file mode 100644
diff --git a/Engine/Plugins/Experimental/DynamicMaterial/Source/DynamicMaterialEditor/Public/Model/DMMaterialNodeArranger.h b/Engine/Plugins/Experimental/DynamicMaterial/Source/DynamicMaterialEditor/Public/Model/DMMaterialNodeArranger.h
deleted file mode 100644
diff --git a/Engine/Plugins/Experimental/DynamicMaterial/Source/DynamicMaterialEditor/Public/Model/DynamicMaterialModelFactory.h b/Engine/Plugins/Experimental/DynamicMaterial/Source/DynamicMaterialEditor/Public/Model/DynamicMaterialModelFactory.h
deleted file mode 100644
