Material Designer:
-cis fixes

#changelist validated
#virtualized

diff --git a/Engine/Plugins/Experimental/DynamicMaterial/Source/DynamicMaterialEditor/Private/DetailsPanel/Slate/SDMMaterialListExtensionWidget.cpp b/Engine/Plugins/Experimental/DynamicMaterial/Source/DynamicMaterialEditor/Private/DetailsPanel/Slate/SDMMaterialListExtensionWidget.cpp
--- a/Engine/Plugins/Experimental/DynamicMaterial/Source/DynamicMaterialEditor/Private/DetailsPanel/Slate/SDMMaterialListExtensionWidget.cpp
+++ b/Engine/Plugins/Experimental/DynamicMaterial/Source/DynamicMaterialEditor/Private/DetailsPanel/Slate/SDMMaterialListExtensionWidget.cpp
@@ -3,6 +3,7 @@
 #include "SDMMaterialListExtensionWidget.h"
 #include "AssetToolsModule.h"
 #include "Components/PrimitiveComponent.h"
+#include "DMObjectMaterialProperty.h"
 #include "DMWorldSubsystem.h"
 #include "DetailLayoutBuilder.h"
 #include "Engine/World.h"
diff --git a/Engine/Plugins/Experimental/DynamicMaterial/Source/DynamicMaterialEditor/Private/Menus/DMToolBarMenus.h b/Engine/Plugins/Experimental/DynamicMaterial/Source/DynamicMaterialEditor/Private/Menus/DMToolBarMenus.h
--- a/Engine/Plugins/Experimental/DynamicMaterial/Source/DynamicMaterialEditor/Private/Menus/DMToolBarMenus.h
+++ b/Engine/Plugins/Experimental/DynamicMaterial/Source/DynamicMaterialEditor/Private/Menus/DMToolBarMenus.h
@@ -2,7 +2,7 @@
 
 #pragma once
 
-#include "Templates/SharedPointerFwd.h"
+#include "Templates/SharedPointer.h"
 
 class SDMEditor;
 class SWidget;
