Material Designer:
- Fixed issue with Material Desigber Tab not appearing on mac

#changelist validated
#virtualized

diff --git a/Engine/Plugins/Experimental/DynamicMaterial/Source/DynamicMaterialEditor/Private/LevelEditor/DMLevelEditorIntegration.cpp b/Engine/Plugins/Experimental/DynamicMaterial/Source/DynamicMaterialEditor/Private/LevelEditor/DMLevelEditorIntegration.cpp
--- a/Engine/Plugins/Experimental/DynamicMaterial/Source/DynamicMaterialEditor/Private/LevelEditor/DMLevelEditorIntegration.cpp
+++ b/Engine/Plugins/Experimental/DynamicMaterial/Source/DynamicMaterialEditor/Private/LevelEditor/DMLevelEditorIntegration.cpp
@@ -18,24 +18,26 @@ namespace UE::DynamicMaterialEditor::Private
 	{
 		return FModuleManager::GetModulePtr<FLevelEditorModule>("LevelEditor");
 	}
+
+	FLevelEditorModule& LoadLevelEditorModuleChecked()
+	{
+		return FModuleManager::LoadModuleChecked<FLevelEditorModule>("LevelEditor");
+	}
 }
 
 void FDMLevelEditorIntegration::Initialize()
 {
 	using namespace UE::DynamicMaterialEditor::Private;
 
-	if (FLevelEditorModule* ModulePtr = GetLevelEditorModulePtr())
-	{
-		LevelEditorCreatedHandle = ModulePtr->OnLevelEditorCreated().AddLambda(
-			[](TSharedPtr<ILevelEditor> InLevelEditor)
+	LevelEditorCreatedHandle = LoadLevelEditorModuleChecked().OnLevelEditorCreated().AddLambda(
+		[](TSharedPtr<ILevelEditor> InLevelEditor)
+		{
+			if (InLevelEditor.IsValid())
 			{
-				if (InLevelEditor.IsValid())
-				{
-					FDMLevelEditorIntegrationInstance::AddIntegration(InLevelEditor.ToSharedRef());
-				}
+				FDMLevelEditorIntegrationInstance::AddIntegration(InLevelEditor.ToSharedRef());
 			}
-		);
-	}
+		}
+	);
 }
 
 void FDMLevelEditorIntegration::Shutdown()
