Material Designer
- Fixed a bug where material function inputs were not visible/editable.
- Moved and consolidated stage source and input functions into more appropriate classes.

diff --git a/Engine/Plugins/Experimental/Avalanche/Source/AvalancheEditor/Private/MaterialDesigner/AvaMaterialDesignerTextureAssetFactory.cpp b/Engine/Plugins/Experimental/Avalanche/Source/AvalancheEditor/Private/MaterialDesigner/AvaMaterialDesignerTextureAssetFactory.cpp
--- a/Engine/Plugins/Experimental/Avalanche/Source/AvalancheEditor/Private/MaterialDesigner/AvaMaterialDesignerTextureAssetFactory.cpp
+++ b/Engine/Plugins/Experimental/Avalanche/Source/AvalancheEditor/Private/MaterialDesigner/AvaMaterialDesignerTextureAssetFactory.cpp
@@ -149,8 +149,14 @@ AActor* UAvaMaterialDesignerTextureAssetFactory::SpawnActor(UObject* InAsset, UL
 		return InNewActor;
 	}
 
-	UDMMaterialStageInputExpression* NewInput = BaseStage->ChangeInput_Expression(UDMMaterialStageBlend::InputB, FDMMaterialStageConnectorChannel::WHOLE_CHANNEL,
-		UDMMaterialStageExpressionTextureSample::StaticClass(), 0, FDMMaterialStageConnectorChannel::WHOLE_CHANNEL);
+	UDMMaterialStageInputExpression* NewInput = UDMMaterialStageInputExpression::ChangeStageInput_Expression(
+		BaseStage,
+		UDMMaterialStageExpressionTextureSample::StaticClass(),
+		UDMMaterialStageBlend::InputB,
+		FDMMaterialStageConnectorChannel::WHOLE_CHANNEL,
+		0,
+		FDMMaterialStageConnectorChannel::WHOLE_CHANNEL
+	);
 
 	if (!NewInput)
 	{
@@ -164,7 +170,14 @@ AActor* UAvaMaterialDesignerTextureAssetFactory::SpawnActor(UObject* InAsset, UL
 		return InNewActor;
 	}
 
-	UDMMaterialStageInputValue* NewInputValue = SubStage->ChangeInput_NewLocalValue(0, FDMMaterialStageConnectorChannel::WHOLE_CHANNEL, EDMValueType::VT_Texture, FDMMaterialStageConnectorChannel::WHOLE_CHANNEL);
+	UDMMaterialStageInputValue* NewInputValue = UDMMaterialStageInputValue::ChangeStageInput_NewLocalValue(
+		SubStage,
+		0, 
+		FDMMaterialStageConnectorChannel::WHOLE_CHANNEL,
+		EDMValueType::VT_Texture,
+		FDMMaterialStageConnectorChannel::WHOLE_CHANNEL
+	);
+
 	UDMMaterialValueTexture* TextureValue = Cast<UDMMaterialValueTexture>(NewInputValue->GetValue());
 
 	if (!TextureValue)
diff --git a/Engine/Plugins/Experimental/DynamicMaterial/Source/DynamicMaterialEditor/DynamicMaterialEditor.Build.cs b/Engine/Plugins/Experimental/DynamicMaterial/Source/DynamicMaterialEditor/DynamicMaterialEditor.Build.cs
--- a/Engine/Plugins/Experimental/DynamicMaterial/Source/DynamicMaterialEditor/DynamicMaterialEditor.Build.cs
+++ b/Engine/Plugins/Experimental/DynamicMaterial/Source/DynamicMaterialEditor/DynamicMaterialEditor.Build.cs
@@ -7,6 +7,7 @@ public class DynamicMaterialEditor : ModuleRules
 	public DynamicMaterialEditor(ReadOnlyTargetRules Target) : base(Target)
 	{
 		PCHUsage = ModuleRules.PCHUsageMode.UseExplicitOrSharedPCHs;
+		bUseUnity = false;
 
 		PublicDependencyModuleNames.AddRange(
 			new string[]
diff --git a/Engine/Plugins/Experimental/DynamicMaterial/Source/DynamicMaterialEditor/Private/Components/DMMaterialEffectFunction.cpp b/Engine/Plugins/Experimental/DynamicMaterial/Source/DynamicMaterialEditor/Private/Components/DMMaterialEffectFunction.cpp
--- a/Engine/Plugins/Experimental/DynamicMaterial/Source/DynamicMaterialEditor/Private/Components/DMMaterialEffectFunction.cpp
+++ b/Engine/Plugins/Experimental/DynamicMaterial/Source/DynamicMaterialEditor/Private/Components/DMMaterialEffectFunction.cpp
@@ -175,13 +175,19 @@ void UDMMaterialEffectFunction::ApplyTo(const TSharedRef<FDMMaterialBuildState>&
 			{
 				if (InOutLastExpressionOutputChannel != FDMMaterialStageConnectorChannel::WHOLE_CHANNEL)
 				{
-					LastInputExpressions[0] = InBuildState->GetBuildUtils().CreateExpressionBitMask(LastInputExpressions[0], InOutLastExpressionOutputIndex,
-						InOutLastExpressionOutputChannel);
+					LastInputExpressions[0] = InBuildState->GetBuildUtils().CreateExpressionBitMask(
+						LastInputExpressions[0], 
+						InOutLastExpressionOutputIndex,
+						InOutLastExpressionOutputChannel
+					);
 
 					InOutStageExpressions.Add(LastInputExpressions[0]);
 				}
 
-				FunctionCall->FunctionInputs[InputIndex].Input.Connect(InOutLastExpressionOutputIndex, LastInputExpressions[InputIndex]);
+				FunctionCall->FunctionInputs[InputIndex].Input.Connect(
+					InOutLastExpressionOutputIndex, 
+					LastInputExpressions[InputIndex]
+				);
 			}
 			else
 			{
@@ -264,7 +270,7 @@ void UDMMaterialEffectFunction::DeinitFunction()
 	InputValues.Empty();
 }
 
-bool UDMMaterialEffectFunction::NeedsFunctionInit()
+bool UDMMaterialEffectFunction::NeedsFunctionInit() const
 {
 	if (!IsValid(MaterialFunctionPtr))
 	{
@@ -272,27 +278,6 @@ bool UDMMaterialEffectFunction::NeedsFunctionInit()
 		return !InputValues.IsEmpty();
 	}
 
-	UDMMaterialSlot* Slot = GetTypedParent<UDMMaterialSlot>(/* bAllowSubclasses */ true);
-
-	if (!Slot)
-	{
-		return false;
-	}
-
-	UDynamicMaterialModelEditorOnlyData* EditorOnlyData = Slot->GetMaterialModelEditorOnlyData();
-
-	if (!EditorOnlyData)
-	{
-		return false;
-	}
-
-	UDynamicMaterialModel* MaterialModel = EditorOnlyData->GetMaterialModel();
-
-	if (!MaterialModel)
-	{
-		return false;
-	}
-
 	TArray<FFunctionExpressionInput> Inputs;
 	TArray<FFunctionExpressionOutput> Outputs;
 	MaterialFunctionPtr->GetInputsAndOutputs(Inputs, Outputs);
@@ -331,19 +316,9 @@ bool UDMMaterialEffectFunction::NeedsFunctionInit()
 			switch (FunctionInput->InputType)
 			{
 				case EFunctionInputType::FunctionInput_Scalar:
-					EffectTarget = EDMMaterialEffectTarget::MaskStage;
-					break;
-
 				case EFunctionInputType::FunctionInput_Vector2:
-					EffectTarget = EDMMaterialEffectTarget::TextureUV;
-					break;
-
 				case EFunctionInputType::FunctionInput_Vector3:
-					EffectTarget = EDMMaterialEffectTarget::BaseStage;
-					break;
-
 				case EFunctionInputType::FunctionInput_Vector4:
-					EffectTarget = EDMMaterialEffectTarget::Slot;
 					break;
 
 				default:
diff --git a/Engine/Plugins/Experimental/DynamicMaterial/Source/DynamicMaterialEditor/Private/Components/DMMaterialStage.cpp b/Engine/Plugins/Experimental/DynamicMaterial/Source/DynamicMaterialEditor/Private/Components/DMMaterialStage.cpp
--- a/Engine/Plugins/Experimental/DynamicMaterial/Source/DynamicMaterialEditor/Private/Components/DMMaterialStage.cpp
+++ b/Engine/Plugins/Experimental/DynamicMaterial/Source/DynamicMaterialEditor/Private/Components/DMMaterialStage.cpp
@@ -243,10 +243,14 @@ UDMMaterialComponent* UDMMaterialStage::GetSubComponentByPath(FDMComponentPath&
 		{
 			for (UDMMaterialStageInput* Input : Inputs)
 			{
-				if (Input && (Input->GetClass()->GetName().Equals(InputType)
-					|| Input->GetClass()->GetName().Equals(UDMMaterialStageInput::StageInputPrefixStr + InputType)))
+				if (Input)
 				{
-					return Input;
+					const FString InputClassName = Input->GetClass()->GetName();
+					
+					if (InputClassName.Equals(InputType) || InputClassName.Equals(UDMMaterialStageInput::StageInputPrefixStr + InputType))
+					{
+						return Input;
+					}
 				}
 			}
 		}
@@ -341,47 +345,6 @@ void UDMMaterialStage::OnTextureUVUpdated(UDynamicMaterialModel* InMaterialModel
 	}
 }
 
-void UDMMaterialStage::ApplyWholeLayerValue(UDMMaterialValue* InValue)
-{
-	if (!InValue || !InValue->IsWholeLayerValue())
-	{
-		return;
-	}
-
-	UDMMaterialLayerObject* Layer = GetLayer();
-
-	if (!Layer)
-	{
-		return;
-	}
-
-	if (!Layer 
-		|| Layer->GetStageType(this) != EDMMaterialLayerStage::Base
-		|| !Layer->IsStageEnabled(EDMMaterialLayerStage::Mask))
-	{
-		return;
-	}
-
-	UDMMaterialStage* MaskStage = Layer->GetStage(EDMMaterialLayerStage::Mask);
-
-	if (!IsValid(MaskStage))
-	{
-		return;
-	}
-
-	if (GUndo)
-	{
-		MaskStage->Modify();
-	}
-
-	MaskStage->ChangeInput_Value(
-		UDMMaterialStageThroughputLayerBlend::InputMaskSource,
-		FDMMaterialStageConnectorChannel::WHOLE_CHANNEL,
-		InValue,
-		FDMMaterialStageConnectorChannel::FOURTH_CHANNEL
-	);
-}
-
 UDMMaterialLayerObject* UDMMaterialStage::GetLayer() const
 {
 	return Cast<UDMMaterialLayerObject>(GetOuterSafe());
@@ -629,14 +592,24 @@ void UDMMaterialStage::GenerateExpressions(const TSharedRef<FDMMaterialBuildStat
 					Channel.OutputChannel = OutputChannelOverride;
 				}
 
-				const int32 NodeOutputIndex = Throughput->ResolveInput(InBuildState, InputIdx, Channel, Expressions);
+				const int32 NodeOutputIndex = Throughput->ResolveInput(
+					InBuildState, 
+					InputIdx, 
+					Channel, 
+					Expressions
+				);
 				
 				if (!Expressions.IsEmpty() && NodeOutputIndex != INDEX_NONE)
 				{
 					StageExpressions.Append(Expressions);
 
-					Throughput->ConnectOutputToInput(InBuildState, ThroughputInputs[InputIdx].Index,
-						Expressions.Last(), NodeOutputIndex, Channel.OutputChannel);
+					Throughput->ConnectOutputToInput(
+						InBuildState, 
+						ThroughputInputs[InputIdx].Index,
+						Expressions.Last(), 
+						NodeOutputIndex, 
+						Channel.OutputChannel
+					);
 				}
 			}
 		}
@@ -781,7 +754,8 @@ void UDMMaterialStage::RemoveUnusedInputs()
 	}
 }
 
-UDMMaterialStageSource* UDMMaterialStage::ChangeInput_PreviousStage(int32 InInputIdx, int32 InInputChannel, EDMMaterialPropertyType InPreviousStageProperty, int32 InOutputIdx, int32 InOutputChannel)
+UDMMaterialStageInput* UDMMaterialStage::ChangeInput(TSubclassOf<UDMMaterialStageInput> InInputClass, int32 InInputIdx,
+	int32 InInputChannel, int32 InOutputIdx, int32 InOutputChannel, FInputInitFunctionPtr InPreInit)
 {
 	check(Source);
 
@@ -791,159 +765,40 @@ UDMMaterialStageSource* UDMMaterialStage::ChangeInput_PreviousStage(int32 InInpu
 	const TArray<FDMMaterialStageConnector>& InputConnectors = Throughput->GetInputConnectors();
 	check(InputConnectors.IsValidIndex(InInputIdx));
 
-	UDMMaterialLayerObject* Layer = GetLayer();
-	check(Layer);
-
-	EDMMaterialPropertyType StageProperty = Layer->GetMaterialProperty();
-	check(StageProperty != EDMMaterialPropertyType::None);
+	check(InInputClass.Get());
+	check(!(InInputClass->ClassFlags & (CLASS_Abstract | CLASS_Hidden | CLASS_Deprecated | CLASS_NewerVersionExists)));
 
-	const UDMMaterialLayerObject* PreviousLayer = Layer->GetPreviousLayer(StageProperty, EDMMaterialLayerStage::Base);
-	UDMMaterialStageSource* PreviousSource = nullptr;
-	
-	if (PreviousLayer)
+	if (UDMMaterialStageThroughput* ThroughputCDO = Cast<UDMMaterialStageThroughput>(InInputClass->GetDefaultObject(true)))
 	{
-		if (UDMMaterialStage* Stage = PreviousLayer->GetLastEnabledStage(EDMMaterialLayerStage::All))
-		{
-			PreviousSource = Stage->GetSource();
-		}
-	}
+		check(!ThroughputCDO->IsInputRequired() || ThroughputCDO->AllowsNestedInputs());
 
-	if (PreviousSource)
-	{
-		const TArray<FDMMaterialStageConnector>& PreviousStageOutputs = PreviousSource->GetOutputConnectors();
-		check(PreviousStageOutputs.IsValidIndex(InOutputIdx));
-		check(Throughput->CanInputConnectTo(InInputIdx, PreviousStageOutputs[InOutputIdx], InOutputChannel));
+		const TArray<FDMMaterialStageConnector>& OutputConnectors = ThroughputCDO->GetOutputConnectors();
+		check(Throughput->CanInputConnectTo(InInputIdx, OutputConnectors[InOutputIdx], InOutputChannel));
 	}
-	else
-	{
-		check(InOutputIdx == 0);
-		check(Throughput->GetInputConnectors().IsEmpty() == false);
-		check(Throughput->GetInputConnectors()[0].IsCompatibleWith(EDMValueType::VT_Float3_RGB));
-	}
-
-	UpdateInputMap(
-		InInputIdx,
-		FDMMaterialStageConnectorChannel::PREVIOUS_STAGE,
-		InInputChannel,
-		InOutputIdx,
-		InOutputChannel,
-		InPreviousStageProperty
-	);
-
-	return PreviousSource;
-}
-
-UDMMaterialStageInputValue* UDMMaterialStage::ChangeInput_NewLocalValue(int32 InInputIdx, int32 InInputChannel, EDMValueType InType, int32 InOutputChannel)
-{
-	check(Source);
-
-	UDMMaterialStageThroughput* Throughput = Cast<UDMMaterialStageThroughput>(Source);
-	check(Throughput);
-
-	const TArray<FDMMaterialStageConnector>& InputConnectors = Throughput->GetInputConnectors();
-	check(InputConnectors.IsValidIndex(InInputIdx));
-	check(Throughput->CanInputAcceptType(InInputIdx, InType));
-
-	UDMMaterialLayerObject* Layer = GetLayer();
-	check(Layer);
-
-	UDMMaterialSlot* Slot = Layer->GetSlot();
-	check(Slot);
-
-	UDynamicMaterialModelEditorOnlyData* ModelEditorOnlyData = Slot->GetMaterialModelEditorOnlyData();
-	check(ModelEditorOnlyData);
-
-	UDynamicMaterialModel* MaterialModel = ModelEditorOnlyData->GetMaterialModel();
-	check(MaterialModel);
 
-	// Don't add via the builder, just parent it to the builder.
-	// It won't appear in the global value list.
-	UDMMaterialValue* NewValue = UDMMaterialValue::CreateMaterialValue(MaterialModel, TEXT(""), InType, true);
-	check(NewValue);
-
-	UDMMaterialStageInputValue* InputValue = NewObject<UDMMaterialStageInputValue>(this, NAME_None, RF_Transactional);
-	check(InputValue);
+	UDMMaterialStageInput* NewInput = NewObject<UDMMaterialStageInput>(this, InInputClass, NAME_None, RF_Transactional);
+	check(NewInput);
 
+	if (InPreInit)
 	{
-		const FDMUpdateGuard Guard;
-		InputValue->SetValue(NewValue);
+		InPreInit(this, NewInput);
 	}
 
-	AddInput(InputValue);
+	AddInput(NewInput);
 
 	UpdateInputMap(
 		InInputIdx,
 		FDMMaterialStageConnectorChannel::FIRST_STAGE_INPUT + Inputs.Num() - 1,
 		InInputChannel,
-		0,
+		InOutputIdx,
 		InOutputChannel,
 		EDMMaterialPropertyType::None
 	);
 
-	ApplyWholeLayerValue(NewValue);
-
-	return InputValue;
-}
-
-UDMMaterialStageInputValue* UDMMaterialStage::ChangeInput_Value(int32 InInputIdx, int32 InInputChannel, UDMMaterialValue* InValue, int32 OutputChannel)
-{
-	check(Source);
-	check(InValue);
-
-	UDMMaterialStageThroughput* Throughput = Cast<UDMMaterialStageThroughput>(Source);
-	check(Throughput);
-
-	const TArray<FDMMaterialStageConnector>& InputConnectors = Throughput->GetInputConnectors();
-	check(InputConnectors.IsValidIndex(InInputIdx));
-
-	if (OutputChannel == FDMMaterialStageConnectorChannel::WHOLE_CHANNEL)
-	{
-		check(Throughput->CanInputAcceptType(InInputIdx, InValue->GetType()));
-	}
-	else
-	{
-		check(UDMValueDefinitionLibrary::GetValueDefinition(InValue->GetType()).IsFloatType());
-		check(Throughput->CanInputAcceptType(InInputIdx, EDMValueType::VT_Float1));
-	}
-
-	UDMMaterialLayerObject* Layer = GetLayer();
-	check(Layer);
-
-	UDMMaterialSlot* Slot = Layer->GetSlot();
-	check(Slot);
-
-	UDynamicMaterialModelEditorOnlyData* ModelEditorOnlyData = Slot->GetMaterialModelEditorOnlyData();
-	check(ModelEditorOnlyData);
-
-	UDynamicMaterialModel* MaterialModel = ModelEditorOnlyData->GetMaterialModel();
-	check(MaterialModel);
-	check(MaterialModel == InValue->GetMaterialModel());
-
-	UDMMaterialStageInputValue* InputValue = NewObject<UDMMaterialStageInputValue>(this, NAME_None, RF_Transactional);
-	check(InputValue);
-
-	{
-		const FDMUpdateGuard Guard;
-		InputValue->SetValue(InValue);
-	}
-
-	AddInput(InputValue);
-
-	UpdateInputMap(
-		InInputIdx,
-		FDMMaterialStageConnectorChannel::FIRST_STAGE_INPUT + Inputs.Num() - 1,
-		InInputChannel,
-		0,
-		OutputChannel,
-		EDMMaterialPropertyType::None
-	);
-
-	ApplyWholeLayerValue(InValue);
-
-	return InputValue;
+	return NewInput;
 }
 
-UDMMaterialStageInputValue* UDMMaterialStage::ChangeInput_NewValue(int32 InInputIdx, int32 InInputChannel, EDMValueType InType, int32 InOutputChannel)
+UDMMaterialStageSource* UDMMaterialStage::ChangeInput_PreviousStage(int32 InInputIdx, int32 InInputChannel, EDMMaterialPropertyType InPreviousStageProperty, int32 InOutputIdx, int32 InOutputChannel)
 {
 	check(Source);
 
@@ -952,251 +807,47 @@ UDMMaterialStageInputValue* UDMMaterialStage::ChangeInput_NewValue(int32 InInput
 
 	const TArray<FDMMaterialStageConnector>& InputConnectors = Throughput->GetInputConnectors();
 	check(InputConnectors.IsValidIndex(InInputIdx));
-	check(Throughput->CanInputAcceptType(InInputIdx, InType));
 
 	UDMMaterialLayerObject* Layer = GetLayer();
 	check(Layer);
 
-	UDMMaterialSlot* Slot = Layer->GetSlot();
-	check(Slot);
-
-	UDynamicMaterialModelEditorOnlyData* ModelEditorOnlyData = Slot->GetMaterialModelEditorOnlyData();
-	check(ModelEditorOnlyData);
-
-	UDynamicMaterialModel* MaterialModel = ModelEditorOnlyData->GetMaterialModel();
-	check(MaterialModel);
-
-	UDMMaterialValue* NewValue = MaterialModel->AddValue(InType);
-	check(NewValue);
-
-	UDMMaterialStageInputValue* InputValue = NewObject<UDMMaterialStageInputValue>(this, NAME_None, RF_Transactional);
-	check(InputValue);
+	EDMMaterialPropertyType StageProperty = Layer->GetMaterialProperty();
+	check(StageProperty != EDMMaterialPropertyType::None);
 
+	const UDMMaterialLayerObject* PreviousLayer = Layer->GetPreviousLayer(StageProperty, EDMMaterialLayerStage::Base);
+	UDMMaterialStageSource* PreviousSource = nullptr;
+	
+	if (PreviousLayer)
 	{
-		const FDMUpdateGuard Guard;
-		InputValue->SetValue(NewValue);
+		if (UDMMaterialStage* Stage = PreviousLayer->GetLastEnabledStage(EDMMaterialLayerStage::All))
+		{
+			PreviousSource = Stage->GetSource();
+		}
 	}
 
-	AddInput(InputValue);
-
-	UpdateInputMap(
-		InInputIdx,
-		FDMMaterialStageConnectorChannel::FIRST_STAGE_INPUT + Inputs.Num() - 1,
-		InInputChannel,
-		0,
-		InOutputChannel,
-		EDMMaterialPropertyType::None
-	);
-
-	ApplyWholeLayerValue(NewValue);
-
-	return InputValue;
-}
-
-UDMMaterialStageInputSlot* UDMMaterialStage::ChangeInput_Slot(int32 InInputIdx, int32 InInputChannel, UDMMaterialSlot* InSlot, EDMMaterialPropertyType InProperty, int32 InOutputIdx, int32 InOutputChannel)
-{
-	check(Source);
-
-	UDMMaterialStageThroughput* Throughput = Cast<UDMMaterialStageThroughput>(Source);
-	check(Throughput);
-
-	const TArray<FDMMaterialStageConnector>& InputConnectors = Throughput->GetInputConnectors();
-	check(InputConnectors.IsValidIndex(InInputIdx));
-
-	UDMMaterialLayerObject* Layer = GetLayer();
-	check(Layer);
-
-	UDMMaterialSlot* Slot = Layer->GetSlot();
-	check(Slot);
-
-	UDynamicMaterialModelEditorOnlyData* ModelEditorOnlyData = Slot->GetMaterialModelEditorOnlyData();
-	check(ModelEditorOnlyData == InSlot->GetMaterialModelEditorOnlyData());
-
-	const TArray<EDMMaterialPropertyType> SlotPropertes = ModelEditorOnlyData->GetMaterialPropertiesForSlot(InSlot);
-	check(SlotPropertes.Contains(InProperty));
-
-	const TArray<EDMValueType>& SlotPropertyOutputTypes = InSlot->GetOutputConnectorTypesForMaterialProperty(InProperty);
-	check(SlotPropertyOutputTypes.IsValidIndex(InOutputIdx));
-
-	if (InOutputChannel == FDMMaterialStageConnectorChannel::WHOLE_CHANNEL)
+	if (PreviousSource)
 	{
-		check(Throughput->CanInputAcceptType(InInputIdx, SlotPropertyOutputTypes[InOutputIdx]));
+		const TArray<FDMMaterialStageConnector>& PreviousStageOutputs = PreviousSource->GetOutputConnectors();
+		check(PreviousStageOutputs.IsValidIndex(InOutputIdx));
+		check(Throughput->CanInputConnectTo(InInputIdx, PreviousStageOutputs[InOutputIdx], InOutputChannel));
 	}
 	else
 	{
-		check(UDMValueDefinitionLibrary::GetValueDefinition(SlotPropertyOutputTypes[InOutputIdx]).IsFloatType());
-		check(Throughput->CanInputAcceptType(InInputIdx, EDMValueType::VT_Float1));
-	}
-
-	UDMMaterialStageInputSlot* InputSlot = NewObject<UDMMaterialStageInputSlot>(this, NAME_None, RF_Transactional);
-	check(InputSlot);
-
-	{
-		const FDMUpdateGuard Guard;
-		InputSlot->SetSlot(InSlot);
-		InputSlot->SetMaterialProperty(InProperty);
-	}
-
-	AddInput(InputSlot);
-
-	UpdateInputMap(
-		InInputIdx,
-		FDMMaterialStageConnectorChannel::FIRST_STAGE_INPUT + Inputs.Num() - 1,
-		InInputChannel,
-		InOutputIdx,
-		InOutputChannel,
-		EDMMaterialPropertyType::None
-	);
-
-	return InputSlot;
-}
-
-UDMMaterialStageInputExpression* UDMMaterialStage::ChangeInput_Expression(int32 InInputIdx, int32 InInputChannel, TSubclassOf<UDMMaterialStageExpression> InExpressionClass, int32 InOutputIdx, int32 InOutputChannel)
-{
-	check(Source);
-
-	UDMMaterialStageThroughput* Throughput = Cast<UDMMaterialStageThroughput>(Source);
-	check(Throughput);
-
-	const TArray<FDMMaterialStageConnector>& InputConnectors = Throughput->GetInputConnectors();
-	check(InputConnectors.IsValidIndex(InInputIdx));
-
-	check(InExpressionClass.Get());
-	check(InExpressionClass.Get() != UDMMaterialStageExpression::StaticClass());
-
-	UDMMaterialStageExpression* ExpressionCDO = Cast<UDMMaterialStageExpression>(InExpressionClass->GetDefaultObject(true));
-	check(ExpressionCDO);
-	check(!ExpressionCDO->IsInputRequired() || ExpressionCDO->AllowsNestedInputs());
-
-	const TArray<FDMMaterialStageConnector>& OutputConnectors = ExpressionCDO->GetOutputConnectors();
-	check(Throughput->CanInputConnectTo(InInputIdx, OutputConnectors[InOutputIdx], InOutputChannel));
-
-	UDMMaterialStageInputExpression* InputExpression = NewObject<UDMMaterialStageInputExpression>(this, NAME_None, RF_Transactional);
-	check(InputExpression);
-
-	{
-		const FDMUpdateGuard Guard;
-		InputExpression->SetMaterialStageExpressionClass(InExpressionClass);
+		check(InOutputIdx == 0);
+		check(Throughput->GetInputConnectors().IsEmpty() == false);
+		check(Throughput->GetInputConnectors()[0].IsCompatibleWith(EDMValueType::VT_Float3_RGB));
 	}
 
-	AddInput(InputExpression);
-
 	UpdateInputMap(
 		InInputIdx,
-		FDMMaterialStageConnectorChannel::FIRST_STAGE_INPUT + Inputs.Num() - 1,
+		FDMMaterialStageConnectorChannel::PREVIOUS_STAGE,
 		InInputChannel,
 		InOutputIdx,
 		InOutputChannel,
-		EDMMaterialPropertyType::None
-	);
-
-	return InputExpression;
-}
-
-UDMMaterialStageInputTextureUV* UDMMaterialStage::ChangeInput_UV(int32 InInputIdx, int32 InInputChannel, int32 OutputChannel)
-{
-	UDMMaterialStageInputTextureUV* InputTextureUV = NewObject<UDMMaterialStageInputTextureUV>(this, NAME_None, RF_Transactional);
-	check(InputTextureUV);
-
-	UDMMaterialLayerObject* Layer = GetLayer();
-	check(Layer);
-
-	UDMMaterialSlot* Slot = Layer->GetSlot();
-	check(Slot);
-
-	UDynamicMaterialModelEditorOnlyData* ModelEditorOnlyData = Slot->GetMaterialModelEditorOnlyData();
-	check(ModelEditorOnlyData);
-
-	UDynamicMaterialModel* MaterialModel = ModelEditorOnlyData->GetMaterialModel();
-	check(MaterialModel);
-
-	InputTextureUV->Init(MaterialModel);
-
-	AddInput(InputTextureUV);
-
-	UpdateInputMap(
-		InInputIdx,
-		FDMMaterialStageConnectorChannel::FIRST_STAGE_INPUT + Inputs.Num() - 1,
-		InInputChannel,
-		0, // Texture UVs only have 1 output.
-		OutputChannel,
-		EDMMaterialPropertyType::None
-	);
-
-	return InputTextureUV;
-}
-
-UDMMaterialStageInputGradient* UDMMaterialStage::ChangeInput_Gradient(int32 InInputIdx, int32 InInputChannel, TSubclassOf<UDMMaterialStageGradient> InGradientClass, int32 InOutputChannel)
-{
-	check(Source);
-
-	UDMMaterialStageThroughput* Throughput = Cast<UDMMaterialStageThroughput>(Source);
-	check(Throughput);
-
-	const TArray<FDMMaterialStageConnector>& InputConnectors = Throughput->GetInputConnectors();
-	check(InputConnectors.IsValidIndex(InInputIdx));
-
-	check(InGradientClass.Get());
-	check(InGradientClass.Get() != UDMMaterialStageGradient::StaticClass());
-
-	UDMMaterialStageGradient* GradientCDO = Cast<UDMMaterialStageGradient>(InGradientClass->GetDefaultObject(true));
-	check(GradientCDO);
-
-	const TArray<FDMMaterialStageConnector>& OutputConnectors = GradientCDO->GetOutputConnectors();
-	check(Throughput->CanInputAcceptType(InInputIdx, EDMValueType::VT_Float1));
-
-	UDMMaterialStageInputGradient* InputGradient = NewObject<UDMMaterialStageInputGradient>(this, NAME_None, RF_Transactional);
-	check(InputGradient);
-
-	{
-		const FDMUpdateGuard Guard;
-		InputGradient->SetMaterialStageGradientClass(InGradientClass);
-	}
-
-	AddInput(InputGradient);
-
-	UpdateInputMap(
-		InInputIdx,
-		FDMMaterialStageConnectorChannel::FIRST_STAGE_INPUT + Inputs.Num() - 1,
-		InInputChannel,
-		0,
-		InOutputChannel,
-		EDMMaterialPropertyType::None
-	);
-
-	return InputGradient;
-}
-
-UDMMaterialStageInputFunction* UDMMaterialStage::ChangeInput_MaterialFunction(int32 InInputIdx, int32 InInputChannel)
-{
-	check(Source);
-
-	UDMMaterialStageThroughput* Throughput = Cast<UDMMaterialStageThroughput>(Source);
-	check(Throughput);
-
-	const TArray<FDMMaterialStageConnector>& InputConnectors = Throughput->GetInputConnectors();
-	check(InputConnectors.IsValidIndex(InInputIdx));
-
-	UDMMaterialStageInputFunction* InputMaterialFunction = NewObject<UDMMaterialStageInputFunction>(this, NAME_None, RF_Transactional);
-	check(InputMaterialFunction);
-
-	{
-		const FDMUpdateGuard Guard;
-		InputMaterialFunction->Init();
-	}
-
-	AddInput(InputMaterialFunction);
-
-	UpdateInputMap(
-		InInputIdx,
-		FDMMaterialStageConnectorChannel::FIRST_STAGE_INPUT + Inputs.Num() - 1,
-		InInputChannel,
-		0,
-		FDMMaterialStageConnectorChannel::WHOLE_CHANNEL,
-		EDMMaterialPropertyType::None
+		InPreviousStageProperty
 	);
 
-	return InputMaterialFunction;
+	return PreviousSource;
 }
 
 void UDMMaterialStage::UpdateInputMap(int32 InInputIdx, int32 InSourceIndex, int32 InInputChannel, int32 InOutputIdx, int32 InOutputChannel, EDMMaterialPropertyType InStageProperty)
@@ -1552,245 +1203,27 @@ bool UDMMaterialStage::VerifyInputMap(int32 InInputIdx)
 	return bVerified;
 }
 
-UDMMaterialStageInputValue* UDMMaterialStage::ChangeSource_NewLocalValue(EDMValueType InType)
-{
-	if (!bCanChangeSource)
-	{
-		return nullptr;
-	}
-
-	UDMMaterialLayerObject* Layer = GetLayer();
-	check(Layer);
-
-	UDMMaterialSlot* Slot = Layer->GetSlot();
-	check(Slot);
-
-	UDynamicMaterialModelEditorOnlyData* ModelEditorOnlyData = Slot->GetMaterialModelEditorOnlyData();
-	check(ModelEditorOnlyData);
-
-	UDynamicMaterialModel* MaterialModel = ModelEditorOnlyData->GetMaterialModel();
-	check(MaterialModel);
-
-	// Don't add via the builder, just parent it to the builder.
-	// It won't appear in the global value list.
-	UDMMaterialValue* NewValue = UDMMaterialValue::CreateMaterialValue(MaterialModel, TEXT(""), InType, true);
-	check(NewValue);
-
-	UDMMaterialStageInputValue* InputValue = NewObject<UDMMaterialStageInputValue>(this, NAME_None, RF_Transactional);
-	check(InputValue);
-
-	{
-		const FDMUpdateGuard Guard;
-		InputValue->SetValue(NewValue);
-	}
-
-	SetSource(InputValue);
-
-	return InputValue;
-}
-
-UDMMaterialStageInputValue* UDMMaterialStage::ChangeSource_Value(UDMMaterialValue* InValue)
+UDMMaterialStageSource* UDMMaterialStage::ChangeSource(TSubclassOf<UDMMaterialStageSource> InSourceClass, FSourceInitFunctionPtr InPreInit)
 {
 	if (!bCanChangeSource)
 	{
 		return nullptr;
 	}
 
-	UDMMaterialLayerObject* Layer = GetLayer();
-	check(Layer);
+	check(InSourceClass);
+	check(!(InSourceClass->ClassFlags & (CLASS_Abstract | CLASS_Hidden | CLASS_Deprecated | CLASS_NewerVersionExists)));
 
-	UDMMaterialSlot* Slot = Layer->GetSlot();
-	check(Slot);
-
-	UDynamicMaterialModelEditorOnlyData* ModelEditorOnlyData = Slot->GetMaterialModelEditorOnlyData();
-	check(ModelEditorOnlyData);
-
-	UDynamicMaterialModel* MaterialModel = ModelEditorOnlyData->GetMaterialModel();
-	check(MaterialModel);
-
-	UDMMaterialStageInputValue* InputValue = NewObject<UDMMaterialStageInputValue>(this, NAME_None, RF_Transactional);
-	check(InputValue);
+	UDMMaterialStageSource* NewSource = NewObject<UDMMaterialStageExpression>(this, InSourceClass, NAME_None, RF_Transactional);
+	check(NewSource);
 
+	if (InPreInit)
 	{
-		const FDMUpdateGuard Guard;
-
-		check(InputValue);
-		InputValue->SetValue(InValue);
+		InPreInit(this, NewSource);
 	}
 
-	SetSource(InputValue);
+	SetSource(NewSource);
 
-	return InputValue;
-}
-
-UDMMaterialStageInputValue* UDMMaterialStage::ChangeSource_NewValue(EDMValueType InType)
-{
-	if (!bCanChangeSource)
-	{
-		return nullptr;
-	}
-
-	UDMMaterialLayerObject* Layer = GetLayer();
-	check(Layer);
-
-	UDMMaterialSlot* Slot = Layer->GetSlot();
-	check(Slot);
-
-	UDynamicMaterialModelEditorOnlyData* ModelEditorOnlyData = Slot->GetMaterialModelEditorOnlyData();
-	check(ModelEditorOnlyData);
-
-	UDynamicMaterialModel* MaterialModel = ModelEditorOnlyData->GetMaterialModel();
-	check(MaterialModel);
-
-	UDMMaterialValue* NewValue = MaterialModel->AddValue(InType);
-	check(NewValue);
-
-	UDMMaterialStageInputValue* InputValue = NewObject<UDMMaterialStageInputValue>(this, NAME_None, RF_Transactional);
-	check(InputValue);
-
-	{
-		const FDMUpdateGuard Guard;
-		InputValue->SetValue(NewValue);
-	}
-
-	SetSource(InputValue);
-
-	return InputValue;
-}
-
-UDMMaterialStageInputSlot* UDMMaterialStage::ChangeSource_Slot(UDMMaterialSlot* InSlot, EDMMaterialPropertyType InProperty)
-{
-	if (!bCanChangeSource)
-	{
-		return nullptr;
-	}
-
-	UDMMaterialLayerObject* Layer = GetLayer();
-	check(Layer);
-
-	UDMMaterialSlot* Slot = Layer->GetSlot();
-	check(Slot);
-
-	check(InSlot);
-	check(Slot != InSlot);
-	check(InSlot->GetMaterialModelEditorOnlyData() == Slot->GetMaterialModelEditorOnlyData());
-
-	UDynamicMaterialModelEditorOnlyData* ModelEditorOnlyData = Slot->GetMaterialModelEditorOnlyData();
-	check(ModelEditorOnlyData);
-
-	TArray<EDMMaterialPropertyType> SlotPropertes = ModelEditorOnlyData->GetMaterialPropertiesForSlot(InSlot);
-	check(SlotPropertes.Contains(InProperty));
-
-	UDMMaterialStageInputSlot* InputSlot = NewObject<UDMMaterialStageInputSlot>(this, NAME_None, RF_Transactional);
-	check(InputSlot);
-
-	{
-		const FDMUpdateGuard Guard;
-		InputSlot->SetSlot(InSlot);
-		InputSlot->SetMaterialProperty(InProperty);
-	}
-
-	SetSource(InputSlot);
-
-	return InputSlot;
-}
-
-UDMMaterialStageExpression* UDMMaterialStage::ChangeSource_Expression(TSubclassOf<UDMMaterialStageExpression> InExpressionClass)
-{
-	if (!bCanChangeSource)
-	{
-		return nullptr;
-	}
-
-	check(InExpressionClass);
-	check(InExpressionClass != UDMMaterialStageExpression::StaticClass());
-
-	UDMMaterialStageExpression* NewExpression = NewObject<UDMMaterialStageExpression>(this, InExpressionClass, NAME_None, RF_Transactional);
-	check(NewExpression);
-
-	SetSource(NewExpression);
-
-	return NewExpression;
-}
-
-UDMMaterialStageBlend* UDMMaterialStage::ChangeSource_Blend(TSubclassOf<UDMMaterialStageBlend> InBlendClass)
-{
-	if (!bCanChangeSource)
-	{
-		return nullptr;
-	}
-
-	check(InBlendClass);
-	check(InBlendClass != UDMMaterialStageBlend::StaticClass());
-
-	UDMMaterialStageBlend* NewBlend = NewObject<UDMMaterialStageBlend>(this, InBlendClass, NAME_None, RF_Transactional);
-	check(NewBlend);
-
-	SetSource(NewBlend);
-	
-	return NewBlend;
-}
-
-UDMMaterialStageGradient* UDMMaterialStage::ChangeSource_Gradient(TSubclassOf<UDMMaterialStageGradient> InGradientClass)
-{
-	if (!bCanChangeSource)
-	{
-		return nullptr;
-	}
-
-	check(InGradientClass);
-	check(InGradientClass != UDMMaterialStageGradient::StaticClass());
-
-	UDMMaterialStageGradient* NewGradient = NewObject<UDMMaterialStageGradient>(this, InGradientClass, NAME_None, RF_Transactional);
-	check(NewGradient);
-
-	SetSource(NewGradient);
-
-	return NewGradient;
-}
-
-UDMMaterialStageInputTextureUV* UDMMaterialStage::ChangeSource_UV(bool bInDoUpdate)
-{
-	if (!bCanChangeSource)
-	{
-		return nullptr;
-	}
-
-	UDMMaterialStageInputTextureUV* InputTextureUV = NewObject<UDMMaterialStageInputTextureUV>(this, NAME_None, RF_Transactional);
-	check(InputTextureUV);
-
-	UDMMaterialLayerObject* Layer = GetLayer();
-	check(Layer);
-
-	UDMMaterialSlot* Slot = Layer->GetSlot();
-	check(Slot);
-
-	UDynamicMaterialModelEditorOnlyData* ModelEditorOnlyData = Slot->GetMaterialModelEditorOnlyData();
-	check(ModelEditorOnlyData);
-
-	UDynamicMaterialModel* MaterialModel = ModelEditorOnlyData->GetMaterialModel();
-	check(MaterialModel);
-
-	InputTextureUV->Init(MaterialModel);
-
-	SetSource(InputTextureUV);
-
-	return InputTextureUV;
-}
-
-UDMMaterialStageFunction* UDMMaterialStage::ChangeSource_MaterialFunction()
-{
-	if (!bCanChangeSource)
-	{
-		return nullptr;
-	}
-
-	UDMMaterialStageFunction* NewMaterialFunction = NewObject<UDMMaterialStageFunction>(this, UDMMaterialStageFunction::StaticClass(), NAME_None, RF_Transactional);
-	check(NewMaterialFunction);
-
-	SetSource(NewMaterialFunction);
-
-	return NewMaterialFunction;
+	return NewSource;
 }
 
 bool UDMMaterialStage::IsCompatibleWithPreviousStage(const UDMMaterialStage* InPreviousStage) const
@@ -1871,6 +1304,28 @@ void UDMMaterialStage::RemoveInput(UDMMaterialStageInput* InInput)
 	Update(EDMUpdateType::Structure);
 }
 
+void UDMMaterialStage::RemoveAllInputs()
+{
+	if (Inputs.IsEmpty())
+	{
+		return;
+	}
+
+	for (UDMMaterialStageInput* Input : Inputs)
+	{
+		if (GUndo)
+		{
+			Input->Modify();
+		}
+
+		Input->SetComponentState(EDMComponentLifetimeState::Removed);
+	}
+
+	Inputs.Empty();
+
+	Update(EDMUpdateType::Structure);
+}
+
 UMaterialInterface* UDMMaterialStage::GetPreviewMaterial()
 {
 	if (!PreviewMaterialBase)
diff --git a/Engine/Plugins/Experimental/DynamicMaterial/Source/DynamicMaterialEditor/Private/Components/DMMaterialStageBlend.cpp b/Engine/Plugins/Experimental/DynamicMaterial/Source/DynamicMaterialEditor/Private/Components/DMMaterialStageBlend.cpp
--- a/Engine/Plugins/Experimental/DynamicMaterial/Source/DynamicMaterialEditor/Private/Components/DMMaterialStageBlend.cpp
+++ b/Engine/Plugins/Experimental/DynamicMaterial/Source/DynamicMaterialEditor/Private/Components/DMMaterialStageBlend.cpp
@@ -6,6 +6,7 @@
 #include "Components/DMMaterialStage.h"
 #include "Components/DMMaterialStageInput.h"
 #include "Components/MaterialStageExpressions/DMMSETextureSample.h"
+#include "Components/MaterialStageInputs/DMMSIExpression.h"
 #include "Components/MaterialStageInputs/DMMSITextureUV.h"
 #include "Components/MaterialStageInputs/DMMSIValue.h"
 #include "Components/MaterialValues/DMMaterialValueFloat1.h"
@@ -49,7 +50,14 @@ UDMMaterialStage* UDMMaterialStageBlend::CreateStage(TSubclassOf<UDMMaterialStag
 	const FDMUpdateGuard Guard;
  
 	UDMMaterialStage* NewStage = UDMMaterialStage::CreateMaterialStage(InLayer);
-	UDMMaterialStageBlend* SourceBlend = NewObject<UDMMaterialStageBlend>(NewStage, InMaterialStageBlendClass.Get(), NAME_None, RF_Transactional);
+
+	UDMMaterialStageBlend* SourceBlend = NewObject<UDMMaterialStageBlend>(
+		NewStage, 
+		InMaterialStageBlendClass.Get(), 
+		NAME_None, 
+		RF_Transactional
+	);
+
 	check(SourceBlend);
  
 	NewStage->SetSource(SourceBlend);
@@ -103,7 +111,8 @@ bool UDMMaterialStageBlend::CanInputAcceptType(int32 InputIndex, EDMValueType Va
 	return InputConnectors[InputIndex].IsCompatibleWith(ValueType);
 }
  
-void UDMMaterialStageBlend::GetMaskAlphaBlendNode(const TSharedRef<FDMMaterialBuildState>& InBuildState, UMaterialExpression*& OutExpression, int32& OutOutputIndex, int32& OutOutputChannel) const
+void UDMMaterialStageBlend::GetMaskAlphaBlendNode(const TSharedRef<FDMMaterialBuildState>& InBuildState, UMaterialExpression*& OutExpression, 
+	int32& OutOutputIndex, int32& OutOutputChannel) const
 {
 	FDMMaterialStageConnectorChannel Channel;
 	TArray<UMaterialExpression*> Expressions;
@@ -205,9 +214,14 @@ void UDMMaterialStageBlend::AddDefaultInput(int32 InInputIndex) const
 	{
 		case InputAlpha:
 		{
-			UDMMaterialStageInputValue* InputValue = Cast<UDMMaterialStageInputValue>(
-				Stage->ChangeInput_NewLocalValue(InInputIndex, FDMMaterialStageConnectorChannel::WHOLE_CHANNEL, EDMValueType::VT_Float1, FDMMaterialStageConnectorChannel::WHOLE_CHANNEL)
+			UDMMaterialStageInputValue* InputValue = UDMMaterialStageInputValue::ChangeStageInput_NewLocalValue(
+				Stage, 
+				InInputIndex, 
+				FDMMaterialStageConnectorChannel::WHOLE_CHANNEL,
+				EDMValueType::VT_Float1, 
+				FDMMaterialStageConnectorChannel::WHOLE_CHANNEL
 			);
+
 			check(InputValue);
  
 			UDMMaterialValueFloat1* Float1Value = Cast<UDMMaterialValueFloat1>(InputValue->GetValue());
@@ -231,21 +245,37 @@ void UDMMaterialStageBlend::AddDefaultInput(int32 InInputIndex) const
  
 			if (PreviousLayer)
 			{
-				Stage->ChangeInput_PreviousStage(InInputIndex, FDMMaterialStageConnectorChannel::WHOLE_CHANNEL, StageProperty,
-					0, FDMMaterialStageConnectorChannel::WHOLE_CHANNEL);
+				Stage->ChangeInput_PreviousStage(
+					InInputIndex, 
+					FDMMaterialStageConnectorChannel::WHOLE_CHANNEL, 
+					StageProperty,
+					0, 
+					FDMMaterialStageConnectorChannel::WHOLE_CHANNEL
+				);
 				break;
 			}
 			
-			Stage->ChangeInput_PreviousStage(InInputIndex, FDMMaterialStageConnectorChannel::WHOLE_CHANNEL, EDMMaterialPropertyType::EmissiveColor,
-				0, FDMMaterialStageConnectorChannel::WHOLE_CHANNEL);
+			Stage->ChangeInput_PreviousStage(
+				InInputIndex, 
+				FDMMaterialStageConnectorChannel::WHOLE_CHANNEL, 
+				EDMMaterialPropertyType::EmissiveColor,
+				0, 
+				FDMMaterialStageConnectorChannel::WHOLE_CHANNEL
+			);
 			break;
 		}
  
 		case InputB:
 			if (CanInputAcceptType(InInputIndex, EDMValueType::VT_Float3_RGB))
 			{
-				Stage->ChangeInput_Expression(InInputIndex, FDMMaterialStageConnectorChannel::WHOLE_CHANNEL,
-					UDMMaterialStageExpressionTextureSample::StaticClass(), 0, FDMMaterialStageConnectorChannel::WHOLE_CHANNEL);
+				UDMMaterialStageInputExpression::ChangeStageInput_Expression(
+					Stage,
+					UDMMaterialStageExpressionTextureSample::StaticClass(), 
+					InInputIndex, 
+					FDMMaterialStageConnectorChannel::WHOLE_CHANNEL, 
+					0,
+					FDMMaterialStageConnectorChannel::WHOLE_CHANNEL
+				);
 			}
 			break;
  
@@ -285,7 +315,13 @@ int32 UDMMaterialStageBlend::ResolveInput(const TSharedRef<FDMMaterialBuildState
 		{
 			if (UDMMaterialLayerObject* Layer = Stage->GetLayer())
 			{
-				Layer->ApplyEffects(InBuildState, Stage, OutExpressions, OutChannel.OutputChannel, NodeOutputIndex);
+				Layer->ApplyEffects(
+					InBuildState, 
+					Stage, 
+					OutExpressions, 
+					OutChannel.OutputChannel, 
+					NodeOutputIndex
+				);
 			}
 		}
 	}
@@ -331,7 +367,13 @@ FDMExpressionInput UDMMaterialStageBlend::GetLayerMaskLinkTextureUVInputExpressi
  
 		int32 StageInputIdx = InputConnections[InputIdx].Channels[0].SourceIndex - FDMMaterialStageConnectorChannel::FIRST_STAGE_INPUT;
  
-		ExpressionInput.OutputIndex = ResolveLayerMaskTextureUVLinkInputImpl(InBuildState, StageInputs[StageInputIdx], Channel, ExpressionInput.OutputExpressions);
+		ExpressionInput.OutputIndex = ResolveLayerMaskTextureUVLinkInputImpl(
+			InBuildState, 
+			StageInputs[StageInputIdx], 
+			Channel, 
+			ExpressionInput.OutputExpressions
+		);
+		
 		ExpressionInput.OutputChannel = Channel.OutputChannel;
 		
 		if (ExpressionInput.IsValid())
@@ -438,11 +480,20 @@ void UDMMaterialStageBlend::UpdatePreviewMaterial(UMaterial* InPreviewMaterial)
 		Stage->GenerateExpressions(BuildState);
 		UMaterialExpression* StageExpression = BuildState->GetLastStageExpression(Stage);
  
-		BuildState->GetBuildUtils().UpdatePreviewMaterial(StageExpression, 0, FDMMaterialStageConnectorChannel::WHOLE_CHANNEL, 32);
+		BuildState->GetBuildUtils().UpdatePreviewMaterial(
+			StageExpression, 
+			0, 
+			FDMMaterialStageConnectorChannel::WHOLE_CHANNEL, 
+			32
+		);
 	}
 	else
 	{
-		UE::DynamicMaterialEditor::Private::BuildExpressionInputs(BuildState, InputConnectionMap, Inputs);
+		UE::DynamicMaterialEditor::Private::BuildExpressionInputs(
+			BuildState, 
+			InputConnectionMap, 
+			Inputs
+		);
 	}
 }
 
diff --git a/Engine/Plugins/Experimental/DynamicMaterial/Source/DynamicMaterialEditor/Private/Components/DMMaterialStageExpression.cpp b/Engine/Plugins/Experimental/DynamicMaterial/Source/DynamicMaterialEditor/Private/Components/DMMaterialStageExpression.cpp
--- a/Engine/Plugins/Experimental/DynamicMaterial/Source/DynamicMaterialEditor/Private/Components/DMMaterialStageExpression.cpp
+++ b/Engine/Plugins/Experimental/DynamicMaterial/Source/DynamicMaterialEditor/Private/Components/DMMaterialStageExpression.cpp
@@ -59,7 +59,14 @@ UDMMaterialStage* UDMMaterialStageExpression::CreateStage(TSubclassOf<UDMMateria
 	const FDMUpdateGuard Guard;
 
 	UDMMaterialStage* NewStage = UDMMaterialStage::CreateMaterialStage(InLayer);
-	UDMMaterialStageExpression* SourceExpression = NewObject<UDMMaterialStageExpression>(NewStage, InMaterialStageExpressionClass.Get(), NAME_None, RF_Transactional);
+
+	UDMMaterialStageExpression* SourceExpression = NewObject<UDMMaterialStageExpression>(
+		NewStage, 
+		InMaterialStageExpressionClass.Get(), 
+		NAME_None, 
+		RF_Transactional
+	);
+	
 	check(SourceExpression);
 
 	NewStage->SetSource(SourceExpression);
@@ -82,6 +89,22 @@ const TArray<TStrongObjectPtr<UClass>>& UDMMaterialStageExpression::GetAvailable
 	return SourceExpressions;
 }
 
+UDMMaterialStageExpression* UDMMaterialStageExpression::ChangeStageSource_Expression(UDMMaterialStage* InStage, 
+	TSubclassOf<UDMMaterialStageExpression> InExpressionClass)
+{
+	check(InStage);
+
+	if (!InStage->CanChangeSource())
+	{
+		return nullptr;
+	}
+
+	check(InExpressionClass);
+	check(!(InExpressionClass->ClassFlags & (CLASS_Abstract | CLASS_Hidden | CLASS_Deprecated | CLASS_NewerVersionExists)));
+
+	return InStage->ChangeSource<UDMMaterialStageExpression>(InExpressionClass);
+}
+
 void UDMMaterialStageExpression::GenerateExpressionList()
 {
 	SourceExpressions.Empty();
diff --git a/Engine/Plugins/Experimental/DynamicMaterial/Source/DynamicMaterialEditor/Private/Components/DMMaterialStageFunction.cpp b/Engine/Plugins/Experimental/DynamicMaterial/Source/DynamicMaterialEditor/Private/Components/DMMaterialStageFunction.cpp
--- a/Engine/Plugins/Experimental/DynamicMaterial/Source/DynamicMaterialEditor/Private/Components/DMMaterialStageFunction.cpp
+++ b/Engine/Plugins/Experimental/DynamicMaterial/Source/DynamicMaterialEditor/Private/Components/DMMaterialStageFunction.cpp
@@ -1,13 +1,24 @@
 // Copyright Epic Games, Inc. All Rights Reserved.
 
 #include "Components/DMMaterialStageFunction.h"
+
 #include "Components/DMMaterialLayer.h"
 #include "Components/DMMaterialSlot.h"
 #include "Components/DMMaterialStage.h"
+#include "Components/DMMaterialValue.h"
+#include "Components/MaterialStageInputs/DMMSIValue.h"
+#include "Components/MaterialValues/DMMaterialValueFloat1.h"
+#include "Components/MaterialValues/DMMaterialValueFloat2.h"
+#include "Components/MaterialValues/DMMaterialValueFloat3XYZ.h"
+#include "Components/MaterialValues/DMMaterialValueFloat4.h"
+#include "DynamicMaterialEditorModule.h"
+#include "Materials/MaterialExpressionFunctionInput.h"
+#include "Materials/MaterialExpressionFunctionOutput.h"
 #include "Materials/MaterialExpressionMaterialFunctionCall.h"
 #include "Materials/MaterialFunctionInterface.h"
 #include "Model/DMMaterialBuildState.h"
 #include "Model/DMMaterialBuildUtils.h"
+#include "Model/DynamicMaterialModelEditorOnlyData.h"
 
 #define LOCTEXT_NAMESPACE "DMMaterialStageFunction"
 
@@ -20,7 +31,14 @@ UDMMaterialStage* UDMMaterialStageFunction::CreateStage(UDMMaterialLayerObject*
 	const FDMUpdateGuard Guard;
 
 	UDMMaterialStage* NewStage = UDMMaterialStage::CreateMaterialStage(InLayer);
-	UDMMaterialStageFunction* SourceFunction = NewObject<UDMMaterialStageFunction>(NewStage, StaticClass(), NAME_None, RF_Transactional);
+
+	UDMMaterialStageFunction* SourceFunction = NewObject<UDMMaterialStageFunction>(
+		NewStage, 
+		StaticClass(), 
+		NAME_None, 
+		RF_Transactional
+	);
+
 	check(SourceFunction);
 
 	NewStage->SetSource(SourceFunction);
@@ -28,6 +46,32 @@ UDMMaterialStage* UDMMaterialStageFunction::CreateStage(UDMMaterialLayerObject*
 	return NewStage;
 }
 
+UDMMaterialStageFunction* UDMMaterialStageFunction::ChangeStageSource_Function(UDMMaterialStage* InStage, UMaterialFunctionInterface* InMaterialFunction)
+{
+	check(InStage);
+
+	if (!InStage->CanChangeSource())
+	{
+		return nullptr;
+	}
+
+	check(InMaterialFunction);
+
+	UDMMaterialStageFunction* NewFunction = InStage->ChangeSource<UDMMaterialStageFunction>(
+		[InMaterialFunction](UDMMaterialStage* InStage, UDMMaterialStageSource* InNewSource)
+		{
+			const FDMUpdateGuard Guard;
+			CastChecked<UDMMaterialStageFunction>(InNewSource)->SetMaterialFunction(InMaterialFunction);
+		});
+
+	return NewFunction;
+}
+
+UMaterialFunctionInterface* UDMMaterialStageFunction::GetNoOpFunction()
+{
+	return NoOp.LoadSynchronous();
+}
+
 void UDMMaterialStageFunction::SetMaterialFunction(UMaterialFunctionInterface* InMaterialFunction)
 {
 	if (MaterialFunction == InMaterialFunction)
@@ -40,6 +84,43 @@ void UDMMaterialStageFunction::SetMaterialFunction(UMaterialFunctionInterface* I
 	OnMaterialFunctionChanged();
 }
 
+UDMMaterialValue* UDMMaterialStageFunction::GetInputValue(int32 InIndex) const
+{
+	TArray<UDMMaterialValue*> Values = GetInputValues();
+
+	if (Values.IsValidIndex(InIndex))
+	{
+		return Values[InIndex];
+	}
+
+	return nullptr;
+}
+
+TArray<UDMMaterialValue*> UDMMaterialStageFunction::GetInputValues() const
+{
+	UDMMaterialStage* Stage = GetStage();
+
+	if (!Stage)
+	{
+		return {};
+	}
+
+	TArray<UDMMaterialValue*> Values;
+
+	for (UDMMaterialStageInput* Input : Stage->GetInputs())
+	{
+		if (UDMMaterialStageInputValue* InputValue = Cast<UDMMaterialStageInputValue>(Input))
+		{
+			if (UDMMaterialValue* Value = InputValue->GetValue())
+			{
+				Values.Add(Value);
+			}
+		}
+	}
+
+	return Values;
+}
+
 void UDMMaterialStageFunction::AddDefaultInput(int32 InInputIndex) const
 {
 	check(InInputIndex == 0);
@@ -57,19 +138,29 @@ void UDMMaterialStageFunction::AddDefaultInput(int32 InInputIndex) const
 
 	if (PreviousLayer)
 	{
-		Stage->ChangeInput_PreviousStage(InInputIndex, FDMMaterialStageConnectorChannel::WHOLE_CHANNEL, StageProperty,
-			0, FDMMaterialStageConnectorChannel::WHOLE_CHANNEL);
+		Stage->ChangeInput_PreviousStage(
+			InInputIndex, 
+			FDMMaterialStageConnectorChannel::WHOLE_CHANNEL, 
+			StageProperty,
+			0, 
+			FDMMaterialStageConnectorChannel::WHOLE_CHANNEL
+		);
 	}
 	else
 	{
-		Stage->ChangeInput_PreviousStage(InInputIndex, FDMMaterialStageConnectorChannel::WHOLE_CHANNEL, EDMMaterialPropertyType::EmissiveColor,
-			0, FDMMaterialStageConnectorChannel::WHOLE_CHANNEL);
+		Stage->ChangeInput_PreviousStage(
+			InInputIndex, 
+			FDMMaterialStageConnectorChannel::WHOLE_CHANNEL, 
+			EDMMaterialPropertyType::EmissiveColor,
+			0, 
+			FDMMaterialStageConnectorChannel::WHOLE_CHANNEL
+		);
 	}
 }
 
 bool UDMMaterialStageFunction::CanChangeInput(int32 InputIndex) const
 {
-	return false;
+	return true;
 }
 
 bool UDMMaterialStageFunction::CanChangeInputType(int32 InputIndex) const
@@ -79,7 +170,7 @@ bool UDMMaterialStageFunction::CanChangeInputType(int32 InputIndex) const
 
 bool UDMMaterialStageFunction::IsInputVisible(int32 InputIndex) const
 {
-	return false;
+	return true;
 }
 
 void UDMMaterialStageFunction::GenerateExpressions(const TSharedRef<FDMMaterialBuildState>& InBuildState) const
@@ -113,6 +204,21 @@ void UDMMaterialStageFunction::GenerateExpressions(const TSharedRef<FDMMaterialB
 	InBuildState->AddStageSourceExpressions(this, {FunctionCall});
 }
 
+FText UDMMaterialStageFunction::GetComponentDescription() const
+{
+	if (UMaterialFunctionInterface* MaterialFunctionInterface = MaterialFunction.Get())
+	{
+		const FString Caption = MaterialFunctionInterface->GetUserExposedCaption();
+
+		if (!Caption.IsEmpty())
+		{
+			return FText::FromString(Caption);
+		}
+	}
+
+	return Super::GetComponentDescription();
+}
+
 void UDMMaterialStageFunction::PreEditChange(FEditPropertyChain& PropertyAboutToChange)
 {
 	Super::PreEditChange(PropertyAboutToChange);
@@ -130,6 +236,16 @@ void UDMMaterialStageFunction::PostEditChangeProperty(FPropertyChangedEvent& Pro
 	}
 }
 
+void UDMMaterialStageFunction::PostLoad()
+{
+	Super::PostLoad();
+
+	if (NeedsFunctionInit())
+	{
+		InitFunction();
+	}
+}
+
 UDMMaterialStageFunction::UDMMaterialStageFunction()
 	: UDMMaterialStageThroughput(LOCTEXT("MaterialFunction", "Material Function"))
 {
@@ -147,7 +263,361 @@ UDMMaterialStageFunction::UDMMaterialStageFunction()
 
 void UDMMaterialStageFunction::OnMaterialFunctionChanged()
 {
-	Update(EDMUpdateType::Structure);
+	DeinitFunction();
+	InitFunction();
+
+	if (FDMUpdateGuard::CanUpdate())
+	{
+		Update(EDMUpdateType::Structure);
+	}
+}
+
+bool UDMMaterialStageFunction::NeedsFunctionInit() const
+{
+	UMaterialFunctionInterface* MaterialFunctionInterface = MaterialFunction.Get();
+
+	TArray<UDMMaterialValue*> InputValues = GetInputValues();
+
+	if (!IsValid(MaterialFunctionInterface))
+	{
+		// If we have no function, but we do have inputs, they need to be refreshed (removed).
+		return !InputValues.IsEmpty();
+	}
+
+	TArray<FFunctionExpressionInput> Inputs;
+	TArray<FFunctionExpressionOutput> Outputs;
+	MaterialFunctionInterface->GetInputsAndOutputs(Inputs, Outputs);
+
+	if (Outputs.IsEmpty())
+	{
+		UE_LOG(LogDynamicMaterialEditor, Error, TEXT("Function must have at least one output."));
+		return false;
+	}
+
+	if (!Inputs[0].ExpressionInput || !Outputs[0].ExpressionOutput
+		|| Inputs[0].ExpressionInput->InputType != Outputs[0].ExpressionOutput->GetOutputType(0))
+	{
+		UE_LOG(LogDynamicMaterialEditor, Error, TEXT("Function's first input must match its first output."));
+		return false;
+	}
+
+	if (Inputs.Num() != InputValues.Num())
+	{
+		return true;
+	}
+
+	for (int32 InputIndex = 0; InputIndex < Inputs.Num(); ++InputIndex)
+	{
+		UMaterialExpressionFunctionInput* FunctionInput = Inputs[InputIndex].ExpressionInput.Get();
+
+		if (!IsValid(FunctionInput))
+		{
+			UE_LOG(LogDynamicMaterialEditor, Error, TEXT("Function has missing input object."));
+			return false;
+		}
+
+		// First input must be a scalar or vector3.
+		if (InputIndex == 0)
+		{
+			switch (FunctionInput->InputType)
+			{
+				case EFunctionInputType::FunctionInput_Scalar:
+				case EFunctionInputType::FunctionInput_Vector3:
+					break;
+
+				default:
+					UE_LOG(LogDynamicMaterialEditor, Error, TEXT("Function has invalid first input - must be a scalar or vector3."));
+					return false;
+			}
+
+			continue;
+		}
+
+		EDMValueType ValueType = EDMValueType::VT_None;
+
+		switch (FunctionInput->InputType)
+		{
+			case EFunctionInputType::FunctionInput_Scalar:
+				ValueType = EDMValueType::VT_Float1;
+				break;
+
+			case EFunctionInputType::FunctionInput_Vector2:
+				ValueType = EDMValueType::VT_Float2;
+				break;
+
+			case EFunctionInputType::FunctionInput_Vector3:
+				ValueType = EDMValueType::VT_Float3_XYZ;
+				break;
+
+			case EFunctionInputType::FunctionInput_Vector4:
+				ValueType = EDMValueType::VT_Float4_RGBA;
+				break;
+
+			case EFunctionInputType::FunctionInput_Texture2D:
+			case EFunctionInputType::FunctionInput_TextureCube:
+			case EFunctionInputType::FunctionInput_VolumeTexture:
+				ValueType = EDMValueType::VT_Texture;
+				break;
+
+			default:
+				UE_LOG(LogDynamicMaterialEditor, Error, TEXT("Function has invalid input type - must be a scalar, vector or texture."));
+				return false;
+		}
+
+		if (ValueType != InputValues[InputIndex]->GetType())
+		{
+			return true;
+		}
+	}
+
+	return false;
+}
+
+void UDMMaterialStageFunction::InitFunction()
+{
+	if (!IsValid(MaterialFunction.Get()))
+	{
+		return;
+	}
+
+	UDMMaterialStage* Stage = GetStage();
+
+	if (!Stage)
+	{
+		MaterialFunction = nullptr;
+		return;
+	}
+
+	UDMMaterialSlot* Slot = GetTypedParent<UDMMaterialSlot>(/* bAllowSubclasses */ true);
+
+	if (!Slot)
+	{
+		MaterialFunction = nullptr;
+		return;
+	}
+
+	UDynamicMaterialModelEditorOnlyData* EditorOnlyData = Slot->GetMaterialModelEditorOnlyData();
+
+	if (!EditorOnlyData)
+	{
+		MaterialFunction = nullptr;
+		return;
+	}
+
+	UDynamicMaterialModel* MaterialModel = EditorOnlyData->GetMaterialModel();
+
+	if (!MaterialModel)
+	{
+		MaterialFunction = nullptr;
+		return;
+	}
+
+	TArray<FFunctionExpressionInput> Inputs;
+	TArray<FFunctionExpressionOutput> Outputs;
+	MaterialFunction->GetInputsAndOutputs(Inputs, Outputs);
+
+	if (Inputs.IsEmpty())
+	{
+		UE_LOG(LogDynamicMaterialEditor, Error, TEXT("Function must have at least one input."));
+		MaterialFunction = nullptr;
+		return;
+	}
+
+	if (Outputs.IsEmpty())
+	{
+		UE_LOG(LogDynamicMaterialEditor, Error, TEXT("Function must have at least one output."));
+		MaterialFunction = nullptr;
+		return;
+	}
+
+	UDMMaterialLayerObject* Layer = Stage->GetLayer();
+	check(Layer);
+
+	const EDMMaterialPropertyType StageProperty = Layer->GetMaterialProperty();
+
+	{
+		FDMMaterialStageConnector PreviousStageConnector = InputConnectors[0];
+		InputConnectors.SetNumZeroed(Inputs.Num());
+		InputConnectors[0] = PreviousStageConnector;
+	}
+
+	for (int32 InputIndex = 0; InputIndex < Inputs.Num(); ++InputIndex)
+	{
+		UMaterialExpressionFunctionInput* FunctionInput = Inputs[InputIndex].ExpressionInput.Get();
+
+		if (!IsValid(FunctionInput))
+		{
+			UE_LOG(LogDynamicMaterialEditor, Error, TEXT("Function has missing input object."));
+			InputConnectors.SetNum(1);
+			MaterialFunction = nullptr;
+			return;
+		}
+
+		// First input must be a scalar or vector3.
+		if (InputIndex == InputPreviousStage)
+		{
+			switch (FunctionInput->InputType)
+			{
+				case EFunctionInputType::FunctionInput_Scalar:
+				case EFunctionInputType::FunctionInput_Vector3:
+					break;
+
+				default:
+					UE_LOG(LogDynamicMaterialEditor, Error, TEXT("Function has invalid first input - must be a scalar or vector."));
+					InputConnectors.SetNum(1);
+					MaterialFunction = nullptr;
+					return;
+			}
+
+			Stage->ChangeInput_PreviousStage(
+				InputPreviousStage, 
+				FDMMaterialStageConnectorChannel::WHOLE_CHANNEL, 
+				StageProperty,
+				0, 
+				FDMMaterialStageConnectorChannel::WHOLE_CHANNEL
+			);
+
+			continue;
+		}
+
+		EDMValueType ValueType = EDMValueType::VT_None;
+
+		switch (FunctionInput->InputType)
+		{
+			case EFunctionInputType::FunctionInput_Scalar:
+				ValueType = EDMValueType::VT_Float1;
+				break;
+
+			case EFunctionInputType::FunctionInput_Vector2:
+				ValueType = EDMValueType::VT_Float2;
+				break;
+
+			case EFunctionInputType::FunctionInput_Vector3:
+				ValueType = EDMValueType::VT_Float3_XYZ;
+				break;
+
+			case EFunctionInputType::FunctionInput_Vector4:
+				ValueType = EDMValueType::VT_Float4_RGBA;
+				break;
+
+			case EFunctionInputType::FunctionInput_Texture2D:
+			case EFunctionInputType::FunctionInput_TextureCube:
+			case EFunctionInputType::FunctionInput_VolumeTexture:
+				ValueType = EDMValueType::VT_Texture;
+				break;
+
+			default:
+				UE_LOG(LogDynamicMaterialEditor, Error, TEXT("Function has invalid input type - must be a scalar, vector or texture."));
+				InputConnectors.SetNum(1);
+				MaterialFunction = nullptr;
+				return;
+		}
+
+		InputConnectors[InputIndex].Index = InputIndex;
+		InputConnectors[InputIndex].Type = ValueType;
+
+		if (Inputs[InputIndex].ExpressionInput->InputName == NAME_None)
+		{
+			static const FText InputNameFormat = LOCTEXT("InputFormat", "Input {0}");
+			InputConnectors[InputIndex].Name = FText::Format(InputNameFormat, FText::AsNumber(InputIndex + 1));
+		}
+		else
+		{
+			InputConnectors[InputIndex].Name = FText::FromName(Inputs[InputIndex].ExpressionInput->InputName);
+		}
+
+		UDMMaterialStageInputValue* InputValue = UDMMaterialStageInputValue::ChangeStageInput_NewLocalValue(
+			Stage,
+			InputIndex, 
+			FDMMaterialStageConnectorChannel::WHOLE_CHANNEL, 
+			ValueType,
+			FDMMaterialStageConnectorChannel::WHOLE_CHANNEL
+		);
+		
+		check(InputValue);
+
+		UDMMaterialValue* Value = InputValue->GetValue();
+		check(Value);
+
+		if (FunctionInput->bUsePreviewValueAsDefault)
+		{
+			switch (FunctionInput->InputType)
+			{
+				case EFunctionInputType::FunctionInput_Scalar:
+					if (UDMMaterialValueFloat1* Float1Value = Cast<UDMMaterialValueFloat1>(Value))
+					{
+						Float1Value->SetDefaultValue(FunctionInput->PreviewValue.X);
+						Float1Value->ApplyDefaultValue();
+					}
+					break;
+
+				case EFunctionInputType::FunctionInput_Vector2:
+					if (UDMMaterialValueFloat2* Float2Value = Cast<UDMMaterialValueFloat2>(Value))
+					{
+						Float2Value->SetDefaultValue({FunctionInput->PreviewValue.X, FunctionInput->PreviewValue.Y});
+						Float2Value->ApplyDefaultValue();
+					}
+					break;
+
+				case EFunctionInputType::FunctionInput_Vector3:
+					if (UDMMaterialValueFloat3XYZ* Float3Value = Cast<UDMMaterialValueFloat3XYZ>(Value))
+					{
+						Float3Value->SetDefaultValue({FunctionInput->PreviewValue.X, FunctionInput->PreviewValue.Y, FunctionInput->PreviewValue.Z});
+						Float3Value->ApplyDefaultValue();
+					}
+					break;
+
+				case EFunctionInputType::FunctionInput_Vector4:
+					if (UDMMaterialValueFloat4* Float4Value = Cast<UDMMaterialValueFloat4>(Value))
+					{
+						Float4Value->SetDefaultValue(FunctionInput->PreviewValue);
+						Float4Value->ApplyDefaultValue();
+					}
+					break;
+
+				default:
+					// Not possible
+					break;
+			}
+		}
+	}
+}
+
+void UDMMaterialStageFunction::DeinitFunction()
+{
+	InputConnectors.SetNum(1);
+
+	if (UDMMaterialStage* Stage = GetStage())
+	{
+		if (GUndo)
+		{
+			Stage->Modify();
+		}
+
+		Stage->RemoveAllInputs();
+	}
+}
+
+void UDMMaterialStageFunction::OnComponentAdded()
+{
+	if (!IsComponentValid())
+	{
+		return;
+	}
+
+	if (NeedsFunctionInit())
+	{
+		InitFunction();
+	}
+
+	Super::OnComponentAdded();
+}
+
+void UDMMaterialStageFunction::OnComponentRemoved()
+{
+	DeinitFunction();
+
+	Super::OnComponentRemoved();
 }
 
 #undef LOCTEXT_NAMESPACE
diff --git a/Engine/Plugins/Experimental/DynamicMaterial/Source/DynamicMaterialEditor/Private/Components/DMMaterialStageGradient.cpp b/Engine/Plugins/Experimental/DynamicMaterial/Source/DynamicMaterialEditor/Private/Components/DMMaterialStageGradient.cpp
--- a/Engine/Plugins/Experimental/DynamicMaterial/Source/DynamicMaterialEditor/Private/Components/DMMaterialStageGradient.cpp
+++ b/Engine/Plugins/Experimental/DynamicMaterial/Source/DynamicMaterialEditor/Private/Components/DMMaterialStageGradient.cpp
@@ -33,7 +33,14 @@ UDMMaterialStage* UDMMaterialStageGradient::CreateStage(TSubclassOf<UDMMaterialS
 	const FDMUpdateGuard Guard;
 
 	UDMMaterialStage* NewStage = UDMMaterialStage::CreateMaterialStage(InLayer);
-	UDMMaterialStageGradient* SourceGradient = NewObject<UDMMaterialStageGradient>(NewStage, InMaterialStageGradientClass.Get(), NAME_None, RF_Transactional);
+
+	UDMMaterialStageGradient* SourceGradient = NewObject<UDMMaterialStageGradient>(
+		NewStage, 
+		InMaterialStageGradientClass.Get(), 
+		NAME_None, 
+		RF_Transactional
+	);
+	
 	check(SourceGradient);
 
 	NewStage->SetSource(SourceGradient);
@@ -51,6 +58,22 @@ const TArray<TStrongObjectPtr<UClass>>& UDMMaterialStageGradient::GetAvailableGr
 	return Gradients;
 }
 
+UDMMaterialStageGradient* UDMMaterialStageGradient::ChangeStageSource_Gradient(UDMMaterialStage* InStage, 
+	TSubclassOf<UDMMaterialStageGradient> InGradientClass)
+{
+	check(InStage);
+
+	if (!InStage->CanChangeSource())
+	{
+		return nullptr;
+	}
+
+	check(InGradientClass);
+	check(!(InGradientClass->ClassFlags & (CLASS_Abstract | CLASS_Hidden | CLASS_Deprecated | CLASS_NewerVersionExists)));
+
+	return InStage->ChangeSource<UDMMaterialStageGradient>(InGradientClass);
+}
+
 bool UDMMaterialStageGradient::CanChangeInputType(int32 InputIndex) const
 {
 	if (InputIndex == 0)
diff --git a/Engine/Plugins/Experimental/DynamicMaterial/Source/DynamicMaterialEditor/Private/Components/DMMaterialStageThroughput.cpp b/Engine/Plugins/Experimental/DynamicMaterial/Source/DynamicMaterialEditor/Private/Components/DMMaterialStageThroughput.cpp
--- a/Engine/Plugins/Experimental/DynamicMaterial/Source/DynamicMaterialEditor/Private/Components/DMMaterialStageThroughput.cpp
+++ b/Engine/Plugins/Experimental/DynamicMaterial/Source/DynamicMaterialEditor/Private/Components/DMMaterialStageThroughput.cpp
@@ -8,6 +8,7 @@
 #include "Components/MaterialStageExpressions/DMMSETextureSample.h"
 #include "Components/MaterialStageInputs/DMMSIExpression.h"
 #include "Components/MaterialStageInputs/DMMSITextureUV.h"
+#include "Components/MaterialStageInputs/DMMSIValue.h"
 #include "Containers/ArrayView.h"
 #include "DMDefs.h"
 #include "DMMaterialFunctionLibrary.h"
@@ -268,7 +269,14 @@ void UDMMaterialStageThroughput::ConnectOutputToInput(const TSharedRef<FDMMateri
 	TConstArrayView<FExpressionInput*> TargetExpressionInputs = TargetExpression->GetInputsView();
 	check(TargetExpressionInputs.IsValidIndex(InInputIndex));
 
-	ConnectOutputToInput_Internal(InBuildState, TargetExpression, InInputIndex, InSourceExpression, InSourceOutputIndex, InSourceOutputChannel);
+	ConnectOutputToInput_Internal(
+		InBuildState, 
+		TargetExpression, 
+		InInputIndex, 
+		InSourceExpression, 
+		InSourceOutputIndex, 
+		InSourceOutputChannel
+	);
 }
 
 void UDMMaterialStageThroughput::ConnectOutputToInput_Internal(const TSharedRef<FDMMaterialBuildState>& InBuildState, UMaterialExpression* TargetExpression,
@@ -312,7 +320,12 @@ void UDMMaterialStageThroughput::ConnectOutputToInput_Internal(const TSharedRef<
 	}
 	else
 	{
-		UMaterialExpression* MaskExpression = InBuildState->GetBuildUtils().CreateExpressionBitMask(SourceExpression, SourceOutputIndex, SourceOutputChannel);
+		UMaterialExpression* MaskExpression = InBuildState->GetBuildUtils().CreateExpressionBitMask(
+			SourceExpression, 
+			SourceOutputIndex, 
+			SourceOutputChannel
+		);
+
 		MaskExpression->ConnectExpression(ExpressionInput, 0);
 	}
 }
@@ -332,7 +345,13 @@ FDMExpressionInput UDMMaterialStageThroughput::GetLayerMaskLinkTextureUVInputExp
 	FDMMaterialStageConnectorChannel Channel;
 	FDMExpressionInput ExpressionInput;
 
-	ExpressionInput.OutputIndex = ResolveInput(InBuildState, TextureUVInputIdx, Channel, ExpressionInput.OutputExpressions);
+	ExpressionInput.OutputIndex = ResolveInput(
+		InBuildState, 
+		TextureUVInputIdx, 
+		Channel, 
+		ExpressionInput.OutputExpressions
+	);
+
 	ExpressionInput.OutputChannel = Channel.OutputChannel;
 
 	return ExpressionInput;
@@ -397,29 +416,58 @@ void UDMMaterialStageThroughput::AddDefaultInput(int32 InInputIndex) const
 		case EDMValueType::VT_Float2:
 			if (InputConnectors[InInputIndex].Name.EqualTo(UVName))
 			{
-				Stage->ChangeInput_UV(InInputIndex, FDMMaterialStageConnectorChannel::WHOLE_CHANNEL, FDMMaterialStageConnectorChannel::WHOLE_CHANNEL);
+				UDMMaterialStageInputTextureUV::ChangeStageInput_UV(
+					Stage, 
+					InInputIndex,
+					FDMMaterialStageConnectorChannel::WHOLE_CHANNEL,
+					FDMMaterialStageConnectorChannel::WHOLE_CHANNEL
+				);
 			}
 			else
 			{
-				Stage->ChangeInput_NewLocalValue(InInputIndex, FDMMaterialStageConnectorChannel::WHOLE_CHANNEL, EDMValueType::VT_Float2, FDMMaterialStageConnectorChannel::WHOLE_CHANNEL);
+				UDMMaterialStageInputValue::ChangeStageInput_NewLocalValue(
+					Stage, 
+					InInputIndex,
+					FDMMaterialStageConnectorChannel::WHOLE_CHANNEL, 
+					EDMValueType::VT_Float2,
+					FDMMaterialStageConnectorChannel::WHOLE_CHANNEL
+				);
 			}
 			break;
 
 		case EDMValueType::VT_Float3_RGB:
-			Stage->ChangeInput_Expression(InInputIndex, FDMMaterialStageConnectorChannel::WHOLE_CHANNEL,
-				UDMMaterialStageExpressionTextureSample::StaticClass(), 0, FDMMaterialStageConnectorChannel::WHOLE_CHANNEL);
+			UDMMaterialStageInputExpression::ChangeStageInput_Expression(
+				Stage, 
+				UDMMaterialStageExpressionTextureSample::StaticClass(),
+				InInputIndex,
+				FDMMaterialStageConnectorChannel::WHOLE_CHANNEL, 
+				0,
+				FDMMaterialStageConnectorChannel::WHOLE_CHANNEL
+			);
 			break;
 
 		case EDMValueType::VT_Float4_RGBA:
-			Stage->ChangeInput_Expression(InInputIndex, FDMMaterialStageConnectorChannel::WHOLE_CHANNEL,
-				UDMMaterialStageExpressionTextureSample::StaticClass(), 5, FDMMaterialStageConnectorChannel::WHOLE_CHANNEL);
+			UDMMaterialStageInputExpression::ChangeStageInput_Expression(
+				Stage,
+				UDMMaterialStageExpressionTextureSample::StaticClass(), 
+				InInputIndex,
+				FDMMaterialStageConnectorChannel::WHOLE_CHANNEL, 
+				5,
+				FDMMaterialStageConnectorChannel::WHOLE_CHANNEL
+			);
 			break;
 
 		case EDMValueType::VT_Float1:
 		case EDMValueType::VT_Float3_RPY:
 		case EDMValueType::VT_Float3_XYZ:
 		case EDMValueType::VT_Texture:
-			Stage->ChangeInput_NewLocalValue(InInputIndex, FDMMaterialStageConnectorChannel::WHOLE_CHANNEL, InputConnectors[InInputIndex].Type, FDMMaterialStageConnectorChannel::WHOLE_CHANNEL);
+			UDMMaterialStageInputValue::ChangeStageInput_NewLocalValue(
+				Stage, 
+				InInputIndex,
+				FDMMaterialStageConnectorChannel::WHOLE_CHANNEL, 
+				InputConnectors[InInputIndex].Type,
+				FDMMaterialStageConnectorChannel::WHOLE_CHANNEL
+			);
 			break;
 
 		case EDMValueType::VT_ColorAtlas:
@@ -428,12 +476,24 @@ void UDMMaterialStageThroughput::AddDefaultInput(int32 InInputIndex) const
 			{
 				if (Layer->GetStageType(Stage) == EDMMaterialLayerStage::Mask)
 				{
-					Stage->ChangeInput_NewLocalValue(InInputIndex, FDMMaterialStageConnectorChannel::WHOLE_CHANNEL, EDMValueType::VT_ColorAtlas, FDMMaterialStageConnectorChannel::FOURTH_CHANNEL);
+					UDMMaterialStageInputValue::ChangeStageInput_NewLocalValue(
+						Stage, 
+						InInputIndex,
+						FDMMaterialStageConnectorChannel::WHOLE_CHANNEL, 
+						EDMValueType::VT_ColorAtlas,
+						FDMMaterialStageConnectorChannel::FOURTH_CHANNEL
+					);
 					break;
 				}
 			}
 
-			Stage->ChangeInput_NewLocalValue(InInputIndex, FDMMaterialStageConnectorChannel::WHOLE_CHANNEL, EDMValueType::VT_ColorAtlas, FDMMaterialStageConnectorChannel::THREE_CHANNELS);
+			UDMMaterialStageInputValue::ChangeStageInput_NewLocalValue(
+				Stage, 
+				InInputIndex,
+				FDMMaterialStageConnectorChannel::WHOLE_CHANNEL, 
+				EDMValueType::VT_ColorAtlas,
+				FDMMaterialStageConnectorChannel::THREE_CHANNELS
+			);
 			break;
 		}
 	}
@@ -461,7 +521,12 @@ int32 UDMMaterialStageThroughput::ResolveInput(const TSharedRef<FDMMaterialBuild
 		&& SupportsLayerMaskTextureUVLink() 
 		&& (GetLayerMaskTextureUVLinkInputIndex() == INDEX_NONE || GetLayerMaskTextureUVLinkInputIndex() == InputIndex))
 	{
-		int32 OutputIndex = ResolveLayerMaskTextureUVLinkInput(InBuildState, InputIndex, OutChannel, OutExpressions);
+		const int32 OutputIndex = ResolveLayerMaskTextureUVLinkInput(
+			InBuildState,
+			InputIndex, 
+			OutChannel, 
+			OutExpressions
+		);
 
 		if (OutputIndex != INDEX_NONE)
 		{
@@ -483,8 +548,16 @@ int32 UDMMaterialStageThroughput::ResolveInput(const TSharedRef<FDMMaterialBuild
 
 	if (InputConnectionMap[InputIndex].Channels.Num() == 1)
 	{
-		OutChannel = InputConnectionMap[InputIndex].Channels[0]; // Full copy in case it is changed by the channel resolve.
-		return ResolveInputChannel(InBuildState, InputIndex, 0, OutChannel, OutExpressions);
+		// Full copy in case it is changed by the channel resolve.
+		OutChannel = InputConnectionMap[InputIndex].Channels[0];
+
+		return ResolveInputChannel(
+			InBuildState, 
+			InputIndex, 
+			0, 
+			OutChannel, 
+			OutExpressions
+		);
 	}
 
 	// Only floats can have sub-channel mapping
@@ -527,7 +600,14 @@ int32 UDMMaterialStageThroughput::ResolveInput(const TSharedRef<FDMMaterialBuild
 		FDMMaterialStageConnectorChannel ChannelTemp = InputConnectionMap[InputIndex].Channels[ChannelIdx];
 		TArray<UMaterialExpression*> ChannelExpressions;
 
-		ResolveInputChannel(InBuildState, InputIndex, ChannelIdx, ChannelTemp, ChannelExpressions);
+		ResolveInputChannel(
+			InBuildState, 
+			InputIndex, 
+			ChannelIdx, 
+			ChannelTemp, 
+			ChannelExpressions
+		);
+
 		check(ChannelExpressions.IsEmpty() == false);
 
 		OutExpressions.Append(ChannelExpressions);
@@ -550,7 +630,12 @@ int32 UDMMaterialStageThroughput::ResolveLayerMaskTextureUVLinkInput(const TShar
 	check(Layer);
 	check(Layer->GetStage(EDMMaterialLayerStage::Base));
 
-	return ResolveLayerMaskTextureUVLinkInputImpl(InBuildState, Layer->GetStage(EDMMaterialLayerStage::Base)->GetSource(), OutChannel, OutExpressions);
+	return ResolveLayerMaskTextureUVLinkInputImpl(
+		InBuildState, 
+		Layer->GetStage(EDMMaterialLayerStage::Base)->GetSource(), 
+		OutChannel, 
+		OutExpressions
+	);
 }
 
 int32 UDMMaterialStageThroughput::ResolveLayerMaskTextureUVLinkInputImpl(const TSharedRef<FDMMaterialBuildState>& InBuildState, 
@@ -851,11 +936,20 @@ void UDMMaterialStageThroughput::UpdatePreviewMaterial(UMaterial* InPreviewMater
 		Stage->GenerateExpressions(BuildState);
 		UMaterialExpression* StageExpression = BuildState->GetLastStageExpression(Stage);
 
-		BuildState->GetBuildUtils().UpdatePreviewMaterial(StageExpression, 0, FDMMaterialStageConnectorChannel::WHOLE_CHANNEL, 32);
+		BuildState->GetBuildUtils().UpdatePreviewMaterial(
+			StageExpression, 
+			0, 
+			FDMMaterialStageConnectorChannel::WHOLE_CHANNEL, 
+			32
+		);
 	}
 	else
 	{
-		UE::DynamicMaterialEditor::Private::BuildExpressionInputs(BuildState, InputConnectionMap, Inputs);
+		UE::DynamicMaterialEditor::Private::BuildExpressionInputs(
+			BuildState,
+			InputConnectionMap, 
+			Inputs
+		);
 	}
 }
 
diff --git a/Engine/Plugins/Experimental/DynamicMaterial/Source/DynamicMaterialEditor/Private/Components/DMMaterialStageThroughputLayerBlend.cpp b/Engine/Plugins/Experimental/DynamicMaterial/Source/DynamicMaterialEditor/Private/Components/DMMaterialStageThroughputLayerBlend.cpp
--- a/Engine/Plugins/Experimental/DynamicMaterial/Source/DynamicMaterialEditor/Private/Components/DMMaterialStageThroughputLayerBlend.cpp
+++ b/Engine/Plugins/Experimental/DynamicMaterial/Source/DynamicMaterialEditor/Private/Components/DMMaterialStageThroughputLayerBlend.cpp
@@ -120,20 +120,26 @@ void UDMMaterialStageThroughputLayerBlend::AddDefaultInput(int32 InInputIndex) c
  
 		case InputMaskSource:
 		{
-			UDMMaterialStageInputExpression* InputExpression = Cast<UDMMaterialStageInputExpression>(
-				Stage->ChangeInput_Expression(
-					InInputIndex, 
-					FDMMaterialStageConnectorChannel::WHOLE_CHANNEL, 
-					UDMMaterialStageExpressionTextureSample::StaticClass(), 
-					0, // RGB pin
-					FDMMaterialStageConnectorChannel::WHOLE_CHANNEL // Alpha channel
-				)
+			UDMMaterialStageInputExpression* InputExpression = UDMMaterialStageInputExpression::ChangeStageInput_Expression(
+				Stage,
+				UDMMaterialStageExpressionTextureSample::StaticClass(),
+				InInputIndex,
+				FDMMaterialStageConnectorChannel::WHOLE_CHANNEL, 
+				0, // RGB pin
+				FDMMaterialStageConnectorChannel::WHOLE_CHANNEL // Alpha channel
 			);
 
 			UDMMaterialSubStage* SubStage = InputExpression->GetSubStage();
 			check(SubStage);
 
-			UDMMaterialStageInputValue* InputValue = SubStage->ChangeInput_NewLocalValue(0, FDMMaterialStageConnectorChannel::WHOLE_CHANNEL, EDMValueType::VT_Texture, FDMMaterialStageConnectorChannel::WHOLE_CHANNEL);
+			UDMMaterialStageInputValue* InputValue = UDMMaterialStageInputValue::ChangeStageInput_NewLocalValue(
+				SubStage,
+				0, 
+				FDMMaterialStageConnectorChannel::WHOLE_CHANNEL, 
+				EDMValueType::VT_Texture, 
+				FDMMaterialStageConnectorChannel::WHOLE_CHANNEL
+			);
+
 			check(InputValue);
 
 			UDMMaterialValueTexture* InputTexture = Cast<UDMMaterialValueTexture>(InputValue->GetValue());
@@ -717,15 +723,36 @@ void UDMMaterialStageThroughputLayerBlend::ConnectOutputToInput(const TSharedRef
 	switch (InInputIndex)
 	{
 		case InputPreviousLayer:
-			ConnectOutputToInput_Internal(InBuildState, StageSourceExpressions[0] /* FunctionCall */, InputPreviousIndex, InSourceExpression, InSourceOutputIndex, InSourceOutputChannel);
+			ConnectOutputToInput_Internal(
+				InBuildState, 
+				StageSourceExpressions[0] /* FunctionCall */, 
+				InputPreviousIndex, 
+				InSourceExpression, 
+				InSourceOutputIndex, 
+				InSourceOutputChannel
+			);
 			break;
  
 		case InputBaseStage:
-			ConnectOutputToInput_Internal(InBuildState, StageSourceExpressions[0] /* FunctionCall */, InputLayerIndex, InSourceExpression, InSourceOutputIndex, InSourceOutputChannel);
+			ConnectOutputToInput_Internal(
+				InBuildState, 
+				StageSourceExpressions[0] /* FunctionCall */, 
+				InputLayerIndex, 
+				InSourceExpression, 
+				InSourceOutputIndex, 
+				InSourceOutputChannel
+			);
 			break;
  
 		case InputMaskSource:
-			ConnectOutputToInput_Internal(InBuildState, StageSourceExpressions[0] /* FunctionCall */, InputMaskIndex, InSourceExpression, InSourceOutputIndex, InSourceOutputChannel);
+			ConnectOutputToInput_Internal(
+				InBuildState, 
+				StageSourceExpressions[0] /* FunctionCall */, 
+				InputMaskIndex, 
+				InSourceExpression, 
+				InSourceOutputIndex, 
+				InSourceOutputChannel
+			);
 			break;
  
 		default:
@@ -775,7 +802,13 @@ int32 UDMMaterialStageThroughputLayerBlend::ResolveMaskInput(const TSharedRef<FD
 	check(ThisLayer);
 	check(ThisLayer->GetStage(EDMMaterialLayerStage::Mask));
 
-	ThisLayer->ApplyEffects(InBuildState, ThisStage, OutExpressions, OutChannel.OutputChannel, OutChannel.OutputIndex);
+	ThisLayer->ApplyEffects(
+		InBuildState, 
+		ThisStage, 
+		OutExpressions, 
+		OutChannel.OutputChannel, 
+		OutChannel.OutputIndex
+	);
 		
 	while (const UDMMaterialLayerObject* NextLayer = ThisLayer->GetNextLayer(ThisLayer->GetMaterialProperty(), EDMMaterialLayerStage::Mask))
 	{
@@ -799,18 +832,46 @@ int32 UDMMaterialStageThroughputLayerBlend::ResolveMaskInput(const TSharedRef<FD
 		// Resolve the input of the next blend layer mask channel
 		FDMMaterialStageConnectorChannel NextOutChannel;
 		TArray<UMaterialExpression*> NextOutExpressions;
-		NextMaskBlend->ResolveInput(InBuildState, InputIndex, NextOutChannel, NextOutExpressions);
+
+		NextMaskBlend->ResolveInput(
+			InBuildState, 
+			InputIndex, 
+			NextOutChannel, 
+			NextOutExpressions
+		);
  
 		if (NextOutExpressions.IsEmpty())
 		{
 			break;
 		}
 
-		NextLayer->ApplyEffects(InBuildState, NextMaskStage, NextOutExpressions, NextOutChannel.OutputChannel, NextOutChannel.OutputIndex);
+		NextLayer->ApplyEffects(
+			InBuildState, 
+			NextMaskStage, 
+			NextOutExpressions, 
+			NextOutChannel.OutputChannel, 
+			NextOutChannel.OutputIndex
+		);
 
 		UMaterialExpressionMultiply* AlphaMultiply = InBuildState->GetBuildUtils().CreateExpression<UMaterialExpressionMultiply>(UE_DM_NodeComment_Default);
-		ConnectOutputToInput_Internal(InBuildState, AlphaMultiply, 0, OutExpressions.Last(), OutChannel.OutputIndex, OutChannel.OutputChannel);
-		ConnectOutputToInput_Internal(InBuildState, AlphaMultiply, 1, NextOutExpressions.Last(), NextOutChannel.OutputIndex, NextOutChannel.OutputChannel);
+
+		ConnectOutputToInput_Internal(
+			InBuildState, 
+			AlphaMultiply, 
+			0, 
+			OutExpressions.Last(), 
+			OutChannel.OutputIndex, 
+			OutChannel.OutputChannel
+		);
+
+		ConnectOutputToInput_Internal(
+			InBuildState, 
+			AlphaMultiply, 
+			1, 
+			NextOutExpressions.Last(), 
+			NextOutChannel.OutputIndex, 
+			NextOutChannel.OutputChannel
+		);
  
 		OutExpressions.Append(NextOutExpressions);
 		OutExpressions.Add(AlphaMultiply);
diff --git a/Engine/Plugins/Experimental/DynamicMaterial/Source/DynamicMaterialEditor/Private/Components/MaterialStageExpressions/DMMSEDivide.cpp b/Engine/Plugins/Experimental/DynamicMaterial/Source/DynamicMaterialEditor/Private/Components/MaterialStageExpressions/DMMSEDivide.cpp
--- a/Engine/Plugins/Experimental/DynamicMaterial/Source/DynamicMaterialEditor/Private/Components/MaterialStageExpressions/DMMSEDivide.cpp
+++ b/Engine/Plugins/Experimental/DynamicMaterial/Source/DynamicMaterialEditor/Private/Components/MaterialStageExpressions/DMMSEDivide.cpp
@@ -27,8 +27,14 @@ void UDMMaterialStageExpressionDivide::AddDefaultInput(int32 InInputIndex) const
  
 	UDMMaterialStage* Stage = GetStage();
 	check(Stage);
- 
-	Stage->ChangeInput_NewLocalValue(InInputIndex, FDMMaterialStageConnectorChannel::WHOLE_CHANNEL, EDMValueType::VT_Float1, FDMMaterialStageConnectorChannel::WHOLE_CHANNEL);
+
+	UDMMaterialStageInputValue::ChangeStageInput_NewLocalValue(
+		Stage, 
+		InInputIndex,
+		FDMMaterialStageConnectorChannel::WHOLE_CHANNEL, 
+		EDMValueType::VT_Float1,
+		FDMMaterialStageConnectorChannel::WHOLE_CHANNEL
+	);
  
 	UDMMaterialStageInputValue* InputValue = Cast<UDMMaterialStageInputValue>(Stage->GetInputs().Last());
 	check(InputValue);
diff --git a/Engine/Plugins/Experimental/DynamicMaterial/Source/DynamicMaterialEditor/Private/Components/MaterialStageExpressions/DMMSELinearInterpolate.cpp b/Engine/Plugins/Experimental/DynamicMaterial/Source/DynamicMaterialEditor/Private/Components/MaterialStageExpressions/DMMSELinearInterpolate.cpp
--- a/Engine/Plugins/Experimental/DynamicMaterial/Source/DynamicMaterialEditor/Private/Components/MaterialStageExpressions/DMMSELinearInterpolate.cpp
+++ b/Engine/Plugins/Experimental/DynamicMaterial/Source/DynamicMaterialEditor/Private/Components/MaterialStageExpressions/DMMSELinearInterpolate.cpp
@@ -2,9 +2,9 @@
  
 #include "Components/MaterialStageExpressions/DMMSELinearInterpolate.h"
 #include "Components/DMMaterialStage.h"
-#include "Materials/MaterialExpressionLinearInterpolate.h"
 #include "Components/MaterialStageInputs/DMMSIValue.h"
 #include "Components/MaterialValues/DMMaterialValueFloat1.h"
+#include "Materials/MaterialExpressionLinearInterpolate.h"
  
 #define LOCTEXT_NAMESPACE "DMMaterialStageExpressionLinearInterpolate"
  
@@ -37,7 +37,9 @@ void UDMMaterialStageExpressionLinearInterpolate::AddDefaultInput(int32 InInputI
 	UDMMaterialStage* Stage = GetStage();
 	check(Stage);
  
-	Stage->ChangeInput_NewLocalValue(InInputIndex, FDMMaterialStageConnectorChannel::WHOLE_CHANNEL, EDMValueType::VT_Float1, FDMMaterialStageConnectorChannel::WHOLE_CHANNEL);
+	UDMMaterialStageInputValue::ChangeStageInput_NewLocalValue(Stage, InInputIndex,
+		FDMMaterialStageConnectorChannel::WHOLE_CHANNEL, EDMValueType::VT_Float1,
+		FDMMaterialStageConnectorChannel::WHOLE_CHANNEL);
  
 	UDMMaterialStageInputValue* InputValue = Cast<UDMMaterialStageInputValue>(Stage->GetInputs().Last());
 	check(InputValue);
diff --git a/Engine/Plugins/Experimental/DynamicMaterial/Source/DynamicMaterialEditor/Private/Components/MaterialStageExpressions/DMMSETextureSampleBase.cpp b/Engine/Plugins/Experimental/DynamicMaterial/Source/DynamicMaterialEditor/Private/Components/MaterialStageExpressions/DMMSETextureSampleBase.cpp
--- a/Engine/Plugins/Experimental/DynamicMaterial/Source/DynamicMaterialEditor/Private/Components/MaterialStageExpressions/DMMSETextureSampleBase.cpp
+++ b/Engine/Plugins/Experimental/DynamicMaterial/Source/DynamicMaterialEditor/Private/Components/MaterialStageExpressions/DMMSETextureSampleBase.cpp
@@ -342,8 +342,9 @@ void UDMMaterialStageExpressionTextureSampleBase::UpdateMask()
 	// Couldn't find a texture to update, so create a new one.
 
 	// 2nd input, 2nd output (Alpha)
-	MaskStage->ChangeInput_Expression(2, FDMMaterialStageConnectorChannel::WHOLE_CHANNEL,
-		UDMMaterialStageExpressionTextureSample::StaticClass(), 1, FDMMaterialStageConnectorChannel::WHOLE_CHANNEL);
+	UDMMaterialStageInputExpression::ChangeStageInput_Expression(MaskStage, 
+		UDMMaterialStageExpressionTextureSample::StaticClass(), 2, FDMMaterialStageConnectorChannel::WHOLE_CHANNEL,
+		1, FDMMaterialStageConnectorChannel::WHOLE_CHANNEL);
 
 	UDMMaterialStageInputThroughput* InputThroughput = Cast<UDMMaterialStageInputThroughput>(MaskStage->GetInputs().Last());
 	UDMMaterialStageExpressionTextureSampleBase* MaskTextureSample = nullptr;
@@ -361,7 +362,9 @@ void UDMMaterialStageExpressionTextureSampleBase::UpdateMask()
 	UDMMaterialStage* MaskTextureSampleStage = MaskTextureSample->GetStage();
 	check(MaskTextureSampleStage);
 
-	MaskTextureSampleStage->ChangeInput_NewLocalValue(0, FDMMaterialStageConnectorChannel::WHOLE_CHANNEL, BaseTextureValue->GetType(), FDMMaterialStageConnectorChannel::WHOLE_CHANNEL);
+	UDMMaterialStageInputValue::ChangeStageInput_NewLocalValue(MaskTextureSampleStage, 0,
+		FDMMaterialStageConnectorChannel::WHOLE_CHANNEL, BaseTextureValue->GetType(),
+		FDMMaterialStageConnectorChannel::WHOLE_CHANNEL);
 
 	UDMMaterialStageInputValue* NewInputValue = Cast<UDMMaterialStageInputValue>(MaskTextureSampleStage->GetInputs().Last());
 	check(NewInputValue);
diff --git a/Engine/Plugins/Experimental/DynamicMaterial/Source/DynamicMaterialEditor/Private/Components/MaterialStageExpressions/DMMSETextureSampleEdgeColor.cpp b/Engine/Plugins/Experimental/DynamicMaterial/Source/DynamicMaterialEditor/Private/Components/MaterialStageExpressions/DMMSETextureSampleEdgeColor.cpp
--- a/Engine/Plugins/Experimental/DynamicMaterial/Source/DynamicMaterialEditor/Private/Components/MaterialStageExpressions/DMMSETextureSampleEdgeColor.cpp
+++ b/Engine/Plugins/Experimental/DynamicMaterial/Source/DynamicMaterialEditor/Private/Components/MaterialStageExpressions/DMMSETextureSampleEdgeColor.cpp
@@ -61,7 +61,9 @@ void UDMMaterialStageExpressionTextureSampleEdgeColor::AddDefaultInput(int32 InI
 	{
 		using namespace UE::DynamicMaterialEditor::TextureSampleEdgeColor::Private;
 
-		UDMMaterialStageInputValue* InputValue = Stage->ChangeInput_NewLocalValue(1, FDMMaterialStageConnectorChannel::WHOLE_CHANNEL, EDMValueType::VT_Float2, FDMMaterialStageConnectorChannel::WHOLE_CHANNEL);
+		UDMMaterialStageInputValue* InputValue = UDMMaterialStageInputValue::ChangeStageInput_NewLocalValue(Stage, 
+			1, FDMMaterialStageConnectorChannel::WHOLE_CHANNEL,
+			EDMValueType::VT_Float2, FDMMaterialStageConnectorChannel::WHOLE_CHANNEL);
 		check(InputValue);
 
 		UDMMaterialValueFloat2* InputFloat2 = Cast<UDMMaterialValueFloat2>(InputValue->GetValue());
diff --git a/Engine/Plugins/Experimental/DynamicMaterial/Source/DynamicMaterialEditor/Private/Components/MaterialStageInputs/DMMSIExpression.cpp b/Engine/Plugins/Experimental/DynamicMaterial/Source/DynamicMaterialEditor/Private/Components/MaterialStageInputs/DMMSIExpression.cpp
--- a/Engine/Plugins/Experimental/DynamicMaterial/Source/DynamicMaterialEditor/Private/Components/MaterialStageInputs/DMMSIExpression.cpp
+++ b/Engine/Plugins/Experimental/DynamicMaterial/Source/DynamicMaterialEditor/Private/Components/MaterialStageInputs/DMMSIExpression.cpp
@@ -1,9 +1,12 @@
 // Copyright Epic Games, Inc. All Rights Reserved.
 
 #include "Components/MaterialStageInputs/DMMSIExpression.h"
+
+#include "Components/DMMaterialLayer.h"
 #include "Components/DMMaterialSlot.h"
 #include "Components/DMMaterialStage.h"
 #include "Components/DMMaterialStageExpression.h"
+#include "Model/DynamicMaterialModelEditorOnlyData.h"
 
 TArray<TStrongObjectPtr<UClass>> UDMMaterialStageInputExpression::InputExpressions = TArray<TStrongObjectPtr<UClass>>();
 
@@ -47,6 +50,53 @@ const TArray<TStrongObjectPtr<UClass>>& UDMMaterialStageInputExpression::GetAvai
 	return InputExpressions;
 }
 
+UDMMaterialStageInputExpression* UDMMaterialStageInputExpression::ChangeStageSource_Expression(UDMMaterialStage* InStage, 
+	TSubclassOf<UDMMaterialStageExpression> InExpressionClass)
+{
+	check(InStage);
+
+	if (!InStage->CanChangeSource())
+	{
+		return nullptr;
+	}
+
+	check(InExpressionClass);
+	check(!(InExpressionClass->ClassFlags & (CLASS_Abstract | CLASS_Hidden | CLASS_Deprecated | CLASS_NewerVersionExists)));
+
+	UDMMaterialStageInputExpression* NewInputExpression = InStage->ChangeSource<UDMMaterialStageInputExpression>(
+		[InExpressionClass](UDMMaterialStage* InStage, UDMMaterialStageSource* InNewSource)
+		{
+			const FDMUpdateGuard Guard;
+			CastChecked<UDMMaterialStageInputExpression>(InNewSource)->SetMaterialStageExpressionClass(InExpressionClass);
+		});
+
+	return NewInputExpression;
+}
+
+UDMMaterialStageInputExpression* UDMMaterialStageInputExpression::ChangeStageInput_Expression(UDMMaterialStage* InStage, 
+	TSubclassOf<UDMMaterialStageExpression> InExpressionClass, int32 InInputIdx, int32 InInputChannel, int32 InOutputIdx, 
+	int32 InOutputChannel)
+{
+	check(InStage);
+
+	UDMMaterialStageSource* Source = InStage->GetSource();
+	check(Source);
+
+	check(InExpressionClass);
+	check(!(InExpressionClass->ClassFlags & (CLASS_Abstract | CLASS_Hidden | CLASS_Deprecated | CLASS_NewerVersionExists)));
+
+	UDMMaterialStageInputExpression* NewInputExpression = InStage->ChangeInput<UDMMaterialStageInputExpression>(
+		InInputIdx, InInputChannel, InOutputIdx, InOutputChannel,
+		[InExpressionClass](UDMMaterialStage* InStage, UDMMaterialStageInput* InNewInput)
+		{
+			const FDMUpdateGuard Guard;
+			CastChecked<UDMMaterialStageInputExpression>(InNewInput)->SetMaterialStageExpressionClass(InExpressionClass);
+		}
+	);
+
+	return NewInputExpression;
+}
+
 TSubclassOf<UDMMaterialStageExpression> UDMMaterialStageInputExpression::GetMaterialStageExpressionClass() const
 {
 	return TSubclassOf<UDMMaterialStageExpression>(GetMaterialStageThroughputClass());
diff --git a/Engine/Plugins/Experimental/DynamicMaterial/Source/DynamicMaterialEditor/Private/Components/MaterialStageInputs/DMMSIFunction.cpp b/Engine/Plugins/Experimental/DynamicMaterial/Source/DynamicMaterialEditor/Private/Components/MaterialStageInputs/DMMSIFunction.cpp
--- a/Engine/Plugins/Experimental/DynamicMaterial/Source/DynamicMaterialEditor/Private/Components/MaterialStageInputs/DMMSIFunction.cpp
+++ b/Engine/Plugins/Experimental/DynamicMaterial/Source/DynamicMaterialEditor/Private/Components/MaterialStageInputs/DMMSIFunction.cpp
@@ -21,6 +21,55 @@ UDMMaterialStage* UDMMaterialStageInputFunction::CreateStage(UDMMaterialLayerObj
 	return NewStage;
 }
 
+UDMMaterialStageInputFunction* UDMMaterialStageInputFunction::ChangeStageSource_Function(UDMMaterialStage* InStage,
+	UMaterialFunctionInterface* InMaterialFunction)
+{
+	check(InStage);
+
+	if (!InStage->CanChangeSource())
+	{
+		return nullptr;
+	}
+
+	check(InMaterialFunction);
+
+	UDMMaterialStageInputFunction* NewInputFunction = InStage->ChangeSource<UDMMaterialStageInputFunction>(
+		[InMaterialFunction](UDMMaterialStage* InStage, UDMMaterialStageSource* InNewSource)
+		{
+			const FDMUpdateGuard Guard;
+			UDMMaterialStageInputFunction* NewFunction = CastChecked<UDMMaterialStageInputFunction>(InNewSource);
+			NewFunction->Init();
+			NewFunction->SetMaterialFunction(InMaterialFunction);
+		});
+
+	return NewInputFunction;
+}
+
+UDMMaterialStageInputFunction* UDMMaterialStageInputFunction::ChangeStageInput_Function(UDMMaterialStage* InStage,
+	UMaterialFunctionInterface* InMaterialFunction, int32 InInputIdx, int32 InInputChannel, int32 InOutputIdx,
+	int32 InOutputChannel)
+{
+	check(InStage);
+
+	UDMMaterialStageSource* Source = InStage->GetSource();
+	check(Source);
+
+	check(InMaterialFunction);
+
+	UDMMaterialStageInputFunction* NewInputFunction = InStage->ChangeInput<UDMMaterialStageInputFunction>(
+		InInputIdx, InInputChannel, InOutputIdx, InOutputChannel,
+		[InMaterialFunction](UDMMaterialStage* InStage, UDMMaterialStageInput* InNewInput)
+		{
+			const FDMUpdateGuard Guard;
+			UDMMaterialStageInputFunction* NewFunction = CastChecked<UDMMaterialStageInputFunction>(InNewInput);
+			NewFunction->Init();
+			NewFunction->SetMaterialFunction(InMaterialFunction);
+		}
+	);
+
+	return NewInputFunction;
+}
+
 void UDMMaterialStageInputFunction::Init()
 {
 	SetMaterialStageThroughputClass(UDMMaterialStageFunction::StaticClass());
@@ -30,3 +79,21 @@ UDMMaterialStageFunction* UDMMaterialStageInputFunction::GetMaterialStageFunctio
 {
 	return Cast<UDMMaterialStageFunction>(GetMaterialStageThroughput());
 }
+
+UMaterialFunctionInterface* UDMMaterialStageInputFunction::GetMaterialFunction() const
+{
+	if (UDMMaterialStageFunction* StageFunction = GetMaterialStageFunction())
+	{
+		return StageFunction->GetMaterialFunction();
+	}
+
+	return nullptr;
+}
+
+void UDMMaterialStageInputFunction::SetMaterialFunction(UMaterialFunctionInterface* InMaterialFunction)
+{
+	if (UDMMaterialStageFunction* StageFunction = GetMaterialStageFunction())
+	{
+		StageFunction->SetMaterialFunction(InMaterialFunction);
+	}
+}
diff --git a/Engine/Plugins/Experimental/DynamicMaterial/Source/DynamicMaterialEditor/Private/Components/MaterialStageInputs/DMMSIGradient.cpp b/Engine/Plugins/Experimental/DynamicMaterial/Source/DynamicMaterialEditor/Private/Components/MaterialStageInputs/DMMSIGradient.cpp
--- a/Engine/Plugins/Experimental/DynamicMaterial/Source/DynamicMaterialEditor/Private/Components/MaterialStageInputs/DMMSIGradient.cpp
+++ b/Engine/Plugins/Experimental/DynamicMaterial/Source/DynamicMaterialEditor/Private/Components/MaterialStageInputs/DMMSIGradient.cpp
@@ -28,6 +28,52 @@ UDMMaterialStage* UDMMaterialStageInputGradient::CreateStage(TSubclassOf<UDMMate
 	return NewStage;
 }
 
+UDMMaterialStageInputGradient* UDMMaterialStageInputGradient::ChangeStageSource_Gradient(UDMMaterialStage* InStage,
+	TSubclassOf<UDMMaterialStageGradient> InGradientClass)
+{
+	check(InStage);
+
+	if (!InStage->CanChangeSource())
+	{
+		return nullptr;
+	}
+
+	check(InGradientClass);
+	check(!(InGradientClass->ClassFlags & (CLASS_Abstract | CLASS_Hidden | CLASS_Deprecated | CLASS_NewerVersionExists)));
+
+	UDMMaterialStageInputGradient* NewInputGradient = InStage->ChangeSource<UDMMaterialStageInputGradient>(
+		[InGradientClass](UDMMaterialStage* InStage, UDMMaterialStageSource* InNewSource)
+		{
+			const FDMUpdateGuard Guard;
+			CastChecked<UDMMaterialStageInputGradient>(InNewSource)->SetMaterialStageGradientClass(InGradientClass);
+		});
+
+	return NewInputGradient;
+}
+
+UDMMaterialStageInputGradient* UDMMaterialStageInputGradient::ChangeStageInput_Gradient(UDMMaterialStage* InStage,
+	TSubclassOf<UDMMaterialStageGradient> InGradientClass, int32 InInputIdx, int32 InInputChannel, int32 InOutputChannel)
+{
+	check(InStage);
+
+	UDMMaterialStageSource* Source = InStage->GetSource();
+	check(Source);
+
+	check(InGradientClass);
+	check(!(InGradientClass->ClassFlags & (CLASS_Abstract | CLASS_Hidden | CLASS_Deprecated | CLASS_NewerVersionExists)));
+
+	UDMMaterialStageInputGradient* NewInputGradient = InStage->ChangeInput<UDMMaterialStageInputGradient>(
+		InInputIdx, InInputChannel, 0, InOutputChannel,
+		[InGradientClass](UDMMaterialStage* InStage, UDMMaterialStageInput* InNewInput)
+		{
+			const FDMUpdateGuard Guard;
+			CastChecked<UDMMaterialStageInputGradient>(InNewInput)->SetMaterialStageGradientClass(InGradientClass);
+		}
+	);
+
+	return NewInputGradient;
+}
+
 void UDMMaterialStageInputGradient::SetMaterialStageGradientClass(TSubclassOf<UDMMaterialStageGradient> InMaterialStageGradientClass)
 {
 	SetMaterialStageThroughputClass(InMaterialStageGradientClass);
diff --git a/Engine/Plugins/Experimental/DynamicMaterial/Source/DynamicMaterialEditor/Private/Components/MaterialStageInputs/DMMSISlot.cpp b/Engine/Plugins/Experimental/DynamicMaterial/Source/DynamicMaterialEditor/Private/Components/MaterialStageInputs/DMMSISlot.cpp
--- a/Engine/Plugins/Experimental/DynamicMaterial/Source/DynamicMaterialEditor/Private/Components/MaterialStageInputs/DMMSISlot.cpp
+++ b/Engine/Plugins/Experimental/DynamicMaterial/Source/DynamicMaterialEditor/Private/Components/MaterialStageInputs/DMMSISlot.cpp
@@ -5,6 +5,7 @@
 #include "Components/DMMaterialProperty.h"
 #include "Components/DMMaterialSlot.h"
 #include "Components/DMMaterialStage.h"
+#include "Components/DMMaterialStageThroughput.h"
 #include "DMComponentPath.h"
 #include "DMPrivate.h"
 #include "DMValueDefinition.h"
@@ -44,6 +45,97 @@ UDMMaterialStage* UDMMaterialStageInputSlot::CreateStage(UDMMaterialSlot* InSour
 	return NewStage;
 }
 
+UDMMaterialStageInputSlot* UDMMaterialStageInputSlot::ChangeStageSource_Slot(UDMMaterialStage* InStage, UDMMaterialSlot* InSlot,
+	EDMMaterialPropertyType InProperty)
+{
+	check(InStage);
+
+	if (!InStage->CanChangeSource())
+	{
+		return nullptr;
+	}
+
+	UDMMaterialLayerObject* Layer = InStage->GetLayer();
+	check(Layer);
+
+	UDMMaterialSlot* Slot = Layer->GetSlot();
+	check(Slot);
+
+	check(InSlot);
+	check(Slot != InSlot);
+	check(InSlot->GetMaterialModelEditorOnlyData() == Slot->GetMaterialModelEditorOnlyData());
+
+	UDynamicMaterialModelEditorOnlyData* ModelEditorOnlyData = Slot->GetMaterialModelEditorOnlyData();
+	check(ModelEditorOnlyData);
+
+	TArray<EDMMaterialPropertyType> SlotPropertes = ModelEditorOnlyData->GetMaterialPropertiesForSlot(InSlot);
+	check(SlotPropertes.Contains(InProperty));
+
+	UDMMaterialStageInputSlot* InputSlot = InStage->ChangeSource<UDMMaterialStageInputSlot>(
+		[InSlot, InProperty](UDMMaterialStage* InStage, UDMMaterialStageSource* InNewSource)
+		{
+			const FDMUpdateGuard Guard;
+			UDMMaterialStageInputSlot* InputSlot = CastChecked<UDMMaterialStageInputSlot>(InNewSource);
+			InputSlot->SetSlot(InSlot);
+			InputSlot->SetMaterialProperty(InProperty);
+		});
+
+	return InputSlot;
+}
+
+UDMMaterialStageInputSlot* UDMMaterialStageInputSlot::ChangeStageInput_Slot(UDMMaterialStage* InStage, int32 InInputIdx, 
+	int32 InInputChannel, UDMMaterialSlot* InSlot, EDMMaterialPropertyType InProperty, int32 InOutputIdx, int32 InOutputChannel)
+{
+	check(InStage);
+
+	UDMMaterialStageSource* Source = InStage->GetSource();
+	check(Source);
+
+	UDMMaterialStageThroughput* Throughput = Cast<UDMMaterialStageThroughput>(Source);
+	check(Throughput);
+
+	const TArray<FDMMaterialStageConnector>& InputConnectors = Throughput->GetInputConnectors();
+	check(InputConnectors.IsValidIndex(InInputIdx));
+
+	UDMMaterialLayerObject* Layer = InStage->GetLayer();
+	check(Layer);
+
+	UDMMaterialSlot* Slot = Layer->GetSlot();
+	check(Slot);
+
+	UDynamicMaterialModelEditorOnlyData* ModelEditorOnlyData = Slot->GetMaterialModelEditorOnlyData();
+	check(ModelEditorOnlyData == InSlot->GetMaterialModelEditorOnlyData());
+
+	const TArray<EDMMaterialPropertyType> SlotPropertes = ModelEditorOnlyData->GetMaterialPropertiesForSlot(InSlot);
+	check(SlotPropertes.Contains(InProperty));
+
+	const TArray<EDMValueType>& SlotPropertyOutputTypes = InSlot->GetOutputConnectorTypesForMaterialProperty(InProperty);
+	check(SlotPropertyOutputTypes.IsValidIndex(InOutputIdx));
+
+	if (InOutputChannel == FDMMaterialStageConnectorChannel::WHOLE_CHANNEL)
+	{
+		check(Throughput->CanInputAcceptType(InInputIdx, SlotPropertyOutputTypes[InOutputIdx]));
+	}
+	else
+	{
+		check(UDMValueDefinitionLibrary::GetValueDefinition(SlotPropertyOutputTypes[InOutputIdx]).IsFloatType());
+		check(Throughput->CanInputAcceptType(InInputIdx, EDMValueType::VT_Float1));
+	}
+
+	UDMMaterialStageInputSlot* NewInputSlot = InStage->ChangeInput<UDMMaterialStageInputSlot>(
+		InInputIdx, InInputChannel, InOutputIdx, InOutputChannel, 
+		[InSlot, InProperty](UDMMaterialStage* InStage, UDMMaterialStageInput* InNewInput)
+		{
+			const FDMUpdateGuard Guard;
+			UDMMaterialStageInputSlot* NewInputSlot = CastChecked<UDMMaterialStageInputSlot>(InNewInput);
+			NewInputSlot->SetSlot(InSlot);
+			NewInputSlot->SetMaterialProperty(InProperty);
+		}
+	);
+
+	return NewInputSlot;
+}
+
 FText UDMMaterialStageInputSlot::GetComponentDescription() const
 {
 	if (Slot)
diff --git a/Engine/Plugins/Experimental/DynamicMaterial/Source/DynamicMaterialEditor/Private/Components/MaterialStageInputs/DMMSITextureUV.cpp b/Engine/Plugins/Experimental/DynamicMaterial/Source/DynamicMaterialEditor/Private/Components/MaterialStageInputs/DMMSITextureUV.cpp
--- a/Engine/Plugins/Experimental/DynamicMaterial/Source/DynamicMaterialEditor/Private/Components/MaterialStageInputs/DMMSITextureUV.cpp
+++ b/Engine/Plugins/Experimental/DynamicMaterial/Source/DynamicMaterialEditor/Private/Components/MaterialStageInputs/DMMSITextureUV.cpp
@@ -14,7 +14,6 @@
 #include "DMMaterialFunctionLibrary.h"
 #include "DynamicMaterialEditorModule.h"
 #include "DynamicMaterialModule.h"
-#include "Materials/MaterialExpressionAdd.h"
 #include "Materials/MaterialExpressionComponentMask.h"
 #include "Materials/MaterialExpressionMaterialFunctionCall.h"
 #include "Materials/MaterialExpressionScalarParameter.h"
@@ -45,6 +44,69 @@ UDMMaterialStage* UDMMaterialStageInputTextureUV::CreateStage(UDynamicMaterialMo
 	return NewStage;
 }
 
+UDMMaterialStageInputTextureUV* UDMMaterialStageInputTextureUV::ChangeStageSource_UV(UDMMaterialStage* InStage, bool bInDoUpdate)
+{
+	check(InStage);
+
+	if (!InStage->CanChangeSource())
+	{
+		return nullptr;
+	}
+
+	UDMMaterialLayerObject* Layer = InStage->GetLayer();
+	check(Layer);
+
+	UDMMaterialSlot* Slot = Layer->GetSlot();
+	check(Slot);
+
+	UDynamicMaterialModelEditorOnlyData* ModelEditorOnlyData = Slot->GetMaterialModelEditorOnlyData();
+	check(ModelEditorOnlyData);
+
+	UDynamicMaterialModel* MaterialModel = ModelEditorOnlyData->GetMaterialModel();
+	check(MaterialModel);
+
+	UDMMaterialStageInputTextureUV* InputTextureUV = InStage->ChangeSource<UDMMaterialStageInputTextureUV>(
+		[MaterialModel](UDMMaterialStage* InStage, UDMMaterialStageSource* InNewSource)
+		{
+			const FDMUpdateGuard Guard;
+			CastChecked<UDMMaterialStageInputTextureUV>(InNewSource)->Init(MaterialModel);;
+		});
+
+	return InputTextureUV;
+}
+
+UDMMaterialStageInputTextureUV* UDMMaterialStageInputTextureUV::ChangeStageInput_UV(UDMMaterialStage* InStage, int32 InInputIdx,
+	int32 InInputChannel, int32 InOutputChannel)
+{
+	check(InStage);
+
+	UDMMaterialStageSource* Source = InStage->GetSource();
+	check(Source);
+
+	UDMMaterialLayerObject* Layer = InStage->GetLayer();
+	check(Layer);
+
+	UDMMaterialSlot* Slot = Layer->GetSlot();
+	check(Slot);
+
+	UDynamicMaterialModelEditorOnlyData* ModelEditorOnlyData = Slot->GetMaterialModelEditorOnlyData();
+	check(ModelEditorOnlyData);
+
+	UDynamicMaterialModel* MaterialModel = ModelEditorOnlyData->GetMaterialModel();
+	check(MaterialModel);
+
+	UDMMaterialStageInputTextureUV* NewInputTextureUV = InStage->ChangeInput<UDMMaterialStageInputTextureUV>(
+		InInputIdx, InInputChannel, 0, InOutputChannel, 
+		[MaterialModel](UDMMaterialStage* InStage, UDMMaterialStageInput* InNewInput)
+		{
+			const FDMUpdateGuard Guard;
+			CastChecked<UDMMaterialStageInputTextureUV>(InNewInput)->Init(MaterialModel);
+		}
+	);
+
+	return NewInputTextureUV;
+}
+
 FText UDMMaterialStageInputTextureUV::GetComponentDescription() const
 {
 	return LOCTEXT("TexUV", "Tex UV");
diff --git a/Engine/Plugins/Experimental/DynamicMaterial/Source/DynamicMaterialEditor/Private/Components/MaterialStageInputs/DMMSIThroughput.cpp b/Engine/Plugins/Experimental/DynamicMaterial/Source/DynamicMaterialEditor/Private/Components/MaterialStageInputs/DMMSIThroughput.cpp
--- a/Engine/Plugins/Experimental/DynamicMaterial/Source/DynamicMaterialEditor/Private/Components/MaterialStageInputs/DMMSIThroughput.cpp
+++ b/Engine/Plugins/Experimental/DynamicMaterial/Source/DynamicMaterialEditor/Private/Components/MaterialStageInputs/DMMSIThroughput.cpp
@@ -289,7 +289,7 @@ void UDMMaterialStageInputThroughput::OnSubStageUpdated(UDMMaterialComponent* In
 		return;
 	}
 
-	if (InComponent == SubStage)
+	if (InComponent == SubStage && FDMUpdateGuard::CanUpdate())
 	{
 		Update(InUpdateType);
 	}
diff --git a/Engine/Plugins/Experimental/DynamicMaterial/Source/DynamicMaterialEditor/Private/Components/MaterialStageInputs/DMMSIValue.cpp b/Engine/Plugins/Experimental/DynamicMaterial/Source/DynamicMaterialEditor/Private/Components/MaterialStageInputs/DMMSIValue.cpp
--- a/Engine/Plugins/Experimental/DynamicMaterial/Source/DynamicMaterialEditor/Private/Components/MaterialStageInputs/DMMSIValue.cpp
+++ b/Engine/Plugins/Experimental/DynamicMaterial/Source/DynamicMaterialEditor/Private/Components/MaterialStageInputs/DMMSIValue.cpp
@@ -5,10 +5,10 @@
 #include "Components/DMMaterialParameter.h"
 #include "Components/DMMaterialSlot.h"
 #include "Components/DMMaterialStage.h"
+#include "Components/DMMaterialStageThroughputLayerBlend.h"
 #include "Components/DMMaterialSubStage.h"
 #include "Components/DMMaterialValue.h"
 #include "DMComponentPath.h"
-#include "Materials/MaterialExpressionComponentMask.h"
 #include "Model/DMMaterialBuildState.h"
 #include "Model/DynamicMaterialModel.h"
 #include "Model/DynamicMaterialModelEditorOnlyData.h"
@@ -34,6 +34,247 @@ UDMMaterialStage* UDMMaterialStageInputValue::CreateStage(UDMMaterialValue* InVa
 	return NewStage;
 }
 
+UDMMaterialStageInputValue* UDMMaterialStageInputValue::ChangeStageSource_NewLocalValue(UDMMaterialStage* InStage, EDMValueType InType)
+{
+	check(InStage);
+
+	if (!InStage->CanChangeSource())
+	{
+		return nullptr;
+	}
+
+	UDMMaterialLayerObject* Layer = InStage->GetLayer();
+	check(Layer);
+
+	UDMMaterialSlot* Slot = Layer->GetSlot();
+	check(Slot);
+
+	UDynamicMaterialModelEditorOnlyData* ModelEditorOnlyData = Slot->GetMaterialModelEditorOnlyData();
+	check(ModelEditorOnlyData);
+
+	UDynamicMaterialModel* MaterialModel = ModelEditorOnlyData->GetMaterialModel();
+	check(MaterialModel);
+
+	// Don't add via the builder, just parent it to the builder.
+	// It won't appear in the global value list.
+	UDMMaterialValue* NewValue = UDMMaterialValue::CreateMaterialValue(MaterialModel, TEXT(""), InType, true);
+	check(NewValue);
+
+	UDMMaterialStageInputValue* InputValue = InStage->ChangeSource<UDMMaterialStageInputValue>(
+		[NewValue](UDMMaterialStage* InStage, UDMMaterialStageSource* InNewSource)
+		{
+			const FDMUpdateGuard Guard;
+			CastChecked<UDMMaterialStageInputValue>(InNewSource)->SetValue(NewValue);
+		});
+
+	return InputValue;
+}
+
+UDMMaterialStageInputValue* UDMMaterialStageInputValue::ChangeStageSource_Value(UDMMaterialStage* InStage, UDMMaterialValue* InValue)
+{
+	check(InStage);
+
+	if (!InStage->CanChangeSource())
+	{
+		return nullptr;
+	}
+
+	UDMMaterialLayerObject* Layer = InStage->GetLayer();
+	check(Layer);
+
+	UDMMaterialSlot* Slot = Layer->GetSlot();
+	check(Slot);
+
+	UDynamicMaterialModelEditorOnlyData* ModelEditorOnlyData = Slot->GetMaterialModelEditorOnlyData();
+	check(ModelEditorOnlyData);
+
+	UDynamicMaterialModel* MaterialModel = ModelEditorOnlyData->GetMaterialModel();
+	check(MaterialModel);
+
+	UDMMaterialStageInputValue* InputValue = InStage->ChangeSource<UDMMaterialStageInputValue>(
+		[InValue](UDMMaterialStage* InStage, UDMMaterialStageSource* InNewSource)
+		{
+			const FDMUpdateGuard Guard;
+			CastChecked<UDMMaterialStageInputValue>(InNewSource)->SetValue(InValue);
+		});
+
+	return InputValue;
+}
+
+UDMMaterialStageInputValue* UDMMaterialStageInputValue::ChangeStageSource_NewValue(UDMMaterialStage* InStage, EDMValueType InType)
+{
+	check(InStage);
+
+	if (!InStage->CanChangeSource())
+	{
+		return nullptr;
+	}
+
+	UDMMaterialLayerObject* Layer = InStage->GetLayer();
+	check(Layer);
+
+	UDMMaterialSlot* Slot = Layer->GetSlot();
+	check(Slot);
+
+	UDynamicMaterialModelEditorOnlyData* ModelEditorOnlyData = Slot->GetMaterialModelEditorOnlyData();
+	check(ModelEditorOnlyData);
+
+	UDynamicMaterialModel* MaterialModel = ModelEditorOnlyData->GetMaterialModel();
+	check(MaterialModel);
+
+	UDMMaterialValue* NewValue = MaterialModel->AddValue(InType);
+	check(NewValue);
+
+	UDMMaterialStageInputValue* InputValue = InStage->ChangeSource<UDMMaterialStageInputValue>(
+		[NewValue](UDMMaterialStage* InStage, UDMMaterialStageSource* InNewSource)
+		{
+			const FDMUpdateGuard Guard;
+			CastChecked<UDMMaterialStageInputValue>(InNewSource)->SetValue(NewValue);
+		});
+
+	return InputValue;
+}
+
+UDMMaterialStageInputValue* UDMMaterialStageInputValue::ChangeStageInput_NewLocalValue(UDMMaterialStage* InStage,
+	int32 InInputIdx, int32 InInputChannel, EDMValueType InType, int32 InOutputChannel)
+{
+	check(InStage);
+
+	UDMMaterialStageSource* Source = InStage->GetSource();
+	check(Source);
+
+	UDMMaterialStageThroughput* Throughput = Cast<UDMMaterialStageThroughput>(Source);
+	check(Throughput);
+
+	const TArray<FDMMaterialStageConnector>& InputConnectors = Throughput->GetInputConnectors();
+	check(InputConnectors.IsValidIndex(InInputIdx));
+	check(Throughput->CanInputAcceptType(InInputIdx, InType));
+
+	UDMMaterialLayerObject* Layer = InStage->GetLayer();
+	check(Layer);
+
+	UDMMaterialSlot* Slot = Layer->GetSlot();
+	check(Slot);
+
+	UDynamicMaterialModelEditorOnlyData* ModelEditorOnlyData = Slot->GetMaterialModelEditorOnlyData();
+	check(ModelEditorOnlyData);
+
+	UDynamicMaterialModel* MaterialModel = ModelEditorOnlyData->GetMaterialModel();
+	check(MaterialModel);
+
+	// Don't add via the builder, just parent it to the builder.
+	// It won't appear in the global value list.
+	UDMMaterialValue* NewValue = UDMMaterialValue::CreateMaterialValue(MaterialModel, TEXT(""), InType, true);
+	check(NewValue);
+
+	UDMMaterialStageInputValue* NewInputValue = InStage->ChangeInput<UDMMaterialStageInputValue>(
+		InInputIdx, InInputChannel, 0, InOutputChannel,
+		[NewValue](UDMMaterialStage* InStage, UDMMaterialStageInput* InNewInput)
+		{
+			const FDMUpdateGuard Guard;
+			CastChecked<UDMMaterialStageInputValue>(InNewInput)->SetValue(NewValue);
+		}
+	);
+
+	NewInputValue->ApplyWholeLayerValue();
+
+	return NewInputValue;
+}
+
+UDMMaterialStageInputValue* UDMMaterialStageInputValue::ChangeStageInput_Value(UDMMaterialStage* InStage, int32 InInputIdx, int32 InInputChannel, UDMMaterialValue* InValue, int32 OutputChannel)
+{
+	check(InStage);
+	check(InValue);
+
+	UDMMaterialStageSource* Source = InStage->GetSource();
+	check(Source);
+
+	UDMMaterialStageThroughput* Throughput = Cast<UDMMaterialStageThroughput>(Source);
+	check(Throughput);
+
+	const TArray<FDMMaterialStageConnector>& InputConnectors = Throughput->GetInputConnectors();
+	check(InputConnectors.IsValidIndex(InInputIdx));
+
+	if (OutputChannel == FDMMaterialStageConnectorChannel::WHOLE_CHANNEL)
+	{
+		check(Throughput->CanInputAcceptType(InInputIdx, InValue->GetType()));
+	}
+	else
+	{
+		check(UDMValueDefinitionLibrary::GetValueDefinition(InValue->GetType()).IsFloatType());
+		check(Throughput->CanInputAcceptType(InInputIdx, EDMValueType::VT_Float1));
+	}
+
+	UDMMaterialLayerObject* Layer = InStage->GetLayer();
+	check(Layer);
+
+	UDMMaterialSlot* Slot = Layer->GetSlot();
+	check(Slot);
+
+	UDynamicMaterialModelEditorOnlyData* ModelEditorOnlyData = Slot->GetMaterialModelEditorOnlyData();
+	check(ModelEditorOnlyData);
+
+	UDynamicMaterialModel* MaterialModel = ModelEditorOnlyData->GetMaterialModel();
+	check(MaterialModel);
+	check(MaterialModel == InValue->GetMaterialModel());
+
+	UDMMaterialStageInputValue* NewInputValue = InStage->ChangeInput<UDMMaterialStageInputValue>(
+		InInputIdx, InInputChannel, 0, OutputChannel,
+		[InValue](UDMMaterialStage* InStage, UDMMaterialStageInput* InNewInput)
+		{
+			const FDMUpdateGuard Guard;
+			CastChecked<UDMMaterialStageInputValue>(InNewInput)->SetValue(InValue);
+		}
+	);
+
+	NewInputValue->ApplyWholeLayerValue();
+
+	return NewInputValue;
+}
+
+UDMMaterialStageInputValue* UDMMaterialStageInputValue::ChangeStageInput_NewValue(UDMMaterialStage* InStage, int32 InInputIdx, int32 InInputChannel, EDMValueType InType, int32 InOutputChannel)
+{
+	check(InStage);
+
+	UDMMaterialStageSource* Source = InStage->GetSource();
+	check(Source);
+
+	UDMMaterialStageThroughput* Throughput = Cast<UDMMaterialStageThroughput>(Source);
+	check(Throughput);
+
+	const TArray<FDMMaterialStageConnector>& InputConnectors = Throughput->GetInputConnectors();
+	check(InputConnectors.IsValidIndex(InInputIdx));
+	check(Throughput->CanInputAcceptType(InInputIdx, InType));
+
+	UDMMaterialLayerObject* Layer = InStage->GetLayer();
+	check(Layer);
+
+	UDMMaterialSlot* Slot = Layer->GetSlot();
+	check(Slot);
+
+	UDynamicMaterialModelEditorOnlyData* ModelEditorOnlyData = Slot->GetMaterialModelEditorOnlyData();
+	check(ModelEditorOnlyData);
+
+	UDynamicMaterialModel* MaterialModel = ModelEditorOnlyData->GetMaterialModel();
+	check(MaterialModel);
+
+	UDMMaterialValue* NewValue = MaterialModel->AddValue(InType);
+	check(NewValue);
+
+	UDMMaterialStageInputValue* NewInputValue = InStage->ChangeInput<UDMMaterialStageInputValue>(
+		InInputIdx, InInputChannel, 0, InOutputChannel,
+		[NewValue](UDMMaterialStage* InStage, UDMMaterialStageInput* InNewInput)
+		{
+			const FDMUpdateGuard Guard;
+			CastChecked<UDMMaterialStageInputValue>(InNewInput)->SetValue(NewValue);
+		}
+	);
+
+	NewInputValue->ApplyWholeLayerValue();
+
+	return NewInputValue;
+}
+
 FText UDMMaterialStageInputValue::GetComponentDescription() const
 {
 	if (!Value->IsComponentValid())
@@ -414,4 +655,53 @@ void UDMMaterialStageInputValue::UpdateOutputConnectors()
 	OutputConnectors.Add({0, LOCTEXT("MaterialValue", "Value"), Value ? Value->GetType() : EDMValueType::VT_None});
 }
 
+void UDMMaterialStageInputValue::ApplyWholeLayerValue()
+{
+	if (!Value.Get() || !Value->IsWholeLayerValue())
+	{
+		return;
+	}
+
+	UDMMaterialStage* Stage = GetStage();
+
+	if (!Stage)
+	{
+		return;
+	}
+
+	UDMMaterialLayerObject* Layer = Stage->GetLayer();
+
+	if (!Layer)
+	{
+		return;
+	}
+
+	if (!Layer
+		|| Layer->GetStageType(Stage) != EDMMaterialLayerStage::Base
+		|| !Layer->IsStageEnabled(EDMMaterialLayerStage::Mask))
+	{
+		return;
+	}
+
+	UDMMaterialStage* MaskStage = Layer->GetStage(EDMMaterialLayerStage::Mask);
+
+	if (!IsValid(MaskStage))
+	{
+		return;
+	}
+
+	if (GUndo)
+	{
+		MaskStage->Modify();
+	}
+
+	ChangeStageInput_Value(
+		MaskStage,
+		UDMMaterialStageThroughputLayerBlend::InputMaskSource,
+		FDMMaterialStageConnectorChannel::WHOLE_CHANNEL,
+		Value.Get(),
+		FDMMaterialStageConnectorChannel::FOURTH_CHANNEL
+	);
+}
+
 #undef LOCTEXT_NAMESPACE
diff --git a/Engine/Plugins/Experimental/DynamicMaterial/Source/DynamicMaterialEditor/Private/DMBlueprintFunctionLibrary.cpp b/Engine/Plugins/Experimental/DynamicMaterial/Source/DynamicMaterialEditor/Private/DMBlueprintFunctionLibrary.cpp
--- a/Engine/Plugins/Experimental/DynamicMaterial/Source/DynamicMaterialEditor/Private/DMBlueprintFunctionLibrary.cpp
+++ b/Engine/Plugins/Experimental/DynamicMaterial/Source/DynamicMaterialEditor/Private/DMBlueprintFunctionLibrary.cpp
@@ -100,14 +100,27 @@ void UDMBlueprintFunctionLibrary::SetDefaultStageSourceTexture(UDMMaterialStage*
 
 	if (UDMMaterialStageBlend* const Blend = Cast<UDMMaterialStageBlend>(StageSource))
 	{
-		UDMMaterialStageInputExpression* NewInput = InStage->ChangeInput_Expression(UDMMaterialStageBlend::InputB, FDMMaterialStageConnectorChannel::WHOLE_CHANNEL,
-			UDMMaterialStageExpressionTextureSample::StaticClass(), 0, FDMMaterialStageConnectorChannel::WHOLE_CHANNEL);
+		UDMMaterialStageInputExpression* NewInput = UDMMaterialStageInputExpression::ChangeStageInput_Expression(
+			InStage,
+			UDMMaterialStageExpressionTextureSample::StaticClass(), 
+			UDMMaterialStageBlend::InputB,
+			FDMMaterialStageConnectorChannel::WHOLE_CHANNEL, 
+			0,
+			FDMMaterialStageConnectorChannel::WHOLE_CHANNEL
+		);
 
 		UDMMaterialSubStage* SubStage = NewInput->GetSubStage();
 		check(SubStage);
-		UDMMaterialStageInputValue* NewInputValue = SubStage->ChangeInput_NewLocalValue(0, FDMMaterialStageConnectorChannel::WHOLE_CHANNEL, EDMValueType::VT_Texture, FDMMaterialStageConnectorChannel::WHOLE_CHANNEL);
-		SetTextureValue(NewInputValue);
 
+		UDMMaterialStageInputValue* NewInputValue = UDMMaterialStageInputValue::ChangeStageInput_NewLocalValue(
+			SubStage,
+			0, 
+			FDMMaterialStageConnectorChannel::WHOLE_CHANNEL, 
+			EDMValueType::VT_Texture,
+			FDMMaterialStageConnectorChannel::WHOLE_CHANNEL
+		);
+
+		SetTextureValue(NewInputValue);
 		return;
 	}
 
@@ -115,19 +128,40 @@ void UDMBlueprintFunctionLibrary::SetDefaultStageSourceTexture(UDMMaterialStage*
 	{
 		const bool bHasAlpha = UE::DynamicMaterial::Private::HasAlpha(InTexture);
 
-		UDMMaterialStageInputExpression* NewInput = InStage->ChangeInput_Expression(2, FDMMaterialStageConnectorChannel::WHOLE_CHANNEL,
-			UDMMaterialStageExpressionTextureSample::StaticClass(), bHasAlpha ? 1 : 0, FDMMaterialStageConnectorChannel::WHOLE_CHANNEL);
+		UDMMaterialStageInputExpression* NewInput = UDMMaterialStageInputExpression::ChangeStageInput_Expression(
+			InStage,
+			UDMMaterialStageExpressionTextureSample::StaticClass(), 
+			2, 
+			FDMMaterialStageConnectorChannel::WHOLE_CHANNEL,
+			bHasAlpha ? 1 : 0, 
+			FDMMaterialStageConnectorChannel::WHOLE_CHANNEL
+		);
 
 		UDMMaterialSubStage* SubStage = NewInput->GetSubStage();
 		check(IsValid(SubStage));
-		UDMMaterialStageInputValue* NewInputValue = SubStage->ChangeInput_NewLocalValue(0, FDMMaterialStageConnectorChannel::WHOLE_CHANNEL, EDMValueType::VT_Texture, FDMMaterialStageConnectorChannel::WHOLE_CHANNEL);
-		SetTextureValue(NewInputValue);
 
+		UDMMaterialStageInputValue* NewInputValue = UDMMaterialStageInputValue::ChangeStageInput_NewLocalValue(
+			SubStage,
+			0, 
+			FDMMaterialStageConnectorChannel::WHOLE_CHANNEL, 
+			EDMValueType::VT_Texture,
+			FDMMaterialStageConnectorChannel::WHOLE_CHANNEL
+		);
+
+		SetTextureValue(NewInputValue);
 		return;
 	}
 
-	UDMMaterialStageExpression* NewExpression = InStage->ChangeSource_Expression(UDMMaterialStageExpressionTextureSample::StaticClass());
-	UDMMaterialStageInputValue* NewInputValue = InStage->ChangeInput_NewLocalValue(0, FDMMaterialStageConnectorChannel::WHOLE_CHANNEL, EDMValueType::VT_Texture, FDMMaterialStageConnectorChannel::WHOLE_CHANNEL);
+	UDMMaterialStageExpression* NewExpression = InStage->ChangeSource<UDMMaterialStageExpressionTextureSample>();
+
+	UDMMaterialStageInputValue* NewInputValue = UDMMaterialStageInputValue::ChangeStageInput_NewLocalValue(
+		InStage,
+		0, 
+		FDMMaterialStageConnectorChannel::WHOLE_CHANNEL, 
+		EDMValueType::VT_Texture,
+		FDMMaterialStageConnectorChannel::WHOLE_CHANNEL
+	);
+
 	SetTextureValue(NewInputValue);
 }
 
diff --git a/Engine/Plugins/Experimental/DynamicMaterial/Source/DynamicMaterialEditor/Private/DynamicMaterialEditorModule.cpp b/Engine/Plugins/Experimental/DynamicMaterial/Source/DynamicMaterialEditor/Private/DynamicMaterialEditorModule.cpp
--- a/Engine/Plugins/Experimental/DynamicMaterial/Source/DynamicMaterialEditor/Private/DynamicMaterialEditorModule.cpp
+++ b/Engine/Plugins/Experimental/DynamicMaterial/Source/DynamicMaterialEditor/Private/DynamicMaterialEditorModule.cpp
@@ -7,6 +7,7 @@
 #include "Components/DMMaterialStageSource.h"
 #include "Components/DMMaterialStageThroughput.h"
 #include "Components/DMTextureUV.h"
+#include "Components/MaterialStageInputs/DMMSIFunction.h"
 #include "Components/MaterialStageInputs/DMMSIThroughput.h"
 #include "Components/MaterialValues/DMMaterialValueBool.h"
 #include "Components/MaterialValues/DMMaterialValueFloat1.h"
@@ -46,6 +47,7 @@
 #include "Slate/Properties/Generators/DMComponentPropertyRowGenerator.h"
 #include "Slate/Properties/Generators/DMInputThroughputPropertyRowGenerator.h"
 #include "Slate/Properties/Generators/DMMaterialEffectFunctionPropertyRowGenerator.h"
+#include "Slate/Properties/Generators/DMMaterialStageFunctionPropertyRowGenerator.h"
 #include "Slate/Properties/Generators/DMMaterialValuePropertyRowGenerator.h"
 #include "Slate/Properties/Generators/DMStagePropertyRowGenerator.h"
 #include "Slate/Properties/Generators/DMTextureUVPropertyRowGenerator.h"
@@ -244,6 +246,7 @@ void FDynamicMaterialEditorModule::StartupModule()
 	RegisterComponentPropertyRowGeneratorDelegate<UDMMaterialStageThroughput,      FDMThroughputPropertyRowGenerator>();
 	RegisterComponentPropertyRowGeneratorDelegate<UDMMaterialStageInputThroughput, FDMInputThroughputPropertyRowGenerator>();
 	RegisterComponentPropertyRowGeneratorDelegate<UDMMaterialEffectFunction,       FDMMaterialEffectFunctionPropertyRowGenerator>();
+	RegisterComponentPropertyRowGeneratorDelegate<UDMMaterialStageInputFunction,   FDMMaterialStageFunctionPropertyRowGenerator>();
 
 	FDynamicMaterialModule::GetCreateEditorOnlyDataDelegate().BindLambda([](UDynamicMaterialModel* InMaterialModel) -> TScriptInterface<IDynamicMaterialModelEditorOnlyDataInterface>
 		{
diff --git a/Engine/Plugins/Experimental/DynamicMaterial/Source/DynamicMaterialEditor/Private/Menus/DMMaterialStageInputMenus.cpp b/Engine/Plugins/Experimental/DynamicMaterial/Source/DynamicMaterialEditor/Private/Menus/DMMaterialStageInputMenus.cpp
--- a/Engine/Plugins/Experimental/DynamicMaterial/Source/DynamicMaterialEditor/Private/Menus/DMMaterialStageInputMenus.cpp
+++ b/Engine/Plugins/Experimental/DynamicMaterial/Source/DynamicMaterialEditor/Private/Menus/DMMaterialStageInputMenus.cpp
@@ -11,6 +11,11 @@
 #include "Components/DMMaterialValue.h"
 #include "Components/MaterialStageExpressions/DMMSETextureSample.h"
 #include "Components/MaterialStageExpressions/DMMSETextureSampleEdgeColor.h"
+#include "Components/MaterialStageInputs/DMMSIExpression.h"
+#include "Components/MaterialStageInputs/DMMSIGradient.h"
+#include "Components/MaterialStageInputs/DMMSISlot.h"
+#include "Components/MaterialStageInputs/DMMSITextureUV.h"
+#include "Components/MaterialStageInputs/DMMSIValue.h"
 #include "DMPrivate.h"
 #include "DMValueDefinition.h"
 #include "DynamicMaterialEditorModule.h"
@@ -216,7 +221,7 @@ EndStageSearch:
 			{
 				MenuBuilder.AddSubMenu(
 					LOCTEXT("ChangeInputSlot", "Slot Output"),
-					LOCTEXT("ChangeInputSlotTooltip", "Change the source of this input to the output from another Material Slot."),
+					LOCTEXT("ChangeInputSlotTooltip", "Change the source of this input to the output from another Material Slot."),					
 					FNewMenuDelegate::CreateStatic(
 						&FDMMaterialStageInputMenus::GenerateChangeInputMenu_Slots,
 						InThroughput,
@@ -237,7 +242,14 @@ EndStageSearch:
 				{
 					FScopedTransaction Transaction(LOCTEXT("SetStageInput", "Set Material Stage Input"));
 					Stage->Modify();
-					Stage->ChangeInput_NewLocalValue(InInputIndex, InInputChannel, EDMValueType::VT_Float3_RGB, FDMMaterialStageConnectorChannel::WHOLE_CHANNEL);
+
+					UDMMaterialStageInputValue::ChangeStageInput_NewLocalValue(
+						Stage, 
+						InInputIndex, 
+						InInputChannel,
+						EDMValueType::VT_Float3_RGB, 
+						FDMMaterialStageConnectorChannel::WHOLE_CHANNEL
+					);
 				})
 			)
 		);
@@ -252,7 +264,14 @@ EndStageSearch:
 				{
 					FScopedTransaction Transaction(LOCTEXT("SetStageInput", "Set Material Stage Input"));
 					Stage->Modify();
-					Stage->ChangeInput_NewLocalValue(InInputIndex, InInputChannel, EDMValueType::VT_ColorAtlas, FDMMaterialStageConnectorChannel::THREE_CHANNELS);
+
+					UDMMaterialStageInputValue::ChangeStageInput_NewLocalValue(
+						Stage, 
+						InInputIndex, 
+						InInputChannel,
+						EDMValueType::VT_ColorAtlas, 
+						FDMMaterialStageConnectorChannel::THREE_CHANNELS
+					);
 				})
 			)
 		);
@@ -544,11 +563,25 @@ void FDMMaterialStageInputMenus::GenerateChangeInputMenu_PreviousStages(FMenuBui
 		if (bJustOneConnector)
 		{
 			static const FText ExpressionNameFormatTemplate = LOCTEXT("ExpressionAndOutput", "{0} [{1}]");
-			const FText PropertyAndChannel = FText::Format(ExpressionNameFormatTemplate, PropertyObj->GetDescription(), PreviousStageOutputConnectors[LastValidOutputConnector].Name);
 
-			UE::DynamicMaterialEditor::Private::GenerateChangeInputMenu_PreviousStage_Output_Impl(InChildMenuBuilder,
-				PropertyAndChannel, Stage, InThroughput, InInputIndex, InInputChannel, SlotProperty, LastValidOutputConnector, 
-				bAcceptsWholeChannel, bAcceptsSubChannels);
+			const FText PropertyAndChannel = FText::Format(
+				ExpressionNameFormatTemplate, 
+				PropertyObj->GetDescription(), 
+				PreviousStageOutputConnectors[LastValidOutputConnector].Name
+				);
+
+			UE::DynamicMaterialEditor::Private::GenerateChangeInputMenu_PreviousStage_Output_Impl(
+				InChildMenuBuilder,
+				PropertyAndChannel, 
+				Stage, 
+				InThroughput, 
+				InInputIndex, 
+				InInputChannel, 
+				SlotProperty, 
+				LastValidOutputConnector, 
+				bAcceptsWholeChannel, 
+				bAcceptsSubChannels
+			);
 		}
 		else
 		{
@@ -622,12 +655,30 @@ void FDMMaterialStageInputMenus::GenerateChangeInputMenu_PreviousStage_Outputs(F
 	{
 		for (int32 OutputIdx = 0; OutputIdx < PreviousStageOutputConnectors.Num(); ++OutputIdx)
 		{
-			const bool bAcceptsWholeChannel = InThroughput->CanInputConnectTo(InInputIndex, PreviousStageOutputConnectors[OutputIdx], FDMMaterialStageConnectorChannel::WHOLE_CHANNEL);
-			const bool bAcceptsSubChannels = CanAcceptSubChannels(InThroughput, InInputIndex, PreviousStageOutputConnectors[OutputIdx]);
+			const bool bAcceptsWholeChannel = InThroughput->CanInputConnectTo(
+				InInputIndex, 
+				PreviousStageOutputConnectors[OutputIdx], 
+				FDMMaterialStageConnectorChannel::WHOLE_CHANNEL
+			);
+
+			const bool bAcceptsSubChannels = CanAcceptSubChannels(
+				InThroughput, 
+				InInputIndex, 
+				PreviousStageOutputConnectors[OutputIdx]
+			);
 
-			UE::DynamicMaterialEditor::Private::GenerateChangeInputMenu_PreviousStage_Output_Impl(InChildMenuBuilder,
-				PreviousStageOutputConnectors[OutputIdx].Name, Stage, InThroughput, InInputIndex, InInputChannel, InMaterialProperty, OutputIdx,
-				bAcceptsWholeChannel, bAcceptsSubChannels);
+			UE::DynamicMaterialEditor::Private::GenerateChangeInputMenu_PreviousStage_Output_Impl(
+				InChildMenuBuilder,
+				PreviousStageOutputConnectors[OutputIdx].Name, 
+				Stage, 
+				InThroughput, 
+				InInputIndex, 
+				InInputChannel,
+				InMaterialProperty, 
+				OutputIdx,
+				bAcceptsWholeChannel, 
+				bAcceptsSubChannels
+			);
 		}
 	}
 }
@@ -714,8 +765,17 @@ void FDMMaterialStageInputMenus::GenerateChangeInputMenu_PreviousStage_Output_Ch
 			const FText ChannelName = UDMValueDefinitionLibrary::GetValueDefinition(PreviousStageOutputConnectors[OutputIndex].Type).GetChannelName(ChannelIndex + 1);
 			const int32 OutputChannel = UE::DynamicMaterialEditor::Private::ChannelIndexToChannelBit(ChannelIndex + 1);
 
-			UE::DynamicMaterialEditor::Private::GenerateChangeInputMenu_PreviousStage_Output_Channel_Impl(InChildMenuBuilder, ChannelName, Stage, InThroughput,
-				InInputIndex, InInputChannel, InMaterialProperty, OutputIndex, OutputChannel);
+			UE::DynamicMaterialEditor::Private::GenerateChangeInputMenu_PreviousStage_Output_Channel_Impl(
+				InChildMenuBuilder, 
+				ChannelName, 
+				Stage, 
+				InThroughput,
+				InInputIndex, 
+				InInputChannel, 
+				InMaterialProperty, 
+				OutputIndex, 
+				OutputChannel
+			);
 		}
 	}
 	InChildMenuBuilder.EndSection();
@@ -762,7 +822,14 @@ void FDMMaterialStageInputMenus::GenerateChangeInputMenu_NewLocalValues(FMenuBui
 				{
 					FScopedTransaction Transaction(LOCTEXT("SetStageInput", "Set Material Stage Input"));
 					Stage->Modify();
-					Stage->ChangeInput_NewLocalValue(InInputIndex, InInputChannel, ValueType, FDMMaterialStageConnectorChannel::WHOLE_CHANNEL);
+
+					UDMMaterialStageInputValue::ChangeStageInput_NewLocalValue(
+						Stage,
+						InInputIndex, 
+						InInputChannel, 
+						ValueType, 
+						FDMMaterialStageConnectorChannel::WHOLE_CHANNEL
+					);
 				})
 			)
 		);
@@ -805,7 +872,14 @@ namespace UE::DynamicMaterialEditor::Private
 						FScopedTransaction Transaction(LOCTEXT("SetStageInput", "Set Material Stage Input"));
 						InStage->Modify();
 						Value->Modify();
-						InStage->ChangeInput_Value(InInputIndex, InInputChannel, Value, InOutputChannel);
+
+						UDMMaterialStageInputValue::ChangeStageInput_Value(
+							InStage, 
+							InInputIndex,
+							InInputChannel, 
+							Value, 
+							InOutputChannel
+						);
 					}
 				})
 			)
@@ -835,7 +909,13 @@ namespace UE::DynamicMaterialEditor::Private
 						FScopedTransaction Transaction(LOCTEXT("SetStageInput", "Set Material Stage Input"));
 						InStage->Modify();
 						Value->Modify();
-						InStage->ChangeInput_Value(InInputIndex, InInputChannel, Value, FDMMaterialStageConnectorChannel::WHOLE_CHANNEL);
+						UDMMaterialStageInputValue::ChangeStageInput_Value(
+							InStage, 
+							InInputIndex, 
+							InInputChannel,
+							Value, 
+							FDMMaterialStageConnectorChannel::WHOLE_CHANNEL
+						);
 					}
 				})
 			),
@@ -857,19 +937,44 @@ namespace UE::DynamicMaterialEditor::Private
 		{
 			if (bInAcceptsWholeChannel)
 			{
-				UE::DynamicMaterialEditor::Private::GenerateChangeInputMenu_GlobalValue_Value_Channels_Impl(InChildMenuBuilder, InMenuName, InStage, InThroughput,
-					InInputIndex, InInputChannel, InValueIndex, InValue);
+				UE::DynamicMaterialEditor::Private::GenerateChangeInputMenu_GlobalValue_Value_Channels_Impl(
+					InChildMenuBuilder, 
+					InMenuName, 
+					InStage, 
+					InThroughput,
+					InInputIndex, 
+					InInputChannel, 
+					InValueIndex, 
+					InValue)
+				;
 			}
 			else
 			{
-				UE::DynamicMaterialEditor::Private::GenerateChangeInputMenu_GlobalValue_Channels_Impl(InChildMenuBuilder, InMenuName, InStage, InThroughput,
-					InInputIndex, InInputChannel, InValueIndex, InValue);
+				UE::DynamicMaterialEditor::Private::GenerateChangeInputMenu_GlobalValue_Channels_Impl(
+					InChildMenuBuilder, 
+					InMenuName, 
+					InStage, 
+					InThroughput,
+					InInputIndex, 
+					InInputChannel, 
+					InValueIndex, 
+					InValue
+				);
 			}
 		}
 		else
 		{
-			UE::DynamicMaterialEditor::Private::GenerateChangeInputMenu_GlobalValue_Channel_Impl(InChildMenuBuilder, InMenuName, InStage, InThroughput,
-				InInputIndex, InInputChannel, InValueIndex, InValue, FDMMaterialStageConnectorChannel::WHOLE_CHANNEL);
+			UE::DynamicMaterialEditor::Private::GenerateChangeInputMenu_GlobalValue_Channel_Impl(
+				InChildMenuBuilder, 
+				InMenuName, 
+				InStage, 
+				InThroughput,
+				InInputIndex, 
+				InInputChannel, 
+				InValueIndex, 
+				InValue, 
+				FDMMaterialStageConnectorChannel::WHOLE_CHANNEL
+			);
 		}
 	}
 }
@@ -930,9 +1035,18 @@ void FDMMaterialStageInputMenus::GenerateChangeInputMenu_GlobalValues(FMenuBuild
 		const bool bAcceptsWholeChannel = InThroughput->CanInputAcceptType(InInputIndex, Values[Index]->GetType());
 		const bool bAcceptsSubChannels = CanAcceptSubChannels(InThroughput, InInputIndex, Values[Index]->GetType());
 
-		UE::DynamicMaterialEditor::Private::GenerateChangeInputMenu_GlobalValue_Impl(InChildMenuBuilder,
-			Values[Index]->GetDescription(), Stage, InThroughput, InInputIndex, InInputChannel, Index, Values[Index],
-			bAcceptsWholeChannel, bAcceptsSubChannels);
+		UE::DynamicMaterialEditor::Private::GenerateChangeInputMenu_GlobalValue_Impl(
+			InChildMenuBuilder,
+			Values[Index]->GetDescription(), 
+			Stage, 
+			InThroughput, 
+			InInputIndex, 
+			InInputChannel, 
+			Index, 
+			Values[Index],
+			bAcceptsWholeChannel, 
+			bAcceptsSubChannels
+		);
 	}
 }
 
@@ -1000,8 +1114,17 @@ void FDMMaterialStageInputMenus::GenerateChangeInputMenu_GlobalValue_Channels(FM
 			const FText ChannelName = UDMValueDefinitionLibrary::GetValueDefinition(Values[InValueIndex]->GetType()).GetChannelName(ChannelIndex + 1);
 			const int32 OutputChannel = UE::DynamicMaterialEditor::Private::ChannelIndexToChannelBit(ChannelIndex + 1);
 
-			UE::DynamicMaterialEditor::Private::GenerateChangeInputMenu_GlobalValue_Channel_Impl(InChildMenuBuilder,
-				ChannelName, Stage, InThroughput, InInputIndex, InInputChannel, InValueIndex, Values[InValueIndex], OutputChannel);
+			UE::DynamicMaterialEditor::Private::GenerateChangeInputMenu_GlobalValue_Channel_Impl(
+				InChildMenuBuilder,
+				ChannelName, 
+				Stage, 
+				InThroughput, 
+				InInputIndex, 
+				InInputChannel, 
+				InValueIndex, 
+				Values[InValueIndex], 
+				OutputChannel
+			);
 		}
 	}
 	InChildMenuBuilder.EndSection();
@@ -1048,7 +1171,14 @@ void FDMMaterialStageInputMenus::GenerateChangeInputMenu_NewGlobalValues(FMenuBu
 				{
 					FScopedTransaction Transaction(LOCTEXT("SetStageInput", "Set Material Stage Input"));
 					Stage->Modify();
-					Stage->ChangeInput_NewValue(InInputIndex, InInputChannel, ValueType, FDMMaterialStageConnectorChannel::WHOLE_CHANNEL);
+
+					UDMMaterialStageInputValue::ChangeStageInput_NewValue(
+						Stage, 
+						InInputIndex, 
+						InInputChannel,
+						ValueType, 
+						FDMMaterialStageConnectorChannel::WHOLE_CHANNEL
+					);
 				})
 			)
 		);
@@ -1128,7 +1258,16 @@ namespace UE::DynamicMaterialEditor::Private
 						FScopedTransaction Transaction(LOCTEXT("SetStageInput", "Set Material Stage Input"));
 						InStage->Modify();
 						Slot->Modify();
-						InStage->ChangeInput_Slot(InInputIndex, InInputChannel, Slot, InSlotProperty, InOutputIndex, InOutputChannel);
+
+						UDMMaterialStageInputSlot::ChangeStageInput_Slot(
+							InStage, 
+							InInputIndex, 
+							InInputChannel,
+							Slot, 
+							InSlotProperty, 
+							InOutputIndex, 
+							InOutputChannel
+						);
 					}
 				})
 			)
@@ -1159,7 +1298,16 @@ namespace UE::DynamicMaterialEditor::Private
 					{
 						FScopedTransaction Transaction(LOCTEXT("SetStageInput", "Set Material Stage Input"));
 						InStage->Modify();
-						InStage->ChangeInput_Slot(InInputIndex, InInputChannel, Slot, InSlotProperty, InOutputIndex, FDMMaterialStageConnectorChannel::WHOLE_CHANNEL);
+
+						UDMMaterialStageInputSlot::ChangeStageInput_Slot(
+							InStage, 
+							InInputIndex, 
+							InInputChannel,
+							Slot, 
+							InSlotProperty, 
+							InOutputIndex, 
+							FDMMaterialStageConnectorChannel::WHOLE_CHANNEL
+						);
 					}
 				})
 			),
@@ -1181,19 +1329,47 @@ namespace UE::DynamicMaterialEditor::Private
 		{
 			if (bInAcceptsWholeChannel)
 			{
-				UE::DynamicMaterialEditor::Private::GenerateChangeInputMenu_Slot_Property_Output_Channel_Channels_Impl(InChildMenuBuilder, InMenuName, InStage, InThroughput,
-					InInputIndex, InInputChannel, InSlot, InSlotProperty, InOutputIndex);
+				UE::DynamicMaterialEditor::Private::GenerateChangeInputMenu_Slot_Property_Output_Channel_Channels_Impl(
+					InChildMenuBuilder, 
+					InMenuName, 
+					InStage, 
+					InThroughput,
+					InInputIndex, 
+					InInputChannel, 
+					InSlot, 
+					InSlotProperty, 
+					InOutputIndex
+				);
 			}
 			else
 			{
-				UE::DynamicMaterialEditor::Private::GenerateChangeInputMenu_Slot_Property_Output_Channels_Impl(InChildMenuBuilder, InMenuName, InStage, InThroughput,
-					InInputIndex, InInputChannel, InSlot, InSlotProperty, InOutputIndex);
+				UE::DynamicMaterialEditor::Private::GenerateChangeInputMenu_Slot_Property_Output_Channels_Impl(
+					InChildMenuBuilder, 
+					InMenuName, 
+					InStage, 
+					InThroughput,
+					InInputIndex, 
+					InInputChannel, 
+					InSlot, 
+					InSlotProperty, 
+					InOutputIndex
+				);
 			}
 		}
 		else
 		{
-			UE::DynamicMaterialEditor::Private::GenerateChangeInputMenu_Slot_Property_Output_Channel_Impl(InChildMenuBuilder, InMenuName, InStage, InThroughput,
-				InInputIndex, InInputChannel, InSlot, InSlotProperty, InOutputIndex, FDMMaterialStageConnectorChannel::WHOLE_CHANNEL);
+			UE::DynamicMaterialEditor::Private::GenerateChangeInputMenu_Slot_Property_Output_Channel_Impl(
+				InChildMenuBuilder, 
+				InMenuName, 
+				InStage, 
+				InThroughput,
+				InInputIndex, 
+				InInputChannel, 
+				InSlot, 
+				InSlotProperty, 
+				InOutputIndex, 
+				FDMMaterialStageConnectorChannel::WHOLE_CHANNEL
+			);
 		}
 	}
 }
@@ -1326,19 +1502,45 @@ void FDMMaterialStageInputMenus::GenerateChangeInputMenu_Slots(FMenuBuilder& InC
 				const bool bAcceptsWholeChannel = InThroughput->CanInputConnectTo(InInputIndex, LastPropertySourceOutputConnectors[0], FDMMaterialStageConnectorChannel::WHOLE_CHANNEL);
 				const bool bAcceptsSubChannels = CanAcceptSubChannels(InThroughput, InInputIndex, LastPropertySourceOutputConnectors[0]);
 
-				UE::DynamicMaterialEditor::Private::GenerateChangeInputMenu_Slot_Property_Output_Impl(InChildMenuBuilder, SlotPropertyName,
-					Stage, InThroughput, InInputIndex, InInputChannel, SlotIter, SlotProperties[0], 0, bAcceptsWholeChannel, bAcceptsSubChannels);
+				UE::DynamicMaterialEditor::Private::GenerateChangeInputMenu_Slot_Property_Output_Impl(
+					InChildMenuBuilder, 
+					SlotPropertyName,
+					Stage,
+					InThroughput, 
+					InInputIndex, 
+					InInputChannel, 
+					SlotIter, 
+					SlotProperties[0], 
+					0, 
+					bAcceptsWholeChannel, 
+					bAcceptsSubChannels
+				);
 			}
 			else // !bOneOutputType
 			{
-				UE::DynamicMaterialEditor::Private::GenerateChangeInputMenu_Slot_Property_Outputs_Impl(InChildMenuBuilder, SlotPropertyName,
-					Stage, InThroughput, InInputIndex, InInputChannel, SlotIter, SlotProperties[0]);
+				UE::DynamicMaterialEditor::Private::GenerateChangeInputMenu_Slot_Property_Outputs_Impl(
+					InChildMenuBuilder, 
+					SlotPropertyName,
+					Stage, 
+					InThroughput, 
+					InInputIndex, 
+					InInputChannel, 
+					SlotIter, 
+					SlotProperties[0]
+				);
 			}
 		}
 		else
 		{
-			UE::DynamicMaterialEditor::Private::GenerateChangeInputMenu_Slot_Properties_Impl(InChildMenuBuilder, SlotIter->GetDescription(),
-				Stage, InThroughput, InInputIndex, InInputChannel, SlotIter);
+			UE::DynamicMaterialEditor::Private::GenerateChangeInputMenu_Slot_Properties_Impl(
+				InChildMenuBuilder, 
+				SlotIter->GetDescription(),
+				Stage, 
+				InThroughput, 
+				InInputIndex, 
+				InInputChannel, 
+				SlotIter
+			);
 		}
 	}
 }
@@ -1446,13 +1648,32 @@ void FDMMaterialStageInputMenus::GenerateChangeInputMenu_Slot_Properties(FMenuBu
 				const bool bAcceptsWholeChannel = InThroughput->CanInputConnectTo(InInputIndex, LastPropertySourceOutputConnectors[0], FDMMaterialStageConnectorChannel::WHOLE_CHANNEL);
 				const bool bAcceptsSubChannels = CanAcceptSubChannels(InThroughput, InInputIndex, LastPropertySourceOutputConnectors[0]);
 
-				UE::DynamicMaterialEditor::Private::GenerateChangeInputMenu_Slot_Property_Output_Impl(InChildMenuBuilder, PropertyObj->GetDescription(),
-					Stage, InThroughput, InInputIndex, InInputChannel, InSlot, SlotProperty, 0, bAcceptsWholeChannel, bAcceptsSubChannels);
+				UE::DynamicMaterialEditor::Private::GenerateChangeInputMenu_Slot_Property_Output_Impl(
+					InChildMenuBuilder, 
+					PropertyObj->GetDescription(),
+					Stage, 
+					InThroughput, 
+					InInputIndex, 
+					InInputChannel, 
+					InSlot, 
+					SlotProperty, 
+					0, 
+					bAcceptsWholeChannel, 
+					bAcceptsSubChannels
+				);
 			}
 			else // !bOneOutputType
 			{
-				UE::DynamicMaterialEditor::Private::GenerateChangeInputMenu_Slot_Property_Outputs_Impl(InChildMenuBuilder, PropertyObj->GetDescription(),
-					Stage, InThroughput, InInputIndex, InInputChannel, InSlot, SlotProperty);
+				UE::DynamicMaterialEditor::Private::GenerateChangeInputMenu_Slot_Property_Outputs_Impl(
+					InChildMenuBuilder, 
+					PropertyObj->GetDescription(),
+					Stage, 
+					InThroughput, 
+					InInputIndex, 
+					InInputChannel, 
+					InSlot, 
+					SlotProperty
+				);
 			}
 		}
 	}
@@ -1543,8 +1764,19 @@ void FDMMaterialStageInputMenus::GenerateChangeInputMenu_Slot_Property_Outputs(F
 			const bool bAcceptsWholeChannel = InThroughput->CanInputConnectTo(InInputIndex, LastPropertySourceOutputConnectors[OutputIdx], FDMMaterialStageConnectorChannel::WHOLE_CHANNEL);
 			const bool bAcceptsSubChannels = CanAcceptSubChannels(InThroughput, InInputIndex, LastPropertySourceOutputConnectors[OutputIdx]);
 
-			UE::DynamicMaterialEditor::Private::GenerateChangeInputMenu_Slot_Property_Output_Impl(InChildMenuBuilder, LastPropertySourceOutputConnectors[OutputIdx].Name,
-				Stage, InThroughput, InInputIndex, InInputChannel, InSlot, InMaterialProperty, OutputIdx, bAcceptsWholeChannel, bAcceptsSubChannels);
+			UE::DynamicMaterialEditor::Private::GenerateChangeInputMenu_Slot_Property_Output_Impl(
+				InChildMenuBuilder, 
+				LastPropertySourceOutputConnectors[OutputIdx].Name,
+				Stage, 
+				InThroughput, 
+				InInputIndex, 
+				InInputChannel, 
+				InSlot, 
+				InMaterialProperty, 
+				OutputIdx, 
+				bAcceptsWholeChannel, 
+				bAcceptsSubChannels
+			);
 		}
 	}
 	InChildMenuBuilder.EndSection();
@@ -1636,8 +1868,18 @@ void FDMMaterialStageInputMenus::GenerateChangeInputMenu_Slot_Property_Output_Ch
 			const FText ChannelName = UDMValueDefinitionLibrary::GetValueDefinition(LastPropertySourceOutputConnectors[OutputIndex].Type).GetChannelName(ChannelIndex + 1);
 			const int32 OutputChannel = UE::DynamicMaterialEditor::Private::ChannelIndexToChannelBit(ChannelIndex + 1);
 
-			UE::DynamicMaterialEditor::Private::GenerateChangeInputMenu_Slot_Property_Output_Channel_Impl(InChildMenuBuilder, ChannelName,
-				Stage, InThroughput, InInputIndex, InInputChannel, InSlot, InMaterialProperty, OutputIndex, OutputChannel);
+			UE::DynamicMaterialEditor::Private::GenerateChangeInputMenu_Slot_Property_Output_Channel_Impl(
+				InChildMenuBuilder, 
+				ChannelName,
+				Stage, 
+				InThroughput, 
+				InInputIndex, 
+				InInputChannel,
+				InSlot, 
+				InMaterialProperty, 
+				OutputIndex, 
+				OutputChannel
+			);
 		}
 	}
 	InChildMenuBuilder.EndSection();
@@ -1735,7 +1977,13 @@ void FDMMaterialStageInputMenus::GenerateChangeInputMenu_Expressions(FMenuBuilde
 
 					for (UDMMaterialStageExpression* ExpressionCDO : InSubmenuExpressionList)
 					{
-						GenerateChangeInputMenu_Expression(ChildSubMenuBuilder, ExpressionCDO, InThroughput, InInputIndex, InInputChannel);
+						GenerateChangeInputMenu_Expression(
+							ChildSubMenuBuilder, 
+							ExpressionCDO, 
+							InThroughput, 
+							InInputIndex, 
+							InInputChannel
+						);
 					}
 
 					ChildSubMenuBuilder.EndSection();
@@ -1797,7 +2045,14 @@ namespace UE::DynamicMaterialEditor::Private
 				{
 					FScopedTransaction Transaction(LOCTEXT("SetStageInput", "Set Material Stage Input"));
 					InStage->Modify();
-					InStage->ChangeInput_Expression(InInputIndex, InInputChannel, InExpressionClass, InOutputIndex, InOutputChannel);
+					UDMMaterialStageInputExpression::ChangeStageInput_Expression(
+						InStage, 
+						InExpressionClass, 
+						InInputIndex,
+						InInputChannel, 
+						InOutputIndex, 
+						InOutputChannel
+					);
 				})
 			)
 		);
@@ -1824,7 +2079,14 @@ namespace UE::DynamicMaterialEditor::Private
 				{
 					FScopedTransaction Transaction(LOCTEXT("SetStageInput", "Set Material Stage Input"));
 					InStage->Modify();
-					InStage->ChangeInput_Expression(InInputIndex, InInputChannel, InExpressionClass, InOutputIndex, FDMMaterialStageConnectorChannel::WHOLE_CHANNEL);
+					UDMMaterialStageInputExpression::ChangeStageInput_Expression(
+						InStage, 
+						InExpressionClass, 
+						InInputIndex, 
+						InInputChannel, 
+						InOutputIndex,
+						FDMMaterialStageConnectorChannel::WHOLE_CHANNEL
+					);
 				})
 			),
 			NAME_None,
@@ -1845,19 +2107,44 @@ namespace UE::DynamicMaterialEditor::Private
 		{
 			if (bInAcceptsWholeChannel)
 			{
-				UE::DynamicMaterialEditor::Private::GenerateChangeInputMenu_Expression_Output_Channel_Channels_Impl(InChildMenuBuilder, InMenuName, InStage, InThroughput,
-					InInputIndex, InInputChannel, InExpressionClass, InOutputIndex);
+				UE::DynamicMaterialEditor::Private::GenerateChangeInputMenu_Expression_Output_Channel_Channels_Impl(
+					InChildMenuBuilder, 
+					InMenuName, 
+					InStage, 
+					InThroughput,
+					InInputIndex, 
+					InInputChannel, 
+					InExpressionClass, 
+					InOutputIndex
+				);
 			}
 			else
 			{
-				UE::DynamicMaterialEditor::Private::GenerateChangeInputMenu_Expression_Output_Channels_Impl(InChildMenuBuilder, InMenuName, InStage, InThroughput,
-					InInputIndex, InInputChannel, InExpressionClass, InOutputIndex);
+				UE::DynamicMaterialEditor::Private::GenerateChangeInputMenu_Expression_Output_Channels_Impl(
+					InChildMenuBuilder, 
+					InMenuName, 
+					InStage, 
+					InThroughput,
+					InInputIndex, 
+					InInputChannel, 
+					InExpressionClass, 
+					InOutputIndex
+				);
 			}
 		}
 		else
 		{
-			UE::DynamicMaterialEditor::Private::GenerateChangeInputMenu_Expression_Output_Channel_Impl(InChildMenuBuilder, InMenuName, InStage, InThroughput,
-				InInputIndex, InInputChannel, InExpressionClass, InOutputIndex, FDMMaterialStageConnectorChannel::WHOLE_CHANNEL);
+			UE::DynamicMaterialEditor::Private::GenerateChangeInputMenu_Expression_Output_Channel_Impl(
+				InChildMenuBuilder, 
+				InMenuName, 
+				InStage, 
+				InThroughput,
+				InInputIndex, 
+				InInputChannel, 
+				InExpressionClass, 
+				InOutputIndex, 
+				FDMMaterialStageConnectorChannel::WHOLE_CHANNEL
+			);
 		}
 	}
 }
@@ -1911,13 +2198,30 @@ void FDMMaterialStageInputMenus::GenerateChangeInputMenu_Expression(FMenuBuilder
 		const bool bAcceptsWholeChannel = InThroughput->CanInputConnectTo(InInputIndex, OutputConnectors[LastValidOutputConnectorIdx], FDMMaterialStageConnectorChannel::WHOLE_CHANNEL);
 		const bool bAcceptsSubChannels = CanAcceptSubChannels(InThroughput, InInputIndex, OutputConnectors[LastValidOutputConnectorIdx]);
 
-		UE::DynamicMaterialEditor::Private::GenerateChangeInputMenu_Expression_Output_Impl(InChildMenuBuilder, ExpressionName, Stage, InThroughput,
-			InInputIndex, InInputChannel, InputClass, LastValidOutputConnectorIdx, bAcceptsWholeChannel, bAcceptsSubChannels);
+		UE::DynamicMaterialEditor::Private::GenerateChangeInputMenu_Expression_Output_Impl(
+			InChildMenuBuilder, 
+			ExpressionName, 
+			Stage, 
+			InThroughput,
+			InInputIndex, 
+			InInputChannel, 
+			InputClass, 
+			LastValidOutputConnectorIdx, 
+			bAcceptsWholeChannel, 
+			bAcceptsSubChannels
+		);
 	}
 	else
 	{
-		UE::DynamicMaterialEditor::Private::GenerateChangeInputMenu_Expression_Outputs_Impl(InChildMenuBuilder, InExpressionCDO->GetDescription(), 
-			Stage, InThroughput, InInputIndex, InInputChannel, InputClass);
+		UE::DynamicMaterialEditor::Private::GenerateChangeInputMenu_Expression_Outputs_Impl(
+			InChildMenuBuilder, 
+			InExpressionCDO->GetDescription(), 
+			Stage, 
+			InThroughput, 
+			InInputIndex, 
+			InInputChannel, 
+			InputClass
+		);
 	}
 }
 
@@ -1961,8 +2265,18 @@ void FDMMaterialStageInputMenus::GenerateChangeInputMenu_Expression_Outputs(FMen
 			const bool bAcceptsWholeChannel = InThroughput->CanInputConnectTo(InInputIndex, OutputConnectors[OutputIdx], FDMMaterialStageConnectorChannel::WHOLE_CHANNEL);
 			const bool bAcceptsSubChannels = CanAcceptSubChannels(InThroughput, InInputIndex, OutputConnectors[OutputIdx]);
 
-			UE::DynamicMaterialEditor::Private::GenerateChangeInputMenu_Expression_Output_Impl(InChildMenuBuilder, OutputConnectors[OutputIdx].Name,
-				Stage, InThroughput, InInputIndex, InInputChannel, InExpressionClass, OutputIdx, bAcceptsWholeChannel, bAcceptsSubChannels);
+			UE::DynamicMaterialEditor::Private::GenerateChangeInputMenu_Expression_Output_Impl(
+				InChildMenuBuilder, 
+				OutputConnectors[OutputIdx].Name,
+				Stage, 
+				InThroughput, 
+				InInputIndex, 
+				InInputChannel, 
+				InExpressionClass, 
+				OutputIdx, 
+				bAcceptsWholeChannel, 
+				bAcceptsSubChannels
+			);
 		}
 	}
 	InChildMenuBuilder.EndSection();
@@ -2024,8 +2338,17 @@ void FDMMaterialStageInputMenus::GenerateChangeInputMenu_Expression_Output_Chann
 			const FText ChannelName = UDMValueDefinitionLibrary::GetValueDefinition(OutputConnectors[InOutputIndex].Type).GetChannelName(ChannelIndex + 1);
 			const int32 OutputChannel = UE::DynamicMaterialEditor::Private::ChannelIndexToChannelBit(ChannelIndex + 1);
 
-			UE::DynamicMaterialEditor::Private::GenerateChangeInputMenu_Expression_Output_Channel_Impl(InChildMenuBuilder, ChannelName, Stage, InThroughput,
-				InInputIndex, InInputChannel, InExpressionClass, InOutputIndex, OutputChannel);
+			UE::DynamicMaterialEditor::Private::GenerateChangeInputMenu_Expression_Output_Channel_Impl(
+				InChildMenuBuilder, 
+				ChannelName, 
+				Stage, 
+				InThroughput,
+				InInputIndex, 
+				InInputChannel, 
+				InExpressionClass, 
+				InOutputIndex, 
+				OutputChannel
+			);
 		}
 	}
 	InChildMenuBuilder.EndSection();
@@ -2061,7 +2384,13 @@ namespace UE::DynamicMaterialEditor::Private
 				{
 					FScopedTransaction Transaction(LOCTEXT("SetStageInput", "Set Material Stage Input"));
 					InStage->Modify();
-					InStage->ChangeInput_UV(InInputIndex, InInputChannel, InOutputChannel);
+
+					UDMMaterialStageInputTextureUV::ChangeStageInput_UV(
+						InStage, 
+						InInputIndex, 
+						InInputChannel, 
+						InOutputChannel
+					);
 				})
 			)
 		);
@@ -2085,7 +2414,13 @@ namespace UE::DynamicMaterialEditor::Private
 				{
 					FScopedTransaction Transaction(LOCTEXT("SetStageInput", "Set Material Stage Input"));
 					InStage->Modify();
-					InStage->ChangeInput_UV(InInputIndex, InInputChannel, FDMMaterialStageConnectorChannel::WHOLE_CHANNEL);
+
+					UDMMaterialStageInputTextureUV::ChangeStageInput_UV(
+						InStage, 
+						InInputIndex, 
+						InInputChannel, 
+						FDMMaterialStageConnectorChannel::WHOLE_CHANNEL
+					);
 				})
 			),
 			NAME_None,
@@ -2106,19 +2441,38 @@ namespace UE::DynamicMaterialEditor::Private
 		{
 			if (bInAcceptsWholeChannel)
 			{
-				UE::DynamicMaterialEditor::Private::GenerateChangeInputMenu_UV_UV_Channels_Impl(InChildMenuBuilder, InMenuName, InStage, InThroughput,
-					InInputIndex, InInputChannel);
+				UE::DynamicMaterialEditor::Private::GenerateChangeInputMenu_UV_UV_Channels_Impl(
+					InChildMenuBuilder, 
+					InMenuName, 
+					InStage, 
+					InThroughput,
+					InInputIndex, 
+					InInputChannel
+				);
 			}
 			else
 			{
-				UE::DynamicMaterialEditor::Private::GenerateChangeInputMenu_UV_Channels_Impl(InChildMenuBuilder, InMenuName, InStage, InThroughput,
-					InInputIndex, InInputChannel);
+				UE::DynamicMaterialEditor::Private::GenerateChangeInputMenu_UV_Channels_Impl(
+					InChildMenuBuilder, 
+					InMenuName, 
+					InStage, 
+					InThroughput,
+					InInputIndex, 
+					InInputChannel
+				);
 			}
 		}
 		else
 		{
-			UE::DynamicMaterialEditor::Private::GenerateChangeInputMenu_UV_Channel_Impl(InChildMenuBuilder, InMenuName, InStage, InThroughput,
-				InInputIndex, InInputChannel, FDMMaterialStageConnectorChannel::WHOLE_CHANNEL);
+			UE::DynamicMaterialEditor::Private::GenerateChangeInputMenu_UV_Channel_Impl(
+				InChildMenuBuilder, 
+				InMenuName, 
+				InStage, 
+				InThroughput,
+				InInputIndex, 
+				InInputChannel, 
+				FDMMaterialStageConnectorChannel::WHOLE_CHANNEL
+			);
 		}
 	}
 }
@@ -2146,8 +2500,16 @@ void FDMMaterialStageInputMenus::GenerateChangeInputMenu_UV(FMenuBuilder& InChil
 	const bool bCompatibleWithTextureUV = InThroughput->CanInputAcceptType(InInputIndex, EDMValueType::VT_Float2);
 	const bool bCompatibleWithTextureUorV = InThroughput->CanInputAcceptType(InInputIndex, EDMValueType::VT_Float1);
 
-	UE::DynamicMaterialEditor::Private::GenerateChangeInputMenu_UV_Impl(InChildMenuBuilder, LOCTEXT("ChangeInputExpressionUV", "Texture UV"),
-		Stage, InThroughput, InInputIndex, InInputChannel, bCompatibleWithTextureUV, bCompatibleWithTextureUorV);
+	UE::DynamicMaterialEditor::Private::GenerateChangeInputMenu_UV_Impl(
+		InChildMenuBuilder, 
+		LOCTEXT("ChangeInputExpressionUV", "Texture UV"),
+		Stage, 
+		InThroughput, 
+		InInputIndex, 
+		InInputChannel, 
+		bCompatibleWithTextureUV, 
+		bCompatibleWithTextureUorV
+	);
 }
 
 void FDMMaterialStageInputMenus::GenerateChangeInputMenu_UV_Channels(FMenuBuilder& InChildMenuBuilder, UDMMaterialStageThroughput* InThroughput, 
@@ -2173,11 +2535,25 @@ void FDMMaterialStageInputMenus::GenerateChangeInputMenu_UV_Channels(FMenuBuilde
 
 	InChildMenuBuilder.BeginSection("ChangeInputOutputChannels", LOCTEXT("ChangeInputOutputChannel", "Output Channels"));
 	{
-		UE::DynamicMaterialEditor::Private::GenerateChangeInputMenu_UV_Channel_Impl(InChildMenuBuilder,
-			LOCTEXT("ChangeInputExpressionUVU", "U"), Stage, InThroughput, InInputIndex, InInputChannel, 0);
+		UE::DynamicMaterialEditor::Private::GenerateChangeInputMenu_UV_Channel_Impl(
+			InChildMenuBuilder,
+			LOCTEXT("ChangeInputExpressionUVU", "U"), 
+			Stage, 
+			InThroughput,
+			InInputIndex, 
+			InInputChannel, 
+			0
+		);
 
-		UE::DynamicMaterialEditor::Private::GenerateChangeInputMenu_UV_Channel_Impl(InChildMenuBuilder,
-			LOCTEXT("ChangeInputExpressionUVV", "V"), Stage, InThroughput, InInputIndex, InInputChannel, 1);
+		UE::DynamicMaterialEditor::Private::GenerateChangeInputMenu_UV_Channel_Impl(
+			InChildMenuBuilder,
+			LOCTEXT("ChangeInputExpressionUVV", "V"), 
+			Stage, 
+			InThroughput, 
+			InInputIndex, 
+			InInputChannel, 
+			1
+		);
 	}
 	InChildMenuBuilder.EndSection();
 }
@@ -2219,7 +2595,13 @@ void FDMMaterialStageInputMenus::GenerateChangeInputMenu_Gradients(FMenuBuilder&
 			continue;;
 		}
 
-		GenerateChangeInputMenu_Gradient(InChildMenuBuilder, GradientCDO, InThroughput, InInputIndex, InInputChannel);
+		GenerateChangeInputMenu_Gradient(
+			InChildMenuBuilder, 
+			GradientCDO, 
+			InThroughput, 
+			InInputIndex, 
+			InInputChannel
+		);
 	}
 }
 
@@ -2249,7 +2631,13 @@ void FDMMaterialStageInputMenus::GenerateChangeInputMenu_Gradient(FMenuBuilder&
 			{
 				FScopedTransaction Transaction(LOCTEXT("SetStageInput", "Set Material Stage Input"));
 				Stage->Modify();
-				Stage->ChangeInput_Gradient(InInputIndex, InInputChannel, GradientClass, FDMMaterialStageConnectorChannel::WHOLE_CHANNEL);
+				UDMMaterialStageInputGradient::ChangeStageInput_Gradient(
+					Stage, 
+					GradientClass, 
+					InInputIndex,
+					InInputChannel, 
+					FDMMaterialStageConnectorChannel::WHOLE_CHANNEL
+				);
 			})
 		)
 	);
diff --git a/Engine/Plugins/Experimental/DynamicMaterial/Source/DynamicMaterialEditor/Private/Menus/DMMaterialStageSourceMenus.cpp b/Engine/Plugins/Experimental/DynamicMaterial/Source/DynamicMaterialEditor/Private/Menus/DMMaterialStageSourceMenus.cpp
--- a/Engine/Plugins/Experimental/DynamicMaterial/Source/DynamicMaterialEditor/Private/Menus/DMMaterialStageSourceMenus.cpp
+++ b/Engine/Plugins/Experimental/DynamicMaterial/Source/DynamicMaterialEditor/Private/Menus/DMMaterialStageSourceMenus.cpp
@@ -7,6 +7,7 @@
 #include "Components/DMMaterialStage.h"
 #include "Components/DMMaterialStageBlend.h"
 #include "Components/DMMaterialStageExpression.h"
+#include "Components/DMMaterialStageFunction.h"
 #include "Components/DMMaterialStageGradient.h"
 #include "Components/DMMaterialStageSource.h"
 #include "Components/DMMaterialStageThroughput.h"
@@ -15,9 +16,15 @@
 #include "Components/MaterialStageExpressions/DMMSESceneTexture.h"
 #include "Components/MaterialStageExpressions/DMMSETextureSample.h"
 #include "Components/MaterialStageExpressions/DMMSETextureSampleEdgeColor.h"
+#include "Components/MaterialStageInputs/DMMSIExpression.h"
+#include "Components/MaterialStageInputs/DMMSIFunction.h"
+#include "Components/MaterialStageInputs/DMMSIGradient.h"
+#include "Components/MaterialStageInputs/DMMSISlot.h"
+#include "Components/MaterialStageInputs/DMMSIValue.h"
 #include "DMDefs.h"
 #include "DMValueDefinition.h"
 #include "DynamicMaterialEditorModule.h"
+#include "Materials/MaterialFunctionInterface.h"
 #include "Menus/DMMenuContext.h"
 #include "Model/DynamicMaterialModel.h"
 #include "Model/DynamicMaterialModelEditorOnlyData.h"
@@ -91,7 +98,14 @@ namespace UE::DynamicMaterialEditor::Private
 
 								FScopedTransaction Transaction(LOCTEXT("SetStageInputBase", "Set Material Designer Base Source"));
 								Stage->Modify();
-								Stage->ChangeInput_NewLocalValue(UDMMaterialStageBlend::InputB, FDMMaterialStageConnectorChannel::WHOLE_CHANNEL, ValueType, OutputChannel);
+
+								UDMMaterialStageInputValue::ChangeStageInput_NewLocalValue(
+									Stage, 
+									UDMMaterialStageBlend::InputB,
+									FDMMaterialStageConnectorChannel::WHOLE_CHANNEL, 
+									ValueType, 
+									OutputChannel
+								);
 							}
 							else if (StageSource->IsA<UDMMaterialStageThroughputLayerBlend>())
 							{
@@ -101,13 +115,20 @@ namespace UE::DynamicMaterialEditor::Private
 
 								FScopedTransaction Transaction(LOCTEXT("SetStageInputMask", "Set Material Designer Mask Source"));
 								Stage->Modify();
-								Stage->ChangeInput_NewLocalValue(UDMMaterialStageThroughputLayerBlend::InputMaskSource, FDMMaterialStageConnectorChannel::WHOLE_CHANNEL, ValueType, OutputChannel);
+
+								UDMMaterialStageInputValue::ChangeStageInput_NewLocalValue(
+									Stage, 
+									UDMMaterialStageThroughputLayerBlend::InputMaskSource,
+									FDMMaterialStageConnectorChannel::WHOLE_CHANNEL,
+									ValueType,
+									OutputChannel
+								);
 							}
 							else
 							{
 								FScopedTransaction Transaction(LOCTEXT("SetStageSource", "Set Material Designer Stage Source"));
 								Stage->Modify();
-								Stage->ChangeSource_NewLocalValue(ValueType);
+								UDMMaterialStageInputValue::ChangeStageSource_NewLocalValue(Stage, ValueType);
 							}
 
 							if (TSharedPtr<SDMSlot> SlotWidget = MenuContext->GetSlotWidget().Pin())
@@ -181,8 +202,14 @@ namespace UE::DynamicMaterialEditor::Private
 
 							FScopedTransaction Transaction(LOCTEXT("SetStageInputBase", "Set Material Designer Base Source"));
 							Stage->Modify();
-							Stage->ChangeInput_Value(UDMMaterialStageBlend::InputB, FDMMaterialStageConnectorChannel::WHOLE_CHANNEL,
-								Value, OutputChannel);
+
+							UDMMaterialStageInputValue::ChangeStageInput_Value(
+								Stage, 
+								UDMMaterialStageBlend::InputB,
+								FDMMaterialStageConnectorChannel::WHOLE_CHANNEL, 
+								Value, 
+								OutputChannel
+							);
 						}
 						else if (StageSource->IsA<UDMMaterialStageThroughputLayerBlend>())
 						{
@@ -192,14 +219,20 @@ namespace UE::DynamicMaterialEditor::Private
 
 							FScopedTransaction Transaction(LOCTEXT("SetStageInputMask", "Set Material Designer Mask Source"));
 							Stage->Modify();
-							Stage->ChangeInput_Value(UDMMaterialStageThroughputLayerBlend::InputMaskSource, FDMMaterialStageConnectorChannel::WHOLE_CHANNEL,
-								Value, OutputChannel);
+
+							UDMMaterialStageInputValue::ChangeStageInput_Value(
+								Stage, 
+								UDMMaterialStageThroughputLayerBlend::InputMaskSource, 
+								FDMMaterialStageConnectorChannel::WHOLE_CHANNEL,
+								Value, 
+								OutputChannel
+							);
 						}
 						else
 						{
 							FScopedTransaction Transaction(LOCTEXT("SetStageSource", "Set Material Designer Stage Source"));
 							Stage->Modify();
-							Stage->ChangeSource_Value(Value);
+							UDMMaterialStageInputValue::ChangeStageSource_Value(Stage, Value);
 						}
 
 						if (TSharedPtr<SDMSlot> SlotWidget = MenuContext->GetSlotWidget().Pin())
@@ -267,7 +300,14 @@ namespace UE::DynamicMaterialEditor::Private
 
 								FScopedTransaction Transaction(LOCTEXT("SetStageInputBase", "Set Material Designer Base Source"));
 								Stage->Modify();
-								Stage->ChangeInput_NewValue(UDMMaterialStageBlend::InputB, FDMMaterialStageConnectorChannel::WHOLE_CHANNEL, ValueType, OutputChannel);
+
+								UDMMaterialStageInputValue::ChangeStageInput_NewValue(
+									Stage, 
+									UDMMaterialStageBlend::InputB, 
+									FDMMaterialStageConnectorChannel::WHOLE_CHANNEL, 
+									ValueType, 
+									OutputChannel
+								);
 							}
 							else if (StageSource->IsA<UDMMaterialStageThroughputLayerBlend>())
 							{
@@ -277,13 +317,20 @@ namespace UE::DynamicMaterialEditor::Private
 
 								FScopedTransaction Transaction(LOCTEXT("SetStageInputMask", "Set Material Designer Mask Source"));
 								Stage->Modify();
-								Stage->ChangeInput_NewValue(UDMMaterialStageThroughputLayerBlend::InputMaskSource, FDMMaterialStageConnectorChannel::WHOLE_CHANNEL, ValueType, OutputChannel);
+
+								UDMMaterialStageInputValue::ChangeStageInput_NewValue(
+									Stage, 
+									UDMMaterialStageThroughputLayerBlend::InputMaskSource,
+									FDMMaterialStageConnectorChannel::WHOLE_CHANNEL,
+									ValueType,
+									 OutputChannel
+									);
 							}
 							else
 							{
 								FScopedTransaction Transaction(LOCTEXT("SetStageSource", "Set Material Designer Stage Source"));
 								Stage->Modify();
-								Stage->ChangeSource_NewValue(ValueType);
+								UDMMaterialStageInputValue::ChangeStageSource_NewValue(Stage, ValueType);
 							}
 
 							if (TSharedPtr<SDMSlot> SlotWidget = MenuContext->GetSlotWidget().Pin())
@@ -380,21 +427,37 @@ namespace UE::DynamicMaterialEditor::Private
 							{
 								FScopedTransaction Transaction(LOCTEXT("SetStageInputBase", "Set Material Designer Base Source"));
 								Stage->Modify();
-								Stage->ChangeInput_Slot(UDMMaterialStageBlend::InputB, FDMMaterialStageConnectorChannel::WHOLE_CHANNEL,
-									Slot, Property, 0, FDMMaterialStageConnectorChannel::WHOLE_CHANNEL);
+
+								UDMMaterialStageInputSlot::ChangeStageInput_Slot(
+									Stage, 
+									UDMMaterialStageBlend::InputB,
+									FDMMaterialStageConnectorChannel::WHOLE_CHANNEL, 
+									Slot, 
+									Property,
+									0, 
+									FDMMaterialStageConnectorChannel::WHOLE_CHANNEL
+								);
 							}
 							else if (StageSource->IsA<UDMMaterialStageThroughputLayerBlend>())
 							{
 								FScopedTransaction Transaction(LOCTEXT("SetStageInputMask", "Set Material Designer Mask Source"));
 								Stage->Modify();
-								Stage->ChangeInput_Slot(UDMMaterialStageThroughputLayerBlend::InputMaskSource, FDMMaterialStageConnectorChannel::WHOLE_CHANNEL,
-									Slot, Property, 0, FDMMaterialStageConnectorChannel::WHOLE_CHANNEL);
+
+								UDMMaterialStageInputSlot::ChangeStageInput_Slot(
+									Stage, 
+									UDMMaterialStageThroughputLayerBlend::InputMaskSource,
+									FDMMaterialStageConnectorChannel::WHOLE_CHANNEL,
+									Slot, 
+									Property, 
+									0,
+									FDMMaterialStageConnectorChannel::WHOLE_CHANNEL
+								);
 							}
 							else
 							{
 								FScopedTransaction Transaction(LOCTEXT("SetStageSource", "Set Material Designer Stage Source"));
 								Stage->Modify();
-								Stage->ChangeSource_Slot(Slot, Property);
+								UDMMaterialStageInputSlot::ChangeStageSource_Slot(Stage, Slot, Property);
 							}
 
 							if (TSharedPtr<SDMSlot> SlotWidget = MenuContext->GetSlotWidget().Pin())
@@ -512,21 +575,37 @@ namespace UE::DynamicMaterialEditor::Private
 								{
 									FScopedTransaction Transaction(LOCTEXT("SetStageInputBase", "Set Material Designer Base Source"));
 									Stage->Modify();
-									Stage->ChangeInput_Slot(UDMMaterialStageBlend::InputB, FDMMaterialStageConnectorChannel::WHOLE_CHANNEL,
-										SlotIter, SlotProperty, 0, FDMMaterialStageConnectorChannel::WHOLE_CHANNEL);
+
+									UDMMaterialStageInputSlot::ChangeStageInput_Slot(
+										Stage, 
+										UDMMaterialStageBlend::InputB,
+										FDMMaterialStageConnectorChannel::WHOLE_CHANNEL,
+										SlotIter, 
+										SlotProperty, 
+										0, 
+										FDMMaterialStageConnectorChannel::WHOLE_CHANNEL
+									);
 								}
 								else if (StageSource->IsA<UDMMaterialStageThroughputLayerBlend>())
 								{
 									FScopedTransaction Transaction(LOCTEXT("SetStageInputMask", "Set Material Designer Mask Source"));
 									Stage->Modify();
-									Stage->ChangeInput_Slot(UDMMaterialStageThroughputLayerBlend::InputMaskSource, FDMMaterialStageConnectorChannel::WHOLE_CHANNEL,
-										SlotIter, SlotProperty, 0, FDMMaterialStageConnectorChannel::WHOLE_CHANNEL);
+
+									UDMMaterialStageInputSlot::ChangeStageInput_Slot(
+										Stage, 
+										UDMMaterialStageThroughputLayerBlend::InputMaskSource,
+										FDMMaterialStageConnectorChannel::WHOLE_CHANNEL,
+										SlotIter, 
+										SlotProperty, 
+										0,
+										FDMMaterialStageConnectorChannel::WHOLE_CHANNEL
+									);
 								}
 								else
 								{
 									FScopedTransaction Transaction(LOCTEXT("SetStageSource", "Set Material Designer Stage Source"));
 									Stage->Modify();
-									Stage->ChangeSource_Slot(SlotIter, SlotProperty);
+									UDMMaterialStageInputSlot::ChangeStageSource_Slot(Stage, SlotIter, SlotProperty);
 								}
 
 								if (TSharedPtr<SDMSlot> SlotWidget = MenuContext->GetSlotWidget().Pin())
@@ -607,19 +686,31 @@ namespace UE::DynamicMaterialEditor::Private
 							{
 								FScopedTransaction Transaction(LOCTEXT("SetStageInputBase", "Set Material Designer Base Source"));
 								Stage->Modify();
-								Stage->ChangeInput_Gradient(UDMMaterialStageBlend::InputB, FDMMaterialStageConnectorChannel::WHOLE_CHANNEL, GradientClass.Get(), FDMMaterialStageConnectorChannel::WHOLE_CHANNEL);
+
+								UDMMaterialStageInputGradient::ChangeStageInput_Gradient(
+									Stage,
+									GradientClass.Get(), UDMMaterialStageBlend::InputB,
+									FDMMaterialStageConnectorChannel::WHOLE_CHANNEL,
+									FDMMaterialStageConnectorChannel::WHOLE_CHANNEL
+								);
 							}
 							else if (StageSource->IsA<UDMMaterialStageThroughputLayerBlend>())
 							{
 								FScopedTransaction Transaction(LOCTEXT("SetStageInputMask", "Set Material Designer Mask Source"));
 								Stage->Modify();
-								Stage->ChangeInput_Gradient(UDMMaterialStageThroughputLayerBlend::InputMaskSource, FDMMaterialStageConnectorChannel::WHOLE_CHANNEL, GradientClass.Get(), FDMMaterialStageConnectorChannel::WHOLE_CHANNEL);
+								UDMMaterialStageInputGradient::ChangeStageInput_Gradient(
+									Stage,
+									GradientClass.Get(), 
+									UDMMaterialStageThroughputLayerBlend::InputMaskSource,
+									FDMMaterialStageConnectorChannel::WHOLE_CHANNEL,
+									FDMMaterialStageConnectorChannel::WHOLE_CHANNEL
+								);
 							}
 							else
 							{
 								FScopedTransaction Transaction(LOCTEXT("SetStageSource", "Set Material Designer Stage Source"));
 								Stage->Modify();
-								Stage->ChangeSource_Gradient(GradientClass.Get());
+								UDMMaterialStageInputGradient::ChangeStageSource_Gradient(Stage, GradientClass.Get());
 							}
 
 							if (TSharedPtr<SDMSlot> SlotWidget = MenuContext->GetSlotWidget().Pin())
@@ -658,21 +749,39 @@ namespace UE::DynamicMaterialEditor::Private
 		{
 			FScopedTransaction Transaction(LOCTEXT("SetStageInputBase", "Set Material Designer Base Source"));
 			Stage->Modify();
-			Stage->ChangeInput_Expression(UDMMaterialStageBlend::InputB, FDMMaterialStageConnectorChannel::WHOLE_CHANNEL,
-				UDMMaterialStageExpressionTextureSample::StaticClass(), 0, FDMMaterialStageConnectorChannel::THREE_CHANNELS);
+
+			UDMMaterialStageInputExpression::ChangeStageInput_Expression(
+				Stage,
+				UDMMaterialStageExpressionTextureSample::StaticClass(), 
+				UDMMaterialStageBlend::InputB,
+				FDMMaterialStageConnectorChannel::WHOLE_CHANNEL, 
+				0,
+				FDMMaterialStageConnectorChannel::THREE_CHANNELS
+			);
 		}
 		else if (StageSource->IsA<UDMMaterialStageThroughputLayerBlend>())
 		{
 			FScopedTransaction Transaction(LOCTEXT("SetStageInputMask", "Set Material Designer Mask Source"));
 			Stage->Modify();
-			Stage->ChangeInput_Expression(UDMMaterialStageThroughputLayerBlend::InputMaskSource, FDMMaterialStageConnectorChannel::WHOLE_CHANNEL,
-				UDMMaterialStageExpressionTextureSample::StaticClass(), 0, FDMMaterialStageConnectorChannel::WHOLE_CHANNEL);
+
+			UDMMaterialStageInputExpression::ChangeStageInput_Expression(
+				Stage, 
+				UDMMaterialStageExpressionTextureSample::StaticClass(), 
+				UDMMaterialStageThroughputLayerBlend::InputMaskSource,
+				FDMMaterialStageConnectorChannel::WHOLE_CHANNEL, 
+				0,
+				FDMMaterialStageConnectorChannel::WHOLE_CHANNEL
+			);
 		}
 		else
 		{
 			FScopedTransaction Transaction(LOCTEXT("SetStageSource", "Set Material Designer Stage Source"));
 			Stage->Modify();
-			Stage->ChangeSource_Expression(TSubclassOf<UDMMaterialStageExpression>(UDMMaterialStageExpressionTextureSample::StaticClass()));
+
+			UDMMaterialStageInputExpression::ChangeStageSource_Expression(
+				Stage, 
+				UDMMaterialStageExpressionTextureSample::StaticClass()
+			);
 		}
 
 		if (TSharedPtr<SDMSlot> SlotWidget = InMenuContext->GetSlotWidget().Pin())
@@ -706,21 +815,33 @@ namespace UE::DynamicMaterialEditor::Private
 		{
 			FScopedTransaction Transaction(LOCTEXT("SetStageInputBase", "Set Material Designer Base Source"));
 			Stage->Modify();
-			Stage->ChangeInput_NewLocalValue(UDMMaterialStageBlend::InputB, FDMMaterialStageConnectorChannel::WHOLE_CHANNEL,
-				EDMValueType::VT_Float3_RGB, FDMMaterialStageConnectorChannel::THREE_CHANNELS);
+
+			UDMMaterialStageInputValue::ChangeStageInput_NewLocalValue(
+				Stage, 
+				UDMMaterialStageBlend::InputB,
+				FDMMaterialStageConnectorChannel::WHOLE_CHANNEL, 
+				EDMValueType::VT_Float3_RGB,
+				FDMMaterialStageConnectorChannel::THREE_CHANNELS
+			);
 		}
 		else if (StageSource->IsA<UDMMaterialStageThroughputLayerBlend>())
 		{
 			FScopedTransaction Transaction(LOCTEXT("SetStageInputMask", "Set Material Designer Mask Source"));
 			Stage->Modify();
-			Stage->ChangeInput_NewLocalValue(UDMMaterialStageThroughputLayerBlend::InputMaskSource, FDMMaterialStageConnectorChannel::WHOLE_CHANNEL,
-				EDMValueType::VT_Float3_RGB, FDMMaterialStageConnectorChannel::WHOLE_CHANNEL);
+
+			UDMMaterialStageInputValue::ChangeStageInput_NewLocalValue(
+				Stage, 
+				UDMMaterialStageThroughputLayerBlend::InputMaskSource,
+				FDMMaterialStageConnectorChannel::WHOLE_CHANNEL, 
+				EDMValueType::VT_Float3_RGB,
+				FDMMaterialStageConnectorChannel::WHOLE_CHANNEL
+			);
 		}
 		else
 		{
 			FScopedTransaction Transaction(LOCTEXT("SetStageSource", "Set Material Designer Stage Source"));
 			Stage->Modify();
-			Stage->ChangeSource_NewLocalValue(EDMValueType::VT_Float3_RGB);
+			UDMMaterialStageInputValue::ChangeStageSource_NewLocalValue(Stage, EDMValueType::VT_Float3_RGB);
 		}
 
 		if (TSharedPtr<SDMSlot> SlotWidget = InMenuContext->GetSlotWidget().Pin())
@@ -754,22 +875,33 @@ namespace UE::DynamicMaterialEditor::Private
 		{
 			FScopedTransaction Transaction(LOCTEXT("SetStageInputBase", "Set Material Designer Base Source"));
 			Stage->Modify();
-			Stage->ChangeInput_NewLocalValue(UDMMaterialStageBlend::InputB, FDMMaterialStageConnectorChannel::WHOLE_CHANNEL,
-				EDMValueType::VT_ColorAtlas, FDMMaterialStageConnectorChannel::THREE_CHANNELS);
+
+			UDMMaterialStageInputValue::ChangeStageInput_NewLocalValue(
+				Stage, 
+				UDMMaterialStageBlend::InputB,
+				FDMMaterialStageConnectorChannel::WHOLE_CHANNEL, EDMValueType::VT_ColorAtlas,
+				FDMMaterialStageConnectorChannel::THREE_CHANNELS
+			);
 
 		}
 		else if (StageSource->IsA<UDMMaterialStageThroughputLayerBlend>())
 		{
 			FScopedTransaction Transaction(LOCTEXT("SetStageInputMask", "Set Material Designer Mask Source"));
 			Stage->Modify();
-			Stage->ChangeInput_NewLocalValue(UDMMaterialStageThroughputLayerBlend::InputMaskSource, FDMMaterialStageConnectorChannel::WHOLE_CHANNEL,
-				EDMValueType::VT_ColorAtlas, FDMMaterialStageConnectorChannel::FOURTH_CHANNEL);
+
+			UDMMaterialStageInputValue::ChangeStageInput_NewLocalValue(
+				Stage, 
+				UDMMaterialStageThroughputLayerBlend::InputMaskSource,
+				FDMMaterialStageConnectorChannel::WHOLE_CHANNEL, 
+				EDMValueType::VT_ColorAtlas,
+				FDMMaterialStageConnectorChannel::FOURTH_CHANNEL
+			);
 		}
 		else
 		{
 			FScopedTransaction Transaction(LOCTEXT("SetStageSource", "Set Material Designer Stage Source"));
 			Stage->Modify();
-			Stage->ChangeSource_NewLocalValue(EDMValueType::VT_ColorAtlas);
+			UDMMaterialStageInputValue::ChangeStageSource_NewLocalValue(Stage, EDMValueType::VT_ColorAtlas);
 		}
 
 		if (TSharedPtr<SDMSlot> SlotWidget = InMenuContext->GetSlotWidget().Pin())
@@ -810,21 +942,39 @@ namespace UE::DynamicMaterialEditor::Private
 		{
 			FScopedTransaction Transaction(LOCTEXT("SetStageInputBase", "Set Material Designer Base Source"));
 			Stage->Modify();
-			Stage->ChangeInput_Expression(UDMMaterialStageBlend::InputB, FDMMaterialStageConnectorChannel::WHOLE_CHANNEL,
-				UDMMaterialStageExpressionTextureSampleEdgeColor::StaticClass(), 0, FDMMaterialStageConnectorChannel::THREE_CHANNELS);
+
+			UDMMaterialStageInputExpression::ChangeStageInput_Expression(
+				Stage,
+				UDMMaterialStageExpressionTextureSampleEdgeColor::StaticClass(), 
+				UDMMaterialStageBlend::InputB,
+				FDMMaterialStageConnectorChannel::WHOLE_CHANNEL, 
+				0,
+				FDMMaterialStageConnectorChannel::THREE_CHANNELS
+			);
 		}
 		else if (StageSource->IsA<UDMMaterialStageThroughputLayerBlend>())
 		{
 			FScopedTransaction Transaction(LOCTEXT("SetStageInputMask", "Set Material Designer Mask Source"));
 			Stage->Modify();
-			Stage->ChangeInput_Expression(UDMMaterialStageThroughputLayerBlend::InputMaskSource, FDMMaterialStageConnectorChannel::WHOLE_CHANNEL,
-				UDMMaterialStageExpressionTextureSampleEdgeColor::StaticClass(), 0, FDMMaterialStageConnectorChannel::WHOLE_CHANNEL);
+
+			UDMMaterialStageInputExpression::ChangeStageInput_Expression(
+				Stage, 
+				UDMMaterialStageExpressionTextureSampleEdgeColor::StaticClass(),
+				UDMMaterialStageThroughputLayerBlend::InputMaskSource,
+				FDMMaterialStageConnectorChannel::WHOLE_CHANNEL, 
+				0,
+				FDMMaterialStageConnectorChannel::WHOLE_CHANNEL
+			);
 		}
 		else
 		{
 			FScopedTransaction Transaction(LOCTEXT("SetStageInput", "Set Material Designer Source"));
 			Stage->Modify();
-			Stage->ChangeSource_Expression(TSubclassOf<UDMMaterialStageExpression>(UDMMaterialStageExpressionTextureSampleEdgeColor::StaticClass()));
+
+			UDMMaterialStageInputExpression::ChangeStageSource_Expression(
+				Stage,
+				UDMMaterialStageExpressionTextureSampleEdgeColor::StaticClass()
+			);
 		}
 
 		if (UDMMaterialStage* MaskStage = Layer->GetStage(EDMMaterialLayerStage::Mask))
@@ -870,15 +1020,30 @@ namespace UE::DynamicMaterialEditor::Private
 		{
 			FScopedTransaction Transaction(LOCTEXT("SetStageInputBase", "Set Material Designer Base Source"));
 			Stage->Modify();
-			Stage->ChangeInput_Expression(UDMMaterialStageBlend::InputB, FDMMaterialStageConnectorChannel::WHOLE_CHANNEL,
-				UDMMaterialStageExpressionSceneTexture::StaticClass(), 0, FDMMaterialStageConnectorChannel::THREE_CHANNELS);
+
+			UDMMaterialStageInputExpression::ChangeStageInput_Expression(
+				Stage, 
+				UDMMaterialStageExpressionSceneTexture::StaticClass(), 
+				UDMMaterialStageBlend::InputB,
+				FDMMaterialStageConnectorChannel::WHOLE_CHANNEL, 
+				0,
+				FDMMaterialStageConnectorChannel::THREE_CHANNELS
+			);
 
 			if (Layer->GetStageType(Stage) == EDMMaterialLayerStage::Base)
 			{
 				if (UDMMaterialStage* MaskStage = Layer->GetStage(EDMMaterialLayerStage::Mask, true))
 				{
-					MaskStage->ChangeInput_Expression(UDMMaterialStageThroughputLayerBlend::InputMaskSource, FDMMaterialStageConnectorChannel::WHOLE_CHANNEL,
-						UDMMaterialStageExpressionSceneTexture::StaticClass(), 0, FDMMaterialStageConnectorChannel::FOURTH_CHANNEL);
+					MaskStage->Modify();
+
+					UDMMaterialStageInputExpression::ChangeStageInput_Expression(
+						MaskStage,
+						UDMMaterialStageExpressionSceneTexture::StaticClass(),
+						UDMMaterialStageThroughputLayerBlend::InputMaskSource,
+						FDMMaterialStageConnectorChannel::WHOLE_CHANNEL,
+						 0,
+						FDMMaterialStageConnectorChannel::FOURTH_CHANNEL
+					);
 				}
 			}
 		}
@@ -886,14 +1051,25 @@ namespace UE::DynamicMaterialEditor::Private
 		{
 			FScopedTransaction Transaction(LOCTEXT("SetStageInputMask", "Set Material Designer Mask Source"));
 			Stage->Modify();
-			Stage->ChangeInput_Expression(UDMMaterialStageThroughputLayerBlend::InputMaskSource, FDMMaterialStageConnectorChannel::WHOLE_CHANNEL,
-				UDMMaterialStageExpressionSceneTexture::StaticClass(), 0, FDMMaterialStageConnectorChannel::FOURTH_CHANNEL);
+
+			UDMMaterialStageInputExpression::ChangeStageInput_Expression(
+				Stage, 
+				UDMMaterialStageExpressionSceneTexture::StaticClass(), 
+				UDMMaterialStageThroughputLayerBlend::InputMaskSource,
+				FDMMaterialStageConnectorChannel::WHOLE_CHANNEL, 
+				0,
+				FDMMaterialStageConnectorChannel::FOURTH_CHANNEL
+			);
 		}
 		else
 		{
 			FScopedTransaction Transaction(LOCTEXT("SetStageInput", "Set Material Designer Source"));
 			Stage->Modify();
-			Stage->ChangeSource_Expression(UDMMaterialStageExpressionSceneTexture::StaticClass());
+
+			UDMMaterialStageInputExpression::ChangeStageSource_Expression(
+				Stage,
+				UDMMaterialStageExpressionSceneTexture::StaticClass()
+			);
 		}
 
 		if (TSharedPtr<SDMSlot> SlotWidget = InMenuContext->GetSlotWidget().Pin())
@@ -952,19 +1128,35 @@ namespace UE::DynamicMaterialEditor::Private
 		{
 			FScopedTransaction Transaction(LOCTEXT("SetStageInputBase", "Set Material Designer Base Source"));
 			Stage->Modify();
-			Stage->ChangeInput_MaterialFunction(UDMMaterialStageBlend::InputB, FDMMaterialStageConnectorChannel::WHOLE_CHANNEL);
+
+			UDMMaterialStageInputFunction::ChangeStageInput_Function(
+				Stage, 
+				UDMMaterialStageFunction::GetNoOpFunction(),
+				UDMMaterialStageBlend::InputB, 	
+				FDMMaterialStageConnectorChannel::WHOLE_CHANNEL,
+				0, 
+				FDMMaterialStageConnectorChannel::WHOLE_CHANNEL
+			);
 		}
 		else if (StageSource->IsA<UDMMaterialStageThroughputLayerBlend>())
 		{
 			FScopedTransaction Transaction(LOCTEXT("SetStageInputMask", "Set Material Designer Mask Source"));
 			Stage->Modify();
-			Stage->ChangeInput_MaterialFunction(UDMMaterialStageThroughputLayerBlend::InputMaskSource, FDMMaterialStageConnectorChannel::WHOLE_CHANNEL);
+
+			UDMMaterialStageInputFunction::ChangeStageInput_Function(
+				Stage, 
+				UDMMaterialStageFunction::GetNoOpFunction(),
+				UDMMaterialStageThroughputLayerBlend::InputMaskSource, 
+				FDMMaterialStageConnectorChannel::WHOLE_CHANNEL, 
+				0,
+				FDMMaterialStageConnectorChannel::WHOLE_CHANNEL
+			);
 		}
 		else
 		{
 			FScopedTransaction Transaction(LOCTEXT("SetStageSource", "Set Material Designer Stage Source"));
 			Stage->Modify();
-			Stage->ChangeSource_MaterialFunction();
+			UDMMaterialStageInputFunction::ChangeStageSource_Function(Stage, UDMMaterialStageFunction::GetNoOpFunction());
 		}
 
 		if (TSharedPtr<SDMSlot> SlotWidget = InMenuContext->GetSlotWidget().Pin())
diff --git a/Engine/Plugins/Experimental/DynamicMaterial/Source/DynamicMaterialEditor/Private/Model/DynamicMaterialModelEditorOnlyData.cpp b/Engine/Plugins/Experimental/DynamicMaterial/Source/DynamicMaterialEditor/Private/Model/DynamicMaterialModelEditorOnlyData.cpp
--- a/Engine/Plugins/Experimental/DynamicMaterial/Source/DynamicMaterialEditor/Private/Model/DynamicMaterialModelEditorOnlyData.cpp
+++ b/Engine/Plugins/Experimental/DynamicMaterial/Source/DynamicMaterialEditor/Private/Model/DynamicMaterialModelEditorOnlyData.cpp
@@ -756,9 +756,15 @@ void UDynamicMaterialModelEditorOnlyData::SetBlendMode(TEnumAsByte<EBlendMode> I
 		}
 	}
 
-	if (BlendMode == EBlendMode::BLEND_Translucent)
+	switch (BlendMode)
 	{
-		SetPixelAnimationFlag(false);
+		case EBlendMode::BLEND_Masked:
+		case EBlendMode::BLEND_Opaque:
+			break;
+
+		default:
+			SetPixelAnimationFlag(false);
+			break;
 	}
 
 	RequestMaterialBuild();
@@ -913,8 +919,10 @@ UDMMaterialSlot* UDynamicMaterialModelEditorOnlyData::AddSlot()
 
 	AssignMaterialPropertyToSlot(NewSlotProperty, NewSlot);
 
-	UDMMaterialStageInputExpression* BaseInputExpression = DefaultStage->ChangeInput_Expression(UDMMaterialStageBlendNormal::InputB, FDMMaterialStageConnectorChannel::WHOLE_CHANNEL,
-		UDMMaterialStageExpressionTextureSample::StaticClass(), 0, FDMMaterialStageConnectorChannel::WHOLE_CHANNEL);
+	UDMMaterialStageInputExpression* BaseInputExpression = UDMMaterialStageInputExpression::ChangeStageInput_Expression(DefaultStage,
+		UDMMaterialStageExpressionTextureSample::StaticClass(), UDMMaterialStageBlendNormal::InputB,
+		FDMMaterialStageConnectorChannel::WHOLE_CHANNEL, 0,
+		FDMMaterialStageConnectorChannel::WHOLE_CHANNEL);
 
 	switch (NewSlotProperty)
 	{
diff --git a/Engine/Plugins/Experimental/DynamicMaterial/Source/DynamicMaterialEditor/Private/Slate/Properties/Generators/DMMaterialStageFunctionPropertyRowGenerator.cpp b/Engine/Plugins/Experimental/DynamicMaterial/Source/DynamicMaterialEditor/Private/Slate/Properties/Generators/DMMaterialStageFunctionPropertyRowGenerator.cpp
new file mode 100644
--- /dev/null
+++ b/Engine/Plugins/Experimental/DynamicMaterial/Source/DynamicMaterialEditor/Private/Slate/Properties/Generators/DMMaterialStageFunctionPropertyRowGenerator.cpp
@@ -0,0 +1,136 @@
+// Copyright Epic Games, Inc. All Rights Reserved.
+
+#include "Slate/Properties/Generators/DMMaterialStageFunctionPropertyRowGenerator.h"
+#include "Components/DMMaterialComponent.h"
+#include "Components/DMMaterialStageFunction.h"
+#include "Components/DMMaterialValue.h"
+#include "Components/MaterialStageInputs/DMMSIFunction.h"
+#include "DynamicMaterialEditorModule.h"
+#include "IDetailPropertyRow.h"
+#include "Materials/MaterialExpressionFunctionInput.h"
+#include "Materials/MaterialExpressionMaterialFunctionCall.h"
+#include "Materials/MaterialFunctionInterface.h"
+#include "Slate/SDMComponentEdit.h"
+#include "Slate/SDMEditor.h"
+
+#define LOCTEXT_NAMESPACE "DMMaterialStageFunctionPropertyRowGenerator"
+
+const TSharedRef<FDMMaterialStageFunctionPropertyRowGenerator>& FDMMaterialStageFunctionPropertyRowGenerator::Get()
+{
+	static TSharedRef<FDMMaterialStageFunctionPropertyRowGenerator> Generator = MakeShared<FDMMaterialStageFunctionPropertyRowGenerator>();
+	return Generator;
+}
+
+void FDMMaterialStageFunctionPropertyRowGenerator::AddComponentProperties(const TSharedRef<SDMComponentEdit>& InComponentEditWidget, UDMMaterialComponent* InComponent,
+	TArray<FDMPropertyHandle>& InOutPropertyRows, TSet<UDMMaterialComponent*>& InOutProcessedObjects)
+{
+	if (!IsValid(InComponent))
+	{
+		return;
+	}
+
+	if (InOutProcessedObjects.Contains(InComponent))
+	{
+		return;
+	}
+
+	UDMMaterialStageInputFunction* StageInputFunction = Cast<UDMMaterialStageInputFunction>(InComponent);
+
+	if (!StageInputFunction)
+	{
+		return;
+	}
+
+	UDMMaterialStageFunction* StageFunction = StageInputFunction->GetMaterialStageFunction();
+
+	if (!StageFunction)
+	{
+		return;
+	}
+
+	InOutProcessedObjects.Add(InComponent);
+
+	UMaterialFunctionInterface* MaterialFunction = StageFunction->GetMaterialFunction();
+
+	if (!IsValid(MaterialFunction))
+	{
+		return;
+	}
+
+	TArray<FFunctionExpressionInput> Inputs;
+	TArray<FFunctionExpressionOutput> Outputs;
+	MaterialFunction->GetInputsAndOutputs(Inputs, Outputs);
+
+	TArray<UDMMaterialValue*> InputValues = StageFunction->GetInputValues();
+
+	// 1 input, the previous stage, does not have a value.
+	if (Inputs.Num() != (InputValues.Num() + 1))
+	{
+		return;
+	}
+
+	TArray<FDMPropertyHandle> AllValuePropertyRows;
+
+	for (int32 InputIndex = 0; InputIndex < InputValues.Num(); ++InputIndex)
+	{
+		FFunctionExpressionInput& Input = Inputs[InputIndex + 1];
+		UDMMaterialValue* Value = InputValues[InputIndex];
+
+		if (!IsValid(Value))
+		{
+			continue;
+		}
+
+		if (!Input.ExpressionInput)
+		{
+			continue;
+		}
+
+		TArray<FDMPropertyHandle> ValuePropertyRows;
+
+		FDynamicMaterialEditorModule::GeneratorComponentPropertyRows(
+			InComponentEditWidget, 
+			Value, 
+			ValuePropertyRows, 
+			InOutProcessedObjects
+		);
+
+		if (ValuePropertyRows.Num() == 1)
+		{
+			ValuePropertyRows[0].NameOverride = FText::FromName(Input.ExpressionInput->InputName);
+		}
+		else
+		{
+			static const FText NameFormat = LOCTEXT("ValueFormat", "{0}[{1}]");
+
+			for (int32 ValuePropertyIndex = 0; ValuePropertyIndex < ValuePropertyRows.Num(); ++ValuePropertyIndex)
+			{
+				ValuePropertyRows[ValuePropertyIndex].NameOverride = FText::Format(
+					NameFormat,
+					FText::FromName(Input.ExpressionInput->InputName),
+					FText::AsNumber(ValuePropertyIndex + 1)
+				);
+			}
+		}
+
+		const FText Description = FText::FromString(Input.ExpressionInput->Description);
+
+		for (FDMPropertyHandle& ValuePropertyRow : ValuePropertyRows)
+		{
+			ValuePropertyRow.NameToolTipOverride = Description;
+		}
+
+		AllValuePropertyRows.Append(ValuePropertyRows);
+	}
+
+	FDynamicMaterialEditorModule::GeneratorComponentPropertyRows(
+		InComponentEditWidget, 
+		StageFunction, 
+		InOutPropertyRows, 
+		InOutProcessedObjects
+	);
+
+	InOutPropertyRows.Append(AllValuePropertyRows);
+}
+
+#undef LOCTEXT_NAMESPACE
diff --git a/Engine/Plugins/Experimental/DynamicMaterial/Source/DynamicMaterialEditor/Private/Slate/Properties/Generators/DMMaterialStageFunctionPropertyRowGenerator.h b/Engine/Plugins/Experimental/DynamicMaterial/Source/DynamicMaterialEditor/Private/Slate/Properties/Generators/DMMaterialStageFunctionPropertyRowGenerator.h
new file mode 100644
--- /dev/null
+++ b/Engine/Plugins/Experimental/DynamicMaterial/Source/DynamicMaterialEditor/Private/Slate/Properties/Generators/DMMaterialStageFunctionPropertyRowGenerator.h
@@ -0,0 +1,17 @@
+// Copyright Epic Games, Inc. All Rights Reserved.
+
+#pragma once
+
+#include "Slate/Properties/Generators/DMComponentPropertyRowGenerator.h"
+
+class FDMMaterialStageFunctionPropertyRowGenerator : public FDMComponentPropertyRowGenerator
+{
+public:
+	static const TSharedRef<FDMMaterialStageFunctionPropertyRowGenerator>& Get();
+
+	FDMMaterialStageFunctionPropertyRowGenerator() = default;
+	virtual ~FDMMaterialStageFunctionPropertyRowGenerator() override = default;
+
+	virtual void AddComponentProperties(const TSharedRef<SDMComponentEdit>& InComponentEditWidget, UDMMaterialComponent* InComponent,
+		TArray<FDMPropertyHandle>& InOutPropertyRows, TSet<UDMMaterialComponent*>& InOutProcessedObjects) override;
+};
diff --git a/Engine/Plugins/Experimental/DynamicMaterial/Source/DynamicMaterialEditor/Private/Slate/SDMComponentEdit.cpp b/Engine/Plugins/Experimental/DynamicMaterial/Source/DynamicMaterialEditor/Private/Slate/SDMComponentEdit.cpp
--- a/Engine/Plugins/Experimental/DynamicMaterial/Source/DynamicMaterialEditor/Private/Slate/SDMComponentEdit.cpp
+++ b/Engine/Plugins/Experimental/DynamicMaterial/Source/DynamicMaterialEditor/Private/Slate/SDMComponentEdit.cpp
@@ -251,15 +251,6 @@ void SDMComponentEdit::Construct(const FArguments& InArgs, UDMMaterialComponent*
 
 	if (ensure(IsValid(InComponent)))
 	{
-		if (UDMMaterialStage* Stage = Cast<UDMMaterialStage>(InComponent))
-		{
-			Stage->GetOnUpdate().AddSP(this, &SDMComponentEdit::OnStageUpdate);
-		}
-		else if (UDMMaterialEffect* Effect = Cast<UDMMaterialEffect>(InComponent))
-		{
-			Effect->GetOnUpdate().AddSP(this, &SDMComponentEdit::OnEffectUpdated);
-		}
-
 		if (const UWorld* const World = InComponent->GetWorld())
 		{
 			if (const UDMWorldSubsystem* const WorldSubsystem = World->GetSubsystem<UDMWorldSubsystem>())
@@ -698,42 +689,6 @@ FText SDMComponentEdit::GetSourceTypeEditWidgetText() const
 	return FText::GetEmpty();
 }
 
-void SDMComponentEdit::OnStageUpdate(UDMMaterialComponent* InComponent, EDMUpdateType InUpdateType)
-{
-	if (InUpdateType == EDMUpdateType::Structure)
-	{
-		if (InComponent == ComponentWeak && !UpdateHandle.IsValid())
-		{
-			UpdateHandle = FCoreDelegates::OnEndFrame.AddSPLambda(this, [this]()
-				{
-					if (Container.IsValid())
-					{
-						Container->SetContent(CreateEditWidget());
-						FCoreDelegates::OnEndFrame.Remove(UpdateHandle);
-					}
-				});
-		}
-	}
-}
-
-void SDMComponentEdit::OnEffectUpdated(UDMMaterialComponent* InComponent, EDMUpdateType InUpdateType)
-{
-	if (InUpdateType == EDMUpdateType::Structure)
-	{
-		if (InComponent == ComponentWeak.Get() && !UpdateHandle.IsValid())
-		{
-			UpdateHandle = FCoreDelegates::OnEndFrame.AddSPLambda(this, [this]()
-				{
-					if (Container.IsValid())
-					{
-						Container->SetContent(CreateEditWidget());
-						FCoreDelegates::OnEndFrame.Remove(UpdateHandle);
-					}
-				});
-		}
-	}
-}
-
 void SDMComponentEdit::OnUndo()
 {
 	if (UDMMaterialStage* Stage = Cast<UDMMaterialStage>(ComponentWeak.Get()))
@@ -744,7 +699,10 @@ void SDMComponentEdit::OnUndo()
 			{
 				if (bCreatedWithLinkedUVs != Layer->IsTextureUVLinkEnabled())
 				{
-					OnStageUpdate(Stage, EDMUpdateType::Structure);
+					if (TSharedPtr<SDMSlot> SlotWidget = SlotWidgetWeak.Pin())
+					{
+						SlotWidget->InvalidateComponentEditWidget();
+					}
 				}
 			}
 		}
diff --git a/Engine/Plugins/Experimental/DynamicMaterial/Source/DynamicMaterialEditor/Private/Slate/SDMEditor.cpp b/Engine/Plugins/Experimental/DynamicMaterial/Source/DynamicMaterialEditor/Private/Slate/SDMEditor.cpp
--- a/Engine/Plugins/Experimental/DynamicMaterial/Source/DynamicMaterialEditor/Private/Slate/SDMEditor.cpp
+++ b/Engine/Plugins/Experimental/DynamicMaterial/Source/DynamicMaterialEditor/Private/Slate/SDMEditor.cpp
@@ -1468,7 +1468,15 @@ bool SDMEditor::CanMaterialBeAnimated() const
 {
 	if (UDynamicMaterialModelEditorOnlyData* ModelEditorOnlyData = UDynamicMaterialModelEditorOnlyData::Get(MaterialModelWeak))
 	{
-		return ModelEditorOnlyData->GetBlendMode() != EBlendMode::BLEND_Translucent;
+		switch (ModelEditorOnlyData->GetBlendMode())
+		{
+			case EBlendMode::BLEND_Masked:
+			case EBlendMode::BLEND_Opaque:
+				return true;
+
+			default:
+				return false;
+		}
 	}
 
 	return false;
diff --git a/Engine/Plugins/Experimental/DynamicMaterial/Source/DynamicMaterialEditor/Private/Slate/SDMSlot.cpp b/Engine/Plugins/Experimental/DynamicMaterial/Source/DynamicMaterialEditor/Private/Slate/SDMSlot.cpp
--- a/Engine/Plugins/Experimental/DynamicMaterial/Source/DynamicMaterialEditor/Private/Slate/SDMSlot.cpp
+++ b/Engine/Plugins/Experimental/DynamicMaterial/Source/DynamicMaterialEditor/Private/Slate/SDMSlot.cpp
@@ -6,12 +6,17 @@
 #include "Components/DMMaterialStage.h"
 #include "Components/DMMaterialStageBlend.h"
 #include "Components/DMMaterialStageExpression.h"
+#include "Components/DMMaterialStageFunction.h"
 #include "Components/DMMaterialStageGradient.h"
 #include "Components/DMMaterialStageThroughputLayerBlend.h"
 #include "Components/DMMaterialValue.h"
 #include "Components/MaterialStageBlends/DMMSBNormal.h"
 #include "Components/MaterialStageExpressions/DMMSESceneTexture.h"
 #include "Components/MaterialStageExpressions/DMMSETextureSample.h"
+#include "Components/MaterialStageInputs/DMMSIExpression.h"
+#include "Components/MaterialStageInputs/DMMSIFunction.h"
+#include "Components/MaterialStageInputs/DMMSIGradient.h"
+#include "Components/MaterialStageInputs/DMMSISlot.h"
 #include "Components/MaterialStageInputs/DMMSITextureUV.h"
 #include "Components/MaterialStageInputs/DMMSIValue.h"
 #include "Components/MaterialValues/DMMaterialValueFloat1.h"
@@ -265,8 +270,18 @@ UDMMaterialComponent* SDMSlot::GetEditedComponent() const
 
 void SDMSlot::SetEditedComponent(UDMMaterialComponent* InComponent)
 {
+	if (EditedComponent.IsValid())
+	{
+		EditedComponent->GetOnUpdate().RemoveAll(this);
+	}
+
 	EditedComponent = InComponent;
 
+	if (EditedComponent.IsValid())
+	{
+		EditedComponent->GetOnUpdate().AddSP(this, &SDMSlot::OnComponentUpdated);
+	}
+
 	InvalidateComponentEditWidget();
 }
 
@@ -323,7 +338,7 @@ void SDMSlot::OnSourceBlendTypedSelected(const TSubclassOf<UDMMaterialStageBlend
 				{
 					FScopedTransaction Transaction(LOCTEXT("SetStageBlendMode", "Material Designer Set Blend Mode"));
 					BaseStage->Modify();
-					BaseStage->ChangeSource_Blend(InNewItem);
+					UDMMaterialStageInputExpression::ChangeStageSource_Expression(BaseStage, InNewItem.Get());
 
 					RefreshSlotSettingsRowWidget();
 				}
@@ -1039,8 +1054,12 @@ void SDMSlot::OnSlotPropertiesUpdated(UDMMaterialSlot* InSlot)
 	InvalidateHeaderPropertyListWidget();
 }
 
-void SDMSlot::OnStageUpdated(UDMMaterialStage* InStage, EDMUpdateType InUpdateType)
+void SDMSlot::OnComponentUpdated(UDMMaterialComponent* InComponent, EDMUpdateType InUpdateType)
 {
+	if (InUpdateType == EDMUpdateType::Structure)
+	{
+		InvalidateComponentEditWidget();
+	}
 }
 
 TSharedPtr<SDMStage> SDMSlot::FindStageWidget(UDMMaterialStage* const InStage) const
@@ -1325,7 +1344,13 @@ void SDMSlot::AddNewLayer_NewLocalValue(EDMValueType InType)
 	UDMMaterialStage* NewBase = UDMMaterialStageBlend::CreateStage(UDMMaterialStageBlendNormal::StaticClass());
 	AddNewLayer(NewBase);
 
-	NewBase->ChangeInput_NewLocalValue(UDMMaterialStageBlend::InputB, FDMMaterialStageConnectorChannel::WHOLE_CHANNEL, InType, FDMMaterialStageConnectorChannel::WHOLE_CHANNEL);
+	UDMMaterialStageInputValue::ChangeStageInput_NewLocalValue(
+		NewBase, 
+		UDMMaterialStageBlend::InputB,
+		FDMMaterialStageConnectorChannel::WHOLE_CHANNEL, 
+		InType,
+		FDMMaterialStageConnectorChannel::WHOLE_CHANNEL
+	);
 }
 
 void SDMSlot::AddNewLayer_GlobalValue(UDMMaterialValue* InValue)
@@ -1351,8 +1376,13 @@ void SDMSlot::AddNewLayer_GlobalValue(UDMMaterialValue* InValue)
 	UDMMaterialStage* NewBase = UDMMaterialStageBlend::CreateStage(UDMMaterialStageBlendNormal::StaticClass());
 	AddNewLayer(NewBase);
 
-	NewBase->ChangeInput_Value(UDMMaterialStageBlend::InputB, FDMMaterialStageConnectorChannel::WHOLE_CHANNEL,
-		InValue, FDMMaterialStageConnectorChannel::WHOLE_CHANNEL);
+	UDMMaterialStageInputValue::ChangeStageInput_Value(
+		NewBase, 
+		UDMMaterialStageBlend::InputB,
+		FDMMaterialStageConnectorChannel::WHOLE_CHANNEL, 
+		InValue,
+		FDMMaterialStageConnectorChannel::WHOLE_CHANNEL
+	);
 }
 
 void SDMSlot::AddNewLayer_NewGlobalValue(EDMValueType InType)
@@ -1367,8 +1397,13 @@ void SDMSlot::AddNewLayer_NewGlobalValue(EDMValueType InType)
 	UDMMaterialStage* NewStage = UDMMaterialStageBlend::CreateStage(UDMMaterialStageBlendNormal::StaticClass());
 	AddNewLayer(NewStage);
 
-	NewStage->ChangeInput_NewValue(UDMMaterialStageBlend::InputB, FDMMaterialStageConnectorChannel::WHOLE_CHANNEL, 
-		InType, FDMMaterialStageConnectorChannel::WHOLE_CHANNEL);
+	UDMMaterialStageInputValue::ChangeStageInput_NewValue(
+		NewStage, 
+		UDMMaterialStageBlend::InputB,
+		FDMMaterialStageConnectorChannel::WHOLE_CHANNEL,
+		InType, 
+		FDMMaterialStageConnectorChannel::WHOLE_CHANNEL
+	);
 }
 
 void SDMSlot::AddNewLayer_Slot(UDMMaterialSlot* InSlot, EDMMaterialPropertyType InMaterialProperty)
@@ -1389,8 +1424,15 @@ void SDMSlot::AddNewLayer_Slot(UDMMaterialSlot* InSlot, EDMMaterialPropertyType
 	UDMMaterialStage* NewStage = UDMMaterialStageBlend::CreateStage(UDMMaterialStageBlendNormal::StaticClass());
 	AddNewLayer(NewStage);
 
-	NewStage->ChangeInput_Slot(UDMMaterialStageBlend::InputB, FDMMaterialStageConnectorChannel::WHOLE_CHANNEL, 
-		InSlot, InMaterialProperty, 0, FDMMaterialStageConnectorChannel::WHOLE_CHANNEL);
+	UDMMaterialStageInputSlot::ChangeStageInput_Slot(
+		NewStage, 
+		UDMMaterialStageBlend::InputB,
+		FDMMaterialStageConnectorChannel::WHOLE_CHANNEL, 
+		InSlot, 
+		InMaterialProperty, 
+		0,
+		FDMMaterialStageConnectorChannel::WHOLE_CHANNEL
+	);
 }
 
 void SDMSlot::AddNewLayer_Expression(TSubclassOf<UDMMaterialStageExpression> InExpressionClass, EDMMaterialLayerStage LayerEnabledMask)
@@ -1414,8 +1456,14 @@ void SDMSlot::AddNewLayer_Expression(TSubclassOf<UDMMaterialStageExpression> InE
 	UDMMaterialStage* NewStage = UDMMaterialStageBlend::CreateStage(UDMMaterialStageBlendNormal::StaticClass());
 	UDMMaterialLayerObject* NewLayer = AddNewLayer(NewStage);
 
-	NewStage->ChangeInput_Expression(UDMMaterialStageBlend::InputB, FDMMaterialStageConnectorChannel::WHOLE_CHANNEL, 
-		InExpressionClass, 0, FDMMaterialStageConnectorChannel::WHOLE_CHANNEL);
+	UDMMaterialStageInputExpression::ChangeStageInput_Expression(
+		NewStage, 
+		InExpressionClass,
+		UDMMaterialStageBlend::InputB, 
+		FDMMaterialStageConnectorChannel::WHOLE_CHANNEL,
+		0, 
+		FDMMaterialStageConnectorChannel::WHOLE_CHANNEL
+	);
 
 	switch (LayerEnabledMask)
 	{
@@ -1456,8 +1504,14 @@ void SDMSlot::AddNewLayer_Blend(TSubclassOf<UDMMaterialStageBlend> InBlendClass)
 	UDMMaterialStage* NewStage = UDMMaterialStageBlend::CreateStage(InBlendClass);
 	AddNewLayer(NewStage);
 
-	NewStage->ChangeInput_Expression(UDMMaterialStageBlend::InputB, FDMMaterialStageConnectorChannel::WHOLE_CHANNEL, 
-		UDMMaterialStageExpressionTextureSample::StaticClass(), 0, FDMMaterialStageConnectorChannel::WHOLE_CHANNEL);
+	UDMMaterialStageInputExpression::ChangeStageInput_Expression(
+		NewStage, 
+		UDMMaterialStageExpressionTextureSample::StaticClass(),
+		UDMMaterialStageBlend::InputB, 
+		FDMMaterialStageConnectorChannel::WHOLE_CHANNEL, 
+		0, 
+		FDMMaterialStageConnectorChannel::WHOLE_CHANNEL
+	);
 }
 
 void SDMSlot::AddNewLayer_Gradient(TSubclassOf<UDMMaterialStageGradient> InGradientClass)
@@ -1477,7 +1531,13 @@ void SDMSlot::AddNewLayer_Gradient(TSubclassOf<UDMMaterialStageGradient> InGradi
 	UDMMaterialStage* NewStage = UDMMaterialStageBlend::CreateStage(UDMMaterialStageBlendNormal::StaticClass());
 	AddNewLayer(NewStage);
 
-	NewStage->ChangeInput_Gradient(UDMMaterialStageBlend::InputB, FDMMaterialStageConnectorChannel::WHOLE_CHANNEL, InGradientClass, FDMMaterialStageConnectorChannel::WHOLE_CHANNEL);
+	UDMMaterialStageInputGradient::ChangeStageInput_Gradient(
+		NewStage, 
+		InGradientClass, 
+		UDMMaterialStageBlend::InputB, 
+		FDMMaterialStageConnectorChannel::WHOLE_CHANNEL,
+		FDMMaterialStageConnectorChannel::WHOLE_CHANNEL
+	);
 }
 
 void SDMSlot::AddNewLayer_UV()
@@ -1519,7 +1579,14 @@ void SDMSlot::AddNewLayer_MaterialFunction()
 	UDMMaterialStage* NewStage = UDMMaterialStageBlend::CreateStage(UDMMaterialStageBlendNormal::StaticClass());
 	AddNewLayer(NewStage);
 
-	NewStage->ChangeInput_MaterialFunction(UDMMaterialStageBlend::InputB, FDMMaterialStageConnectorChannel::WHOLE_CHANNEL);
+	UDMMaterialStageInputFunction::ChangeStageInput_Function(
+		NewStage, 
+		UDMMaterialStageFunction::GetNoOpFunction(),
+		UDMMaterialStageBlend::InputB, 
+		FDMMaterialStageConnectorChannel::WHOLE_CHANNEL, 
+		0,
+		FDMMaterialStageConnectorChannel::WHOLE_CHANNEL
+	);
 }
 
 void SDMSlot::AddNewLayer_SceneTexture()
@@ -1527,13 +1594,25 @@ void SDMSlot::AddNewLayer_SceneTexture()
 	UDMMaterialStage* NewStage = UDMMaterialStageBlend::CreateStage(UDMMaterialStageBlendNormal::StaticClass());
 	UDMMaterialLayerObject* NewLayer = AddNewLayer(NewStage);
 
-	NewStage->ChangeInput_Expression(UDMMaterialStageBlend::InputB, FDMMaterialStageConnectorChannel::WHOLE_CHANNEL,
-		UDMMaterialStageExpressionSceneTexture::StaticClass(), 0, FDMMaterialStageConnectorChannel::THREE_CHANNELS);
+	UDMMaterialStageInputExpression::ChangeStageInput_Expression(
+		NewStage, 
+		UDMMaterialStageExpressionSceneTexture::StaticClass(),
+		UDMMaterialStageBlend::InputB, 
+		FDMMaterialStageConnectorChannel::WHOLE_CHANNEL,
+		0, 
+		FDMMaterialStageConnectorChannel::THREE_CHANNELS
+	);
 
 	if (UDMMaterialStage* MaskStage = NewLayer->GetStage(EDMMaterialLayerStage::Mask))
 	{
-		MaskStage->ChangeInput_Expression(UDMMaterialStageThroughputLayerBlend::InputMaskSource, FDMMaterialStageConnectorChannel::WHOLE_CHANNEL,
-			UDMMaterialStageExpressionSceneTexture::StaticClass(), 0, FDMMaterialStageConnectorChannel::FOURTH_CHANNEL);
+		UDMMaterialStageInputExpression::ChangeStageInput_Expression(
+			MaskStage, 
+			UDMMaterialStageExpressionSceneTexture::StaticClass(),
+			UDMMaterialStageThroughputLayerBlend::InputMaskSource, 
+			FDMMaterialStageConnectorChannel::WHOLE_CHANNEL,
+			0, 
+			FDMMaterialStageConnectorChannel::FOURTH_CHANNEL
+		);
 	}
 }
 
diff --git a/Engine/Plugins/Experimental/DynamicMaterial/Source/DynamicMaterialEditor/Private/Slate/SDMSlot.h b/Engine/Plugins/Experimental/DynamicMaterial/Source/DynamicMaterialEditor/Private/Slate/SDMSlot.h
--- a/Engine/Plugins/Experimental/DynamicMaterial/Source/DynamicMaterialEditor/Private/Slate/SDMSlot.h
+++ b/Engine/Plugins/Experimental/DynamicMaterial/Source/DynamicMaterialEditor/Private/Slate/SDMSlot.h
@@ -131,7 +131,7 @@ protected:
 
 	void OnSlotLayersUpdated(UDMMaterialSlot* InSlot);
 	void OnSlotPropertiesUpdated(UDMMaterialSlot* InSlot);
-	void OnStageUpdated(UDMMaterialStage* InStage, EDMUpdateType InUpdateType);
+	void OnComponentUpdated(UDMMaterialComponent* InComponent, EDMUpdateType InUpdateType);
 
 	void StartStageOpacityTransaction();
 	void EndStageOpacityTransaction(const float InValue) { EndTransaction(); }
diff --git a/Engine/Plugins/Experimental/DynamicMaterial/Source/DynamicMaterialEditor/Private/Slate/SDMStage.cpp b/Engine/Plugins/Experimental/DynamicMaterial/Source/DynamicMaterialEditor/Private/Slate/SDMStage.cpp
--- a/Engine/Plugins/Experimental/DynamicMaterial/Source/DynamicMaterialEditor/Private/Slate/SDMStage.cpp
+++ b/Engine/Plugins/Experimental/DynamicMaterial/Source/DynamicMaterialEditor/Private/Slate/SDMStage.cpp
@@ -1,6 +1,7 @@
 // Copyright Epic Games, Inc. All Rights Reserved.
 
 #include "SDMStage.h"
+#include "Components/DMMaterialLayer.h"
 #include "Components/DMMaterialSlot.h"
 #include "Components/DMMaterialStage.h"
 #include "Components/DMMaterialStageBlend.h"
@@ -21,7 +22,6 @@
 #include "Slate/Previews/SDMStagePreview.h"
 #include "Slate/SMaterialToolTip.h"
 #include "SlateOptMacros.h"
-#include "Components/DMMaterialLayer.h"
 #include "Styling/CoreStyle.h"
 #include "Widgets/Colors/SColorBlock.h"
 #include "Widgets/Layout/SBox.h"
@@ -566,11 +566,24 @@ void SDMStage::HandleAssetDragDropOperation(FAssetDragDropOp& AssetDragDropOpera
 
 	if (DraggedOverStageSource->IsA<UDMMaterialStageBlend>())
 	{
-		UDMMaterialStageInputExpression* NewInput = DraggedOverStage->ChangeInput_Expression(UDMMaterialStageBlend::InputB, FDMMaterialStageConnectorChannel::WHOLE_CHANNEL,
-			UDMMaterialStageExpressionTextureSample::StaticClass(), 0, FDMMaterialStageConnectorChannel::WHOLE_CHANNEL);
+		UDMMaterialStageInputExpression* NewInput = UDMMaterialStageInputExpression::ChangeStageInput_Expression(
+			DraggedOverStage,
+			UDMMaterialStageExpressionTextureSample::StaticClass(), 
+			UDMMaterialStageBlend::InputB,
+			FDMMaterialStageConnectorChannel::WHOLE_CHANNEL, 
+			0,
+			FDMMaterialStageConnectorChannel::WHOLE_CHANNEL
+		);
 
 		UDMMaterialSubStage* SubStage = NewInput->GetSubStage();
-		UDMMaterialStageInputValue* NewInputValue = SubStage->ChangeInput_NewLocalValue(0, FDMMaterialStageConnectorChannel::WHOLE_CHANNEL, EDMValueType::VT_Texture, FDMMaterialStageConnectorChannel::WHOLE_CHANNEL);
+
+		UDMMaterialStageInputValue* NewInputValue = UDMMaterialStageInputValue::ChangeStageInput_NewLocalValue(
+			SubStage,
+			0, 
+			FDMMaterialStageConnectorChannel::WHOLE_CHANNEL, 
+			EDMValueType::VT_Texture,
+			FDMMaterialStageConnectorChannel::WHOLE_CHANNEL
+		);
 
 		TextureValue = Cast<UDMMaterialValueTexture>(NewInputValue->GetValue());
 	}
@@ -578,18 +591,38 @@ void SDMStage::HandleAssetDragDropOperation(FAssetDragDropOp& AssetDragDropOpera
 	{
 		bool bHasAlpha = UE::DynamicMaterial::Private::HasAlpha(Texture);
 
-		UDMMaterialStageInputExpression* NewInput = DraggedOverStage->ChangeInput_Expression(2, FDMMaterialStageConnectorChannel::WHOLE_CHANNEL,
-			UDMMaterialStageExpressionTextureSample::StaticClass(), bHasAlpha ? 1 : 0, FDMMaterialStageConnectorChannel::WHOLE_CHANNEL);
+		UDMMaterialStageInputExpression* NewInput = UDMMaterialStageInputExpression::ChangeStageInput_Expression(
+			DraggedOverStage,
+			UDMMaterialStageExpressionTextureSample::StaticClass(), 
+			2,
+			FDMMaterialStageConnectorChannel::WHOLE_CHANNEL, 
+			bHasAlpha ? 1 : 0,
+			FDMMaterialStageConnectorChannel::WHOLE_CHANNEL
+		);
 
 		UDMMaterialSubStage* SubStage = NewInput->GetSubStage();
-		UDMMaterialStageInputValue* NewInputValue = SubStage->ChangeInput_NewLocalValue(0, FDMMaterialStageConnectorChannel::WHOLE_CHANNEL, EDMValueType::VT_Texture, FDMMaterialStageConnectorChannel::WHOLE_CHANNEL);
+
+		UDMMaterialStageInputValue* NewInputValue = UDMMaterialStageInputValue::ChangeStageInput_NewLocalValue(
+			SubStage,
+			0, 
+			FDMMaterialStageConnectorChannel::WHOLE_CHANNEL, 
+			EDMValueType::VT_Texture,
+			FDMMaterialStageConnectorChannel::WHOLE_CHANNEL
+		);
 
 		TextureValue = Cast<UDMMaterialValueTexture>(NewInputValue->GetValue());
 	}
 	else
 	{
-		UDMMaterialStageExpression* NewExpression = DraggedOverStage->ChangeSource_Expression(UDMMaterialStageExpressionTextureSample::StaticClass());
-		UDMMaterialStageInputValue* NewInputValue = DraggedOverStage->ChangeInput_NewLocalValue(0, FDMMaterialStageConnectorChannel::WHOLE_CHANNEL, EDMValueType::VT_Texture, FDMMaterialStageConnectorChannel::WHOLE_CHANNEL);
+		UDMMaterialStageExpression* NewExpression = DraggedOverStage->ChangeSource<UDMMaterialStageExpressionTextureSample>();
+
+		UDMMaterialStageInputValue* NewInputValue = UDMMaterialStageInputValue::ChangeStageInput_NewLocalValue(
+			DraggedOverStage,
+			0, 
+			FDMMaterialStageConnectorChannel::WHOLE_CHANNEL, 
+			EDMValueType::VT_Texture,
+			FDMMaterialStageConnectorChannel::WHOLE_CHANNEL
+		);
 
 		TextureValue = Cast<UDMMaterialValueTexture>(NewInputValue->GetValue());
 	}
diff --git a/Engine/Plugins/Experimental/DynamicMaterial/Source/DynamicMaterialEditor/Public/Components/DMMaterialEffectFunction.h b/Engine/Plugins/Experimental/DynamicMaterial/Source/DynamicMaterialEditor/Public/Components/DMMaterialEffectFunction.h
--- a/Engine/Plugins/Experimental/DynamicMaterial/Source/DynamicMaterialEditor/Public/Components/DMMaterialEffectFunction.h
+++ b/Engine/Plugins/Experimental/DynamicMaterial/Source/DynamicMaterialEditor/Public/Components/DMMaterialEffectFunction.h
@@ -61,7 +61,7 @@ protected:
 
 	void DeinitFunction();
 
-	bool NeedsFunctionInit();
+	bool NeedsFunctionInit() const;
 
 	//~ Begin UDMMaterialComponent
 	virtual UDMMaterialComponent* GetSubComponentByPath(FDMComponentPath& InPath, const FDMComponentPathSegment& InPathSegment) const override;
diff --git a/Engine/Plugins/Experimental/DynamicMaterial/Source/DynamicMaterialEditor/Public/Components/DMMaterialStage.h b/Engine/Plugins/Experimental/DynamicMaterial/Source/DynamicMaterialEditor/Public/Components/DMMaterialStage.h
--- a/Engine/Plugins/Experimental/DynamicMaterial/Source/DynamicMaterialEditor/Public/Components/DMMaterialStage.h
+++ b/Engine/Plugins/Experimental/DynamicMaterial/Source/DynamicMaterialEditor/Public/Components/DMMaterialStage.h
@@ -3,7 +3,9 @@
 #pragma once
 
 #include "Components/DMMaterialComponent.h"
+#include "Components/DMMaterialStageSource.h"
 #include "DMEDefs.h"
+#include "Templates/SubclassOf.h"
 #include "DMMaterialStage.generated.h"
 
 class FAssetThumbnailPool;
@@ -103,6 +105,9 @@ public:
 	UFUNCTION(BlueprintCallable, Category = "Material Designer")
 	void RemoveInput(UDMMaterialStageInput* InInput);
 
+	UFUNCTION(BlueprintCallable, Category = "Material Designer")
+	void RemoveAllInputs();
+
 	virtual void InputUpdated(UDMMaterialStageInput* InInput, EDMUpdateType InUpdateType);
 
 	UFUNCTION(BlueprintCallable, Category = "Material Designer")
@@ -118,64 +123,61 @@ public:
 
 	TMap<EDMMaterialPropertyType, UDMMaterialLayerObject*> GetPreviousStagesPropertyMap();
 	TMap<EDMMaterialPropertyType, UDMMaterialLayerObject*> GetPropertyMap();
+	 
+	using FSourceInitFunctionPtr = TFunction<void(UDMMaterialStage*, UDMMaterialStageSource*)>;
 
 	UFUNCTION(BlueprintCallable, Category = "Material Designer")
-	UDMMaterialStageInputValue* ChangeSource_NewLocalValue(EDMValueType InType);
+	UDMMaterialStageSource* ChangeSource(TSubclassOf<UDMMaterialStageSource> InSourceClass)
+	{
+		return ChangeSource(InSourceClass, nullptr);
+	}
 
-	UFUNCTION(BlueprintCallable, Category = "Material Designer")
-	UDMMaterialStageInputValue* ChangeSource_Value(UDMMaterialValue* InValue);
+	UDMMaterialStageSource* ChangeSource(TSubclassOf<UDMMaterialStageSource> InSourceClass, FSourceInitFunctionPtr InPreInit);
 
-	UFUNCTION(BlueprintCallable, Category = "Material Designer")
-	UDMMaterialStageInputValue* ChangeSource_NewValue(EDMValueType InType);
+	template<typename InSourceClass>
+	InSourceClass* ChangeSource(FSourceInitFunctionPtr InPreInit = nullptr)
+	{
+		return Cast<InSourceClass>(ChangeSource(InSourceClass::StaticClass(), InPreInit));
+	}
 
-	UFUNCTION(BlueprintCallable, Category = "Material Designer")
-	UDMMaterialStageInputSlot* ChangeSource_Slot(UDMMaterialSlot* InSlot, EDMMaterialPropertyType InProperty);
+	template<typename InSourceClass>
+	InSourceClass* ChangeSource(TSubclassOf<UDMMaterialStageSource> InSourceSubclass, FSourceInitFunctionPtr InPreInit = nullptr)
+	{
+		return Cast<InSourceClass>(ChangeSource(InSourceSubclass, InPreInit));
+	}
 
-	UFUNCTION(BlueprintCallable, Category = "Material Designer")
-	UDMMaterialStageExpression* ChangeSource_Expression(TSubclassOf<UDMMaterialStageExpression> InExpressionClass);
+	using FInputInitFunctionPtr = TFunction<void(UDMMaterialStage*, UDMMaterialStageInput*)>;
 
 	UFUNCTION(BlueprintCallable, Category = "Material Designer")
-	UDMMaterialStageBlend* ChangeSource_Blend(TSubclassOf<UDMMaterialStageBlend> InBlendClass);
+	UDMMaterialStageInput* ChangeInput(TSubclassOf<UDMMaterialStageInput> InInputClass, int32 InInputIdx, int32 InInputChannel,
+		int32 InOutputIdx, int32 InOutputChannel)
+	{
+		return ChangeInput(InInputClass, InInputIdx, InInputChannel, InOutputIdx, InOutputChannel, nullptr);
+	}
 
-	UFUNCTION(BlueprintCallable, Category = "Material Designer")
-	UDMMaterialStageGradient* ChangeSource_Gradient(TSubclassOf<UDMMaterialStageGradient> InGradientClass);
+	UDMMaterialStageInput* ChangeInput(TSubclassOf<UDMMaterialStageInput> InInputClass, int32 InInputIdx, int32 InInputChannel,
+		int32 InOutputIdx, int32 InOutputChannel, FInputInitFunctionPtr InPreInit);
 
-	UFUNCTION(BlueprintCallable, Category = "Material Designer")
-	UDMMaterialStageInputTextureUV* ChangeSource_UV(bool bInDoUpdate);
+	template<typename InInputClass>
+	InInputClass* ChangeInput(int32 InInputIdx, int32 InInputChannel, int32 InOutputIdx, int32 InOutputChannel, 
+		FInputInitFunctionPtr InPreInit = nullptr)
+	{
+		return Cast<InInputClass>(ChangeInput(InInputClass::StaticClass(), InInputIdx, InInputChannel, InOutputIdx, InOutputChannel, 
+			InPreInit));
+	}
 
-	UFUNCTION(BlueprintCallable, Category = "Material Designer")
-	UDMMaterialStageFunction* ChangeSource_MaterialFunction();
+	template<typename InInputClass>
+	InInputClass* ChangeInput(TSubclassOf<UDMMaterialStageInput> InInputSubclass, int32 InInputIdx, int32 InInputChannel, 
+		int32 InOutputIdx, int32 InOutputChannel, FInputInitFunctionPtr InPreInit = nullptr)
+	{
+		return Cast<InInputClass>(ChangeInput(InInputSubclass, InInputIdx, InInputChannel, InOutputIdx, InOutputChannel,
+			InPreInit));
+	}
 
 	UFUNCTION(BlueprintCallable, Category = "Material Designer")
 	UDMMaterialStageSource* ChangeInput_PreviousStage(int32 InInputIdx, int32 InInputChannel, EDMMaterialPropertyType InPreviousStageProperty, 
 		int32 InOutputIdx, int32 InOutputChannel);
 
-	UFUNCTION(BlueprintCallable, Category = "Material Designer")
-	UDMMaterialStageInputValue* ChangeInput_NewLocalValue(int32 InInputIdx, int32 InInputChannel, EDMValueType InType, int32 InOutputChannel);
-
-	UFUNCTION(BlueprintCallable, Category = "Material Designer")
-	UDMMaterialStageInputValue* ChangeInput_Value(int32 InInputIdx, int32 InInputChannel, UDMMaterialValue* InValue, int32 InOutputChannel);
-
-	UFUNCTION(BlueprintCallable, Category = "Material Designer")
-	UDMMaterialStageInputValue* ChangeInput_NewValue(int32 InInputIdx, int32 InInputChannel, EDMValueType InType, int32 InOutputChannel);
-
-	UFUNCTION(BlueprintCallable, Category = "Material Designer")
-	UDMMaterialStageInputSlot* ChangeInput_Slot(int32 InInputIdx, int32 InInputChannel, UDMMaterialSlot* InSlot, EDMMaterialPropertyType InProperty, 
-		int32 InOutputIdx, int32 InOutputChannel);
-
-	UFUNCTION(BlueprintCallable, Category = "Material Designer")
-	UDMMaterialStageInputExpression* ChangeInput_Expression(int32 InInputIdx, int32 InInputChannel, TSubclassOf<UDMMaterialStageExpression> InExpressionClass, 
-		int32 InOutputIdx, int32 InOutputChannel);
-
-	UFUNCTION(BlueprintCallable, Category = "Material Designer")
-	UDMMaterialStageInputTextureUV* ChangeInput_UV(int32 InInputIdx, int32 InInputChannel, int32 InOutputChannel);
-
-	UFUNCTION(BlueprintCallable, Category = "Material Designer")
-	UDMMaterialStageInputGradient* ChangeInput_Gradient(int32 InInputIdx, int32 InInputChannel, TSubclassOf<UDMMaterialStageGradient> InGradientClass, int32 InOutputChannel);
-
-	UFUNCTION(BlueprintCallable, Category = "Material Designer")
-	UDMMaterialStageInputFunction* ChangeInput_MaterialFunction(int32 InInputIdx, int32 InInputChannel);
-
 	UFUNCTION(BlueprintCallable, Category = "Material Designer")
 	void RemoveUnusedInputs();
 
@@ -264,6 +266,4 @@ protected:
 
 	void OnValueUpdated(UDynamicMaterialModel* InMaterialModel, UDMMaterialValue* InValue);
 	void OnTextureUVUpdated(UDynamicMaterialModel* InMaterialModel, UDMTextureUV* InTextureUV);
-
-	void ApplyWholeLayerValue(UDMMaterialValue* InValue);
 };
diff --git a/Engine/Plugins/Experimental/DynamicMaterial/Source/DynamicMaterialEditor/Public/Components/DMMaterialStageExpression.h b/Engine/Plugins/Experimental/DynamicMaterial/Source/DynamicMaterialEditor/Public/Components/DMMaterialStageExpression.h
--- a/Engine/Plugins/Experimental/DynamicMaterial/Source/DynamicMaterialEditor/Public/Components/DMMaterialStageExpression.h
+++ b/Engine/Plugins/Experimental/DynamicMaterial/Source/DynamicMaterialEditor/Public/Components/DMMaterialStageExpression.h
@@ -47,6 +47,16 @@ public:
 
 	static const TArray<TStrongObjectPtr<UClass>>& GetAvailableSourceExpressions();
 
+	UFUNCTION(BlueprintCallable, Category = "Material Designer")
+	static UDMMaterialStageExpression* ChangeStageSource_Expression(UDMMaterialStage* InStage,
+		TSubclassOf<UDMMaterialStageExpression> InExpressionClass);
+
+	template<typename InExpressionClass>
+	static UDMMaterialStageExpression* ChangeStageSource_Expression(UDMMaterialStage* InStage)
+	{
+		return ChangeStageSource_Expression(InStage, InExpressionClass::StaticClass());
+	}
+
 	UFUNCTION(BlueprintPure, Category = "Material Designer")
 	TSubclassOf<UMaterialExpression> GetMaterialExpressionClass() const { return MaterialExpressionClass; }
 
diff --git a/Engine/Plugins/Experimental/DynamicMaterial/Source/DynamicMaterialEditor/Public/Components/DMMaterialStageFunction.h b/Engine/Plugins/Experimental/DynamicMaterial/Source/DynamicMaterialEditor/Public/Components/DMMaterialStageFunction.h
--- a/Engine/Plugins/Experimental/DynamicMaterial/Source/DynamicMaterialEditor/Public/Components/DMMaterialStageFunction.h
+++ b/Engine/Plugins/Experimental/DynamicMaterial/Source/DynamicMaterialEditor/Public/Components/DMMaterialStageFunction.h
@@ -18,12 +18,26 @@ public:
 
 	static UDMMaterialStage* CreateStage(UDMMaterialLayerObject* InLayer = nullptr);
 
+	UFUNCTION(BlueprintCallable, Category = "Material Designer")
+	static UDMMaterialStageFunction* ChangeStageSource_Function(UDMMaterialStage* InStage,
+		UMaterialFunctionInterface* InMaterialFunction);
+
+	static UMaterialFunctionInterface* GetNoOpFunction();
+
+	UDMMaterialStageFunction();
+
 	UFUNCTION(BlueprintPure, Category = "Material Designer")
 	UMaterialFunctionInterface* GetMaterialFunction() const { return MaterialFunction.Get(); }
 
 	UFUNCTION(BlueprintCallable, Category = "Material Designer")
 	void SetMaterialFunction(UMaterialFunctionInterface* InMaterialFunction);
 
+	UFUNCTION(BlueprintPure, Category = "Material Designer")
+	UDMMaterialValue* GetInputValue(int32 InIndex) const;
+
+	UFUNCTION(BlueprintCallable, Category = "Material Designer")
+	TArray<UDMMaterialValue*> GetInputValues() const;
+
 	//~ Begin UDMMaterialStageThroughput
 	virtual void AddDefaultInput(int32 InInputIndex) const override;
 	virtual bool CanChangeInput(int32 InputIndex) const override;
@@ -35,9 +49,14 @@ public:
 	virtual void GenerateExpressions(const TSharedRef<FDMMaterialBuildState>& InBuildState) const override;
 	//~ End UDMMaterialStageSource
 
+	//~ Begin UDMMaterialComponent
+	virtual FText GetComponentDescription() const override;
+	//~ End UDMMaterialComponent
+
 	//~ Begin UObject
 	virtual void PreEditChange(FEditPropertyChain& PropertyAboutToChange) override;
 	virtual void PostEditChangeProperty(FPropertyChangedEvent& PropertyChangedEvent) override;
+	virtual void PostLoad() override;
 	//~ End UObject
 
 protected:
@@ -50,7 +69,16 @@ protected:
 	UPROPERTY(Transient, DuplicateTransient, TextExportTransient)
 	TObjectPtr<UMaterialFunctionInterface> MaterialFunction_PreEdit;
 
-	UDMMaterialStageFunction();
-
 	void OnMaterialFunctionChanged();
+
+	void InitFunction();
+
+	void DeinitFunction();
+
+	bool NeedsFunctionInit() const;
+
+	//~ Begin UDMMaterialComponent
+	virtual void OnComponentAdded() override;
+	virtual void OnComponentRemoved() override;
+	//~ End UDMMaterialComponent
 };
diff --git a/Engine/Plugins/Experimental/DynamicMaterial/Source/DynamicMaterialEditor/Public/Components/DMMaterialStageGradient.h b/Engine/Plugins/Experimental/DynamicMaterial/Source/DynamicMaterialEditor/Public/Components/DMMaterialStageGradient.h
--- a/Engine/Plugins/Experimental/DynamicMaterial/Source/DynamicMaterialEditor/Public/Components/DMMaterialStageGradient.h
+++ b/Engine/Plugins/Experimental/DynamicMaterial/Source/DynamicMaterialEditor/Public/Components/DMMaterialStageGradient.h
@@ -24,6 +24,16 @@ public:
 
 	static const TArray<TStrongObjectPtr<UClass>>& GetAvailableGradients();
 
+	UFUNCTION(BlueprintCallable, Category = "Material Designer")
+	static UDMMaterialStageGradient* ChangeStageSource_Gradient(UDMMaterialStage* InStage,
+		TSubclassOf<UDMMaterialStageGradient> InGradientClass);
+
+	template<typename InGradientClass>
+	static UDMMaterialStageGradient* ChangeStageSource_Gradient(UDMMaterialStage* InStage)
+	{
+		return ChangeStageSource_Gradient(InStage, InGradientClass::StaticClass());
+	}
+
 	//~ Begin UDMMaterialStageThroughput
 	virtual bool CanChangeInputType(int32 InputIndex) const override;
 	//~ End UDMMaterialStageThroughput
diff --git a/Engine/Plugins/Experimental/DynamicMaterial/Source/DynamicMaterialEditor/Public/Components/MaterialStageInputs/DMMSIExpression.h b/Engine/Plugins/Experimental/DynamicMaterial/Source/DynamicMaterialEditor/Public/Components/MaterialStageInputs/DMMSIExpression.h
--- a/Engine/Plugins/Experimental/DynamicMaterial/Source/DynamicMaterialEditor/Public/Components/MaterialStageInputs/DMMSIExpression.h
+++ b/Engine/Plugins/Experimental/DynamicMaterial/Source/DynamicMaterialEditor/Public/Components/MaterialStageInputs/DMMSIExpression.h
@@ -3,6 +3,7 @@
 #pragma once
 
 #include "Components/MaterialStageInputs/DMMSIThroughput.h"
+#include "Templates/SubclassOf.h"
 #include "UObject/StrongObjectPtr.h"
 #include "DMMSIExpression.generated.h"
 
@@ -19,6 +20,29 @@ public:
 
 	static const TArray<TStrongObjectPtr<UClass>>& GetAvailableInputExpressions();
 
+	UFUNCTION(BlueprintCallable, Category = "Material Designer")
+	static UDMMaterialStageInputExpression* ChangeStageSource_Expression(UDMMaterialStage* InStage, 
+		TSubclassOf<UDMMaterialStageExpression> InExpressionClass);
+
+	template<typename InExpressionClass>
+	static UDMMaterialStageInputExpression* ChangeStageSource_Expression(UDMMaterialStage* InStage)
+	{
+		return ChangeStageSource_Expression(InStage, InExpressionClass::StaticClass());
+	}
+
+	UFUNCTION(BlueprintCallable, Category = "Material Designer")
+	static UDMMaterialStageInputExpression* ChangeStageInput_Expression(UDMMaterialStage* InStage, 
+		TSubclassOf<UDMMaterialStageExpression> InExpressionClass, int32 InInputIdx, int32 InInputChannel, int32 InOutputIdx,
+		int32 InOutputChannel);
+
+	template<typename InExpressionClass>
+	static UDMMaterialStageInputExpression* ChangeStageInput_Expression(UDMMaterialStage* InStage, int32 InInputIdx, 
+		int32 InInputChannel, int32 InOutputIdx, int32 InOutputChannel)
+	{
+		return ChangeStageSource_Expression(InStage, InExpressionClass::StaticClass(), InInputIdx, InInputChannel, InOutputIdx,
+			InOutputChannel);
+	}
+
 	UFUNCTION(BlueprintPure, Category = "Material Designer")
 	TSubclassOf<UDMMaterialStageExpression> GetMaterialStageExpressionClass() const;
 
diff --git a/Engine/Plugins/Experimental/DynamicMaterial/Source/DynamicMaterialEditor/Public/Components/MaterialStageInputs/DMMSIFunction.h b/Engine/Plugins/Experimental/DynamicMaterial/Source/DynamicMaterialEditor/Public/Components/MaterialStageInputs/DMMSIFunction.h
--- a/Engine/Plugins/Experimental/DynamicMaterial/Source/DynamicMaterialEditor/Public/Components/MaterialStageInputs/DMMSIFunction.h
+++ b/Engine/Plugins/Experimental/DynamicMaterial/Source/DynamicMaterialEditor/Public/Components/MaterialStageInputs/DMMSIFunction.h
@@ -7,6 +7,7 @@
 
 class UDMMaterialLayerObject;
 class UDMMaterialStageFunction;
+class UMaterialFunctionInterface;
 
 UCLASS(BlueprintType, ClassGroup = "Material Designer")
 class DYNAMICMATERIALEDITOR_API UDMMaterialStageInputFunction : public UDMMaterialStageInputThroughput
@@ -16,11 +17,26 @@ class DYNAMICMATERIALEDITOR_API UDMMaterialStageInputFunction : public UDMMateri
 public:
 	static UDMMaterialStage* CreateStage(UDMMaterialLayerObject* InLayer = nullptr);
 
+	UFUNCTION(BlueprintCallable, Category = "Material Designer")
+	static UDMMaterialStageInputFunction* ChangeStageSource_Function(UDMMaterialStage* InStage,
+		UMaterialFunctionInterface* InMaterialFunction);
+
+	UFUNCTION(BlueprintCallable, Category = "Material Designer")
+	static UDMMaterialStageInputFunction* ChangeStageInput_Function(UDMMaterialStage* InStage,
+		UMaterialFunctionInterface* InMaterialFunction, int32 InInputIdx, int32 InInputChannel, int32 InOutputIdx,
+		int32 InOutputChannel);
+
 	void Init();
 
 	UFUNCTION(BlueprintPure, Category = "Material Designer")
 	UDMMaterialStageFunction* GetMaterialStageFunction() const;
 
+	UFUNCTION(BlueprintPure, Category = "Material Designer")
+	UMaterialFunctionInterface* GetMaterialFunction() const;
+
+	UFUNCTION(BlueprintCallable, Category = "Material Designer")
+	void SetMaterialFunction(UMaterialFunctionInterface* InMaterialFunction);
+
 protected:
 	UDMMaterialStageInputFunction() = default;
 };
diff --git a/Engine/Plugins/Experimental/DynamicMaterial/Source/DynamicMaterialEditor/Public/Components/MaterialStageInputs/DMMSIGradient.h b/Engine/Plugins/Experimental/DynamicMaterial/Source/DynamicMaterialEditor/Public/Components/MaterialStageInputs/DMMSIGradient.h
--- a/Engine/Plugins/Experimental/DynamicMaterial/Source/DynamicMaterialEditor/Public/Components/MaterialStageInputs/DMMSIGradient.h
+++ b/Engine/Plugins/Experimental/DynamicMaterial/Source/DynamicMaterialEditor/Public/Components/MaterialStageInputs/DMMSIGradient.h
@@ -3,6 +3,7 @@
 #pragma once
 
 #include "Components/MaterialStageInputs/DMMSIThroughput.h"
+#include "Templates/SubclassOf.h"
 #include "UObject/StrongObjectPtr.h"
 #include "DMMSIGradient.generated.h"
 
@@ -19,6 +20,27 @@ public:
 
 	static const TArray<TStrongObjectPtr<UClass>>& GetAvailableGradients();
 
+	UFUNCTION(BlueprintCallable, Category = "Material Designer")
+	static UDMMaterialStageInputGradient* ChangeStageSource_Gradient(UDMMaterialStage* InStage,
+		TSubclassOf<UDMMaterialStageGradient> InGradientClass);
+
+	template<typename InGradientClass>
+	static UDMMaterialStageInputGradient* ChangeStageSource_Gradient(UDMMaterialStage* InStage)
+	{
+		return ChangeStageSource_Gradient(InStage, InGradientClass::StaticClass());
+	}
+
+	UFUNCTION(BlueprintCallable, Category = "Material Designer")
+	static UDMMaterialStageInputGradient* ChangeStageInput_Gradient(UDMMaterialStage* InStage,
+		TSubclassOf<UDMMaterialStageGradient> InGradientClass, int32 InInputIdx, int32 InInputChannel, int32 InOutputChannel);
+
+	template<typename InGradientClass>
+	static UDMMaterialStageInputGradient* ChangeStageInput_Gradient(UDMMaterialStage* InStage, int32 InInputIdx,
+		int32 InInputChannel, int32 InOutputChannel)
+	{
+		return ChangeStageSource_Gradient(InStage, InGradientClass::StaticClass(), InInputIdx, InInputChannel, InOutputChannel);
+	}
+
 	UFUNCTION(BlueprintPure, Category = "Material Designer")
 	TSubclassOf<UDMMaterialStageGradient> GetMaterialStageGradientClass() const;
 
diff --git a/Engine/Plugins/Experimental/DynamicMaterial/Source/DynamicMaterialEditor/Public/Components/MaterialStageInputs/DMMSISlot.h b/Engine/Plugins/Experimental/DynamicMaterial/Source/DynamicMaterialEditor/Public/Components/MaterialStageInputs/DMMSISlot.h
--- a/Engine/Plugins/Experimental/DynamicMaterial/Source/DynamicMaterialEditor/Public/Components/MaterialStageInputs/DMMSISlot.h
+++ b/Engine/Plugins/Experimental/DynamicMaterial/Source/DynamicMaterialEditor/Public/Components/MaterialStageInputs/DMMSISlot.h
@@ -3,6 +3,7 @@
 #pragma once
 
 #include "Components/DMMaterialStageInput.h"
+#include "Templates/SubclassOf.h"
 #include "DMMSISlot.generated.h"
 
 class UDMMaterialLayerObject;
@@ -21,6 +22,14 @@ public:
 	static UDMMaterialStage* CreateStage(UDMMaterialSlot* InSourceSlot, EDMMaterialPropertyType InMaterialProperty, 
 		UDMMaterialLayerObject* InLayer = nullptr);
 
+	UFUNCTION(BlueprintCallable, Category = "Material Designer")
+	static UDMMaterialStageInputSlot* ChangeStageSource_Slot(UDMMaterialStage* InStage, UDMMaterialSlot* InSlot,
+		EDMMaterialPropertyType InProperty);
+
+	UFUNCTION(BlueprintCallable, Category = "Material Designer")
+	static UDMMaterialStageInputSlot* ChangeStageInput_Slot(UDMMaterialStage* InStage, int32 InInputIdx, int32 InInputChannel, 
+		UDMMaterialSlot* InSlot, EDMMaterialPropertyType InProperty, int32 InOutputIdx, int32 InOutputChannel);
+
 	UFUNCTION(BlueprintPure, Category = "Material Designer")
 	UDMMaterialSlot* GetSlot() const { return Slot; }
 
diff --git a/Engine/Plugins/Experimental/DynamicMaterial/Source/DynamicMaterialEditor/Public/Components/MaterialStageInputs/DMMSITextureUV.h b/Engine/Plugins/Experimental/DynamicMaterial/Source/DynamicMaterialEditor/Public/Components/MaterialStageInputs/DMMSITextureUV.h
--- a/Engine/Plugins/Experimental/DynamicMaterial/Source/DynamicMaterialEditor/Public/Components/MaterialStageInputs/DMMSITextureUV.h
+++ b/Engine/Plugins/Experimental/DynamicMaterial/Source/DynamicMaterialEditor/Public/Components/MaterialStageInputs/DMMSITextureUV.h
@@ -4,6 +4,7 @@
 
 #include "Components/DMMaterialStageInput.h"
 #include "Templates/SharedPointer.h"
+#include "Templates/SubclassOf.h"
 #include "DMMSITextureUV.generated.h"
 
 class SWidget;
@@ -26,6 +27,13 @@ public:
 
 	static UDMMaterialStage* CreateStage(UDynamicMaterialModel* InMaterialModel, UDMMaterialLayerObject* InLayer = nullptr);
 
+	UFUNCTION(BlueprintCallable, Category = "Material Designer")
+	static UDMMaterialStageInputTextureUV* ChangeStageSource_UV(UDMMaterialStage* InStage, bool bInDoUpdate);
+
+	UFUNCTION(BlueprintCallable, Category = "Material Designer")
+	static UDMMaterialStageInputTextureUV* ChangeStageInput_UV(UDMMaterialStage* InStage, int32 InInputIdx, int32 InInputChannel, 
+		int32 InOutputChannel);
+
 	virtual ~UDMMaterialStageInputTextureUV() override = default;
 
 	virtual FText GetComponentDescription() const override;
diff --git a/Engine/Plugins/Experimental/DynamicMaterial/Source/DynamicMaterialEditor/Public/Components/MaterialStageInputs/DMMSIValue.h b/Engine/Plugins/Experimental/DynamicMaterial/Source/DynamicMaterialEditor/Public/Components/MaterialStageInputs/DMMSIValue.h
--- a/Engine/Plugins/Experimental/DynamicMaterial/Source/DynamicMaterialEditor/Public/Components/MaterialStageInputs/DMMSIValue.h
+++ b/Engine/Plugins/Experimental/DynamicMaterial/Source/DynamicMaterialEditor/Public/Components/MaterialStageInputs/DMMSIValue.h
@@ -3,6 +3,7 @@
 #pragma once
 
 #include "Components/DMMaterialStageInput.h"
+#include "Templates/SubclassOf.h"
 #include "DMMSIValue.generated.h"
 
 class UDMMaterialLayerObject;
@@ -21,6 +22,27 @@ public:
 
 	static UDMMaterialStage* CreateStage(UDMMaterialValue* InValue, UDMMaterialLayerObject* InLayer = nullptr);
 
+	UFUNCTION(BlueprintCallable, Category = "Material Designer")
+	static UDMMaterialStageInputValue* ChangeStageSource_NewLocalValue(UDMMaterialStage* InStage, EDMValueType InType);
+
+	UFUNCTION(BlueprintCallable, Category = "Material Designer")
+	static UDMMaterialStageInputValue* ChangeStageSource_Value(UDMMaterialStage* InStage, UDMMaterialValue* InValue);
+
+	UFUNCTION(BlueprintCallable, Category = "Material Designer")
+	static UDMMaterialStageInputValue* ChangeStageSource_NewValue(UDMMaterialStage* InStage, EDMValueType InType);
+
+	UFUNCTION(BlueprintCallable, Category = "Material Designer")
+	static UDMMaterialStageInputValue* ChangeStageInput_NewLocalValue(UDMMaterialStage* InStage, int32 InInputIdx, int32 InInputChannel,
+		EDMValueType InType, int32 InOutputChannel);
+
+	UFUNCTION(BlueprintCallable, Category = "Material Designer")
+	static UDMMaterialStageInputValue* ChangeStageInput_Value(UDMMaterialStage* InStage, int32 InInputIdx, int32 InInputChannel,
+		UDMMaterialValue* InValue, int32 InOutputChannel);
+
+	UFUNCTION(BlueprintCallable, Category = "Material Designer")
+	static UDMMaterialStageInputValue* ChangeStageInput_NewValue(UDMMaterialStage* InStage, int32 InInputIdx, int32 InInputChannel,
+		EDMValueType InType, int32 InOutputChannel);
+
 	virtual FText GetComponentDescription() const override;
 	virtual FText GetChannelDescription(const FDMMaterialStageConnectorChannel& Channel) override;
 
@@ -65,4 +87,6 @@ protected:
 	//~ End UDMMaterialStageSource
 
 	bool IsSharedStageValue() const;
+
+	void ApplyWholeLayerValue();
 };
diff --git a/Engine/Plugins/Experimental/DynamicMaterial/Source/DynamicMaterialEditor/Public/Slate/SDMComponentEdit.h b/Engine/Plugins/Experimental/DynamicMaterial/Source/DynamicMaterialEditor/Public/Slate/SDMComponentEdit.h
--- a/Engine/Plugins/Experimental/DynamicMaterial/Source/DynamicMaterialEditor/Public/Slate/SDMComponentEdit.h
+++ b/Engine/Plugins/Experimental/DynamicMaterial/Source/DynamicMaterialEditor/Public/Slate/SDMComponentEdit.h
@@ -86,9 +86,5 @@ protected:
 	TSharedRef<SWidget> MakeSourceTypeEditWidgetMenuContent();
 	FText GetSourceTypeEditWidgetText() const;
 
-	void OnStageUpdate(UDMMaterialComponent* InComponent, EDMUpdateType InUpdateType);
-
-	void OnEffectUpdated(UDMMaterialComponent* InComponent, EDMUpdateType InUpdateType);
-
 	void OnUndo();
 };
