Material Designer
- Added toasts to errors.

diff --git a/Engine/Plugins/Experimental/DynamicMaterial/Source/DynamicMaterialEditor/Private/Components/DMMaterialEffectFunction.cpp b/Engine/Plugins/Experimental/DynamicMaterial/Source/DynamicMaterialEditor/Private/Components/DMMaterialEffectFunction.cpp
--- a/Engine/Plugins/Experimental/DynamicMaterial/Source/DynamicMaterialEditor/Private/Components/DMMaterialEffectFunction.cpp
+++ b/Engine/Plugins/Experimental/DynamicMaterial/Source/DynamicMaterialEditor/Private/Components/DMMaterialEffectFunction.cpp
@@ -8,6 +8,7 @@
 #include "Components/MaterialValues/DMMaterialValueFloat3XYZ.h"
 #include "Components/MaterialValues/DMMaterialValueFloat4.h"
 #include "DMComponentPath.h"
+#include "DMPrivate.h"
 #include "DynamicMaterialEditorModule.h"
 #include "Materials/MaterialExpressionComponentMask.h"
 #include "Materials/MaterialExpressionFunctionInput.h"
@@ -284,14 +285,14 @@ bool UDMMaterialEffectFunction::NeedsFunctionInit() const
 
 	if (Outputs.IsEmpty())
 	{
-		UE_LOG(LogDynamicMaterialEditor, Error, TEXT("Effect Function must have at least one output."));
+		UE::DynamicMaterialEditor::Private::LogError(TEXT("Effect Function must have at least one output."));
 		return false;
 	}
 
 	if (!Inputs[0].ExpressionInput || !Outputs[0].ExpressionOutput
 		|| Inputs[0].ExpressionInput->InputType != Outputs[0].ExpressionOutput->GetOutputType(0))
 	{
-		UE_LOG(LogDynamicMaterialEditor, Error, TEXT("Effect Function's first input must match its first output."));
+		UE::DynamicMaterialEditor::Private::LogError(TEXT("Effect Function's first input must match its first output."));
 		return false;
 	}
 
@@ -306,7 +307,7 @@ bool UDMMaterialEffectFunction::NeedsFunctionInit() const
 
 		if (!IsValid(FunctionInput))
 		{
-			UE_LOG(LogDynamicMaterialEditor, Error, TEXT("Effect Function has missing input object."));
+			UE::DynamicMaterialEditor::Private::LogError(TEXT("Effect Function has missing input object."));
 			return false;
 		}
 
@@ -322,7 +323,7 @@ bool UDMMaterialEffectFunction::NeedsFunctionInit() const
 					break;
 
 				default:
-					UE_LOG(LogDynamicMaterialEditor, Error, TEXT("Effect Function has invalid first input -  must be a scalar or vector."));
+					UE::DynamicMaterialEditor::Private::LogError(TEXT("Effect Function has invalid first input -  must be a scalar or vector."));
 					return false;
 			}
 
@@ -356,7 +357,7 @@ bool UDMMaterialEffectFunction::NeedsFunctionInit() const
 				break;
 
 			default:
-				UE_LOG(LogDynamicMaterialEditor, Error, TEXT("Effect Function has invalid input type -  must be a scalar, vector or texture."));
+				UE::DynamicMaterialEditor::Private::LogError(TEXT("Effect Function has invalid input type -  must be a scalar, vector or texture."));
 				return false;
 		}
 
@@ -406,14 +407,14 @@ void UDMMaterialEffectFunction::InitFunction()
 
 	if (Inputs.IsEmpty())
 	{
-		UE_LOG(LogDynamicMaterialEditor, Error, TEXT("Effect Function must have at least one input."));
+		UE::DynamicMaterialEditor::Private::LogError(TEXT("Effect Function must have at least one input."));
 		MaterialFunctionPtr = nullptr;
 		return;
 	}
 
 	if (Outputs.IsEmpty())
 	{
-		UE_LOG(LogDynamicMaterialEditor, Error, TEXT("Effect Function must have at least one output."));
+		UE::DynamicMaterialEditor::Private::LogError(TEXT("Effect Function must have at least one output."));
 		MaterialFunctionPtr = nullptr;
 		return;
 	}
@@ -428,7 +429,7 @@ void UDMMaterialEffectFunction::InitFunction()
 
 		if (!IsValid(FunctionInput))
 		{
-			UE_LOG(LogDynamicMaterialEditor, Error, TEXT("Effect Function has missing input object."));
+			UE::DynamicMaterialEditor::Private::LogError(TEXT("Effect Function has missing input object."));
 			MaterialFunctionPtr = nullptr;
 			return;
 		}
@@ -455,7 +456,7 @@ void UDMMaterialEffectFunction::InitFunction()
 					break;
 
 				default:
-					UE_LOG(LogDynamicMaterialEditor, Error, TEXT("Effect Function has invalid first input -  must be a scalar or vector."));
+					UE::DynamicMaterialEditor::Private::LogError(TEXT("Effect Function has invalid first input -  must be a scalar or vector."));
 					MaterialFunctionPtr = nullptr;
 					return;
 			}
@@ -491,7 +492,7 @@ void UDMMaterialEffectFunction::InitFunction()
 				break;
 
 			default:
-				UE_LOG(LogDynamicMaterialEditor, Error, TEXT("Effect Function has invalid input type -  must be a scalar, vector or texture."));
+				UE::DynamicMaterialEditor::Private::LogError(TEXT("Effect Function has invalid input type -  must be a scalar, vector or texture."));
 				MaterialFunctionPtr = nullptr;
 				return;
 		}
diff --git a/Engine/Plugins/Experimental/DynamicMaterial/Source/DynamicMaterialEditor/Private/Components/DMMaterialLayer.cpp b/Engine/Plugins/Experimental/DynamicMaterial/Source/DynamicMaterialEditor/Private/Components/DMMaterialLayer.cpp
--- a/Engine/Plugins/Experimental/DynamicMaterial/Source/DynamicMaterialEditor/Private/Components/DMMaterialLayer.cpp
+++ b/Engine/Plugins/Experimental/DynamicMaterial/Source/DynamicMaterialEditor/Private/Components/DMMaterialLayer.cpp
@@ -968,7 +968,7 @@ UDMMaterialLayerObject* UDMMaterialLayerObject::DeserializeFromString(UDMMateria
 
 	if (!FJsonSerializer::Deserialize(Reader, LayerJson))
 	{
-		UE_LOG(LogDynamicMaterialEditor, Error, TEXT("Unable to deserialise clipboard data."));
+		UE::DynamicMaterialEditor::Private::LogError(TEXT("Unable to deserialise clipboard data."));
 		return nullptr;
 	}
 
diff --git a/Engine/Plugins/Experimental/DynamicMaterial/Source/DynamicMaterialEditor/Private/Components/DMMaterialStageFunction.cpp b/Engine/Plugins/Experimental/DynamicMaterial/Source/DynamicMaterialEditor/Private/Components/DMMaterialStageFunction.cpp
--- a/Engine/Plugins/Experimental/DynamicMaterial/Source/DynamicMaterialEditor/Private/Components/DMMaterialStageFunction.cpp
+++ b/Engine/Plugins/Experimental/DynamicMaterial/Source/DynamicMaterialEditor/Private/Components/DMMaterialStageFunction.cpp
@@ -11,6 +11,7 @@
 #include "Components/MaterialValues/DMMaterialValueFloat2.h"
 #include "Components/MaterialValues/DMMaterialValueFloat3XYZ.h"
 #include "Components/MaterialValues/DMMaterialValueFloat4.h"
+#include "DMPrivate.h"
 #include "DynamicMaterialEditorModule.h"
 #include "Materials/MaterialExpressionFunctionInput.h"
 #include "Materials/MaterialExpressionFunctionOutput.h"
@@ -290,14 +291,14 @@ bool UDMMaterialStageFunction::NeedsFunctionInit() const
 
 	if (Outputs.IsEmpty())
 	{
-		UE_LOG(LogDynamicMaterialEditor, Error, TEXT("Function must have at least one output."));
+		UE::DynamicMaterialEditor::Private::LogError(TEXT("Function must have at least one output."));
 		return false;
 	}
 
 	if (!Inputs[0].ExpressionInput || !Outputs[0].ExpressionOutput
 		|| Inputs[0].ExpressionInput->InputType != Outputs[0].ExpressionOutput->GetOutputType(0))
 	{
-		UE_LOG(LogDynamicMaterialEditor, Error, TEXT("Function's first input must match its first output."));
+		UE::DynamicMaterialEditor::Private::LogError(TEXT("Function's first input must match its first output."));
 		return false;
 	}
 
@@ -312,7 +313,7 @@ bool UDMMaterialStageFunction::NeedsFunctionInit() const
 
 		if (!IsValid(FunctionInput))
 		{
-			UE_LOG(LogDynamicMaterialEditor, Error, TEXT("Function has missing input object."));
+			UE::DynamicMaterialEditor::Private::LogError(TEXT("Function has missing input object."));
 			return false;
 		}
 
@@ -326,7 +327,7 @@ bool UDMMaterialStageFunction::NeedsFunctionInit() const
 					break;
 
 				default:
-					UE_LOG(LogDynamicMaterialEditor, Error, TEXT("Function has invalid first input - must be a scalar or vector3."));
+					UE::DynamicMaterialEditor::Private::LogError(TEXT("Function has invalid first input - must be a scalar or vector3."));
 					return false;
 			}
 
@@ -360,7 +361,7 @@ bool UDMMaterialStageFunction::NeedsFunctionInit() const
 				break;
 
 			default:
-				UE_LOG(LogDynamicMaterialEditor, Error, TEXT("Function has invalid input type - must be a scalar, vector or texture."));
+				UE::DynamicMaterialEditor::Private::LogError(TEXT("Function has invalid input type - must be a scalar, vector or texture."));
 				return false;
 		}
 
@@ -418,14 +419,14 @@ void UDMMaterialStageFunction::InitFunction()
 
 	if (Inputs.IsEmpty())
 	{
-		UE_LOG(LogDynamicMaterialEditor, Error, TEXT("Function must have at least one input."));
+		UE::DynamicMaterialEditor::Private::LogError(TEXT("Function must have at least one input."));
 		MaterialFunction = nullptr;
 		return;
 	}
 
 	if (Outputs.IsEmpty())
 	{
-		UE_LOG(LogDynamicMaterialEditor, Error, TEXT("Function must have at least one output."));
+		UE::DynamicMaterialEditor::Private::LogError(TEXT("Function must have at least one output."));
 		MaterialFunction = nullptr;
 		return;
 	}
@@ -447,7 +448,7 @@ void UDMMaterialStageFunction::InitFunction()
 
 		if (!IsValid(FunctionInput))
 		{
-			UE_LOG(LogDynamicMaterialEditor, Error, TEXT("Function has missing input object."));
+			UE::DynamicMaterialEditor::Private::LogError(TEXT("Function has missing input object."));
 			InputConnectors.SetNum(1);
 			MaterialFunction = nullptr;
 			return;
@@ -463,7 +464,7 @@ void UDMMaterialStageFunction::InitFunction()
 					break;
 
 				default:
-					UE_LOG(LogDynamicMaterialEditor, Error, TEXT("Function has invalid first input - must be a scalar or vector."));
+					UE::DynamicMaterialEditor::Private::LogError(TEXT("Function has invalid first input - must be a scalar or vector."));
 					InputConnectors.SetNum(1);
 					MaterialFunction = nullptr;
 					return;
@@ -507,7 +508,7 @@ void UDMMaterialStageFunction::InitFunction()
 				break;
 
 			default:
-				UE_LOG(LogDynamicMaterialEditor, Error, TEXT("Function has invalid input type - must be a scalar, vector or texture."));
+				UE::DynamicMaterialEditor::Private::LogError(TEXT("Function has invalid input type - must be a scalar, vector or texture."));
 				InputConnectors.SetNum(1);
 				MaterialFunction = nullptr;
 				return;
diff --git a/Engine/Plugins/Experimental/DynamicMaterial/Source/DynamicMaterialEditor/Private/Components/DMMaterialSubStage.cpp b/Engine/Plugins/Experimental/DynamicMaterial/Source/DynamicMaterialEditor/Private/Components/DMMaterialSubStage.cpp
--- a/Engine/Plugins/Experimental/DynamicMaterial/Source/DynamicMaterialEditor/Private/Components/DMMaterialSubStage.cpp
+++ b/Engine/Plugins/Experimental/DynamicMaterial/Source/DynamicMaterialEditor/Private/Components/DMMaterialSubStage.cpp
@@ -1,9 +1,10 @@
 // Copyright Epic Games, Inc. All Rights Reserved.
 
 #include "Components/DMMaterialSubStage.h"
+
 #include "Components/DMMaterialLayer.h"
 #include "Components/MaterialStageInputs/DMMSIThroughput.h"
-#include "DynamicMaterialEditorModule.h"
+#include "DMPrivate.h"
 
 UDMMaterialSubStage* UDMMaterialSubStage::CreateMaterialSubStage(UDMMaterialStage* InParentStage)
 {
@@ -54,7 +55,7 @@ void UDMMaterialSubStage::PostEditorDuplicate(UDynamicMaterialModel* InMaterialM
 	}
 	else
 	{
-		UE_LOG(LogDynamicMaterialEditor, Error, TEXT("Wrong parent component passed to substage."));
+		UE::DynamicMaterialEditor::Private::LogError(TEXT("Wrong parent component passed to substage."));
 		ParentStage = nullptr;
 		ParentComponent = nullptr;
 		Super::PostEditorDuplicate(InMaterialModel, InParent);
diff --git a/Engine/Plugins/Experimental/DynamicMaterial/Source/DynamicMaterialEditor/Private/DMBlueprintFunctionLibrary.cpp b/Engine/Plugins/Experimental/DynamicMaterial/Source/DynamicMaterialEditor/Private/DMBlueprintFunctionLibrary.cpp
--- a/Engine/Plugins/Experimental/DynamicMaterial/Source/DynamicMaterialEditor/Private/DMBlueprintFunctionLibrary.cpp
+++ b/Engine/Plugins/Experimental/DynamicMaterial/Source/DynamicMaterialEditor/Private/DMBlueprintFunctionLibrary.cpp
@@ -21,6 +21,7 @@
 #include "Components/PrimitiveComponent.h"
 #include "DMEDefs.h"
 #include "DMObjectMaterialProperty.h"
+#include "DMPrivate.h"
 #include "DMWorldSubsystem.h"
 #include "DynamicMaterialEditorModule.h"
 #include "Engine/Texture.h"
@@ -243,20 +244,20 @@ bool UDMBlueprintFunctionLibrary::ExportMaterialInstance(UDynamicMaterialModel*
 {
 	if (!IsValid(InMaterialModel))
 	{
-		UE_LOG(LogDynamicMaterialEditor, Error, TEXT("Invalid material to export."));
+		UE::DynamicMaterialEditor::Private::LogError(TEXT("Invalid material to export."));
 		return false;
 	}
 
 	if (InSavePath.IsEmpty())
 	{
-		UE_LOG(LogDynamicMaterialEditor, Error, TEXT("Invalid material save path to export."));
+		UE::DynamicMaterialEditor::Private::LogError(TEXT("Invalid material save path to export."));
 		return false;
 	}
 
 	UDynamicMaterialInstance* Instance = InMaterialModel->GetDynamicMaterialInstance();
 	if (!Instance)
 	{
-		UE_LOG(LogDynamicMaterialEditor, Error, TEXT("Failed to find a Material Designer Instance to export."));
+		UE::DynamicMaterialEditor::Private::LogError(TEXT("Failed to find a Material Designer Instance to export."));
 		return false;
 	}
 
@@ -265,7 +266,7 @@ bool UDMBlueprintFunctionLibrary::ExportMaterialInstance(UDynamicMaterialModel*
 	UPackage* Package = CreatePackage(*PackagePath);
 	if (!Package)
 	{
-		UE_LOG(LogDynamicMaterialEditor, Error, TEXT("Failed to create package for Material Designer Instance (%s)."), *PackagePath);
+		UE::DynamicMaterialEditor::Private::LogError(FString::Printf(TEXT("Failed to create package for Material Designer Instance (%s)."), *PackagePath));
 		return false;
 	}
 
@@ -277,7 +278,7 @@ bool UDMBlueprintFunctionLibrary::ExportMaterialInstance(UDynamicMaterialModel*
 	UObject* NewAsset = StaticDuplicateObjectEx(Params);
 	if (!NewAsset)
 	{
-		UE_LOG(LogDynamicMaterialEditor, Error, TEXT("Failed to create new Material Designer Instance asset."));
+		UE::DynamicMaterialEditor::Private::LogError(TEXT("Failed to create new Material Designer Instance asset."));
 		return false;
 	}
 
@@ -305,20 +306,20 @@ bool UDMBlueprintFunctionLibrary::ExportGeneratedMaterial(UDynamicMaterialModel*
 {
 	if (!IsValid(InMaterialModel))
 	{
-		UE_LOG(LogDynamicMaterialEditor, Error, TEXT("Invalid material to export."));
+		UE::DynamicMaterialEditor::Private::LogError(TEXT("Invalid material to export."));
 		return false;
 	}
 
 	if (InSavePath.IsEmpty())
 	{
-		UE_LOG(LogDynamicMaterialEditor, Error, TEXT("Invalid material save path to export."));
+		UE::DynamicMaterialEditor::Private::LogError(TEXT("Invalid material save path to export."));
 		return false;
 	}
 
 	UMaterial* GeneratedMaterial = InMaterialModel->GetGeneratedMaterial();
 	if (!GeneratedMaterial)
 	{
-		UE_LOG(LogDynamicMaterialEditor, Error, TEXT("Failed to find a generated material to export."));
+		UE::DynamicMaterialEditor::Private::LogError(TEXT("Failed to find a generated material to export."));
 		return false;
 	}
 
@@ -327,7 +328,7 @@ bool UDMBlueprintFunctionLibrary::ExportGeneratedMaterial(UDynamicMaterialModel*
 	UPackage* Package = CreatePackage(*PackagePath);
 	if (!Package)
 	{
-		UE_LOG(LogDynamicMaterialEditor, Error, TEXT("Failed to create package for exported material (%s)."), *PackagePath);
+		UE::DynamicMaterialEditor::Private::LogError(FString::Printf(TEXT("Failed to create package for exported material (%s)."), *PackagePath));
 		return false;
 	}
 
@@ -339,7 +340,7 @@ bool UDMBlueprintFunctionLibrary::ExportGeneratedMaterial(UDynamicMaterialModel*
 	UObject* NewAsset = StaticDuplicateObjectEx(Params);
 	if (!NewAsset)
 	{
-		UE_LOG(LogDynamicMaterialEditor, Error, TEXT("Failed to create new material asset."));
+		UE::DynamicMaterialEditor::Private::LogError(TEXT("Failed to create new material asset."));
 		return false;
 	}
 
diff --git a/Engine/Plugins/Experimental/DynamicMaterial/Source/DynamicMaterialEditor/Private/DMPrivate.cpp b/Engine/Plugins/Experimental/DynamicMaterial/Source/DynamicMaterialEditor/Private/DMPrivate.cpp
--- a/Engine/Plugins/Experimental/DynamicMaterial/Source/DynamicMaterialEditor/Private/DMPrivate.cpp
+++ b/Engine/Plugins/Experimental/DynamicMaterial/Source/DynamicMaterialEditor/Private/DMPrivate.cpp
@@ -4,7 +4,10 @@
 #include "Components/DMMaterialLayer.h"
 #include "Components/DMMaterialStage.h"
 #include "DMDefs.h"
+#include "DynamicMaterialEditorModule.h"
 #include "MaterialExpressionIO.h"
+#include "Framework/Notifications/NotificationManager.h"
+#include "Widgets/Notifications/SNotificationList.h"
 
 namespace UE::DynamicMaterialEditor::Private
 {
@@ -65,6 +68,18 @@ namespace UE::DynamicMaterialEditor::Private
 	{
 		return (InMaterialProperty >= EDMMaterialPropertyType::Custom1 && InMaterialProperty <= EDMMaterialPropertyType::Custom4);
 	}
+
+	void LogError(const FString& InMessage, bool bInToast)
+	{
+		UE_LOG(LogDynamicMaterialEditor, Error, TEXT("%s"), *InMessage);
+
+		if (bInToast)
+		{
+			FNotificationInfo Info(FText::FromString(InMessage));
+			Info.ExpireDuration = 5.0f;
+			FSlateNotificationManager::Get().AddNotification(Info);
+		}
+	}
 }
 
 FDMMaterialLayerReference::FDMMaterialLayerReference()
diff --git a/Engine/Plugins/Experimental/DynamicMaterial/Source/DynamicMaterialEditor/Private/DMPrivate.h b/Engine/Plugins/Experimental/DynamicMaterial/Source/DynamicMaterialEditor/Private/DMPrivate.h
--- a/Engine/Plugins/Experimental/DynamicMaterial/Source/DynamicMaterialEditor/Private/DMPrivate.h
+++ b/Engine/Plugins/Experimental/DynamicMaterial/Source/DynamicMaterialEditor/Private/DMPrivate.h
@@ -28,6 +28,8 @@ namespace UE::DynamicMaterialEditor::Private
 	int32 ChannelIndexToChannelBit(int32 InChannelIndex);
 
 	bool IsCustomMaterialProperty(EDMMaterialPropertyType InMaterialProperty);
+
+	void LogError(const FString& InMessage, bool bInToast = true);
 }
 
 struct FDMMaterialLayerReference
diff --git a/Engine/Plugins/Experimental/DynamicMaterial/Source/DynamicMaterialEditor/Private/Menus/DMToolBarMenus.cpp b/Engine/Plugins/Experimental/DynamicMaterial/Source/DynamicMaterialEditor/Private/Menus/DMToolBarMenus.cpp
--- a/Engine/Plugins/Experimental/DynamicMaterial/Source/DynamicMaterialEditor/Private/Menus/DMToolBarMenus.cpp
+++ b/Engine/Plugins/Experimental/DynamicMaterial/Source/DynamicMaterialEditor/Private/Menus/DMToolBarMenus.cpp
@@ -3,6 +3,7 @@
 #include "DMToolBarMenus.h"
 #include "ContentBrowserModule.h"
 #include "DMBlueprintFunctionLibrary.h"
+#include "DMPrivate.h"
 #include "DynamicMaterialEditorModule.h"
 #include "DynamicMaterialEditorSettings.h"
 #include "IContentBrowserSingleton.h"
@@ -64,7 +65,7 @@ namespace UE::DynamicMaterialEditor::Private
 
 			if (!GeneratedMaterial)
 			{
-				UE_LOG(LogDynamicMaterialEditor, Error, TEXT("Failed to find a generated material to export."));
+				UE::DynamicMaterialEditor::Private::LogError(TEXT("Failed to find a generated material to export."));
 				return;
 			}
 
