Material Designer
- Added some analytics.

diff --git a/Engine/Plugins/Experimental/DynamicMaterial/Source/DynamicMaterialEditor/Private/Material/DynamicMaterialInstanceFactory.cpp b/Engine/Plugins/Experimental/DynamicMaterial/Source/DynamicMaterialEditor/Private/Material/DynamicMaterialInstanceFactory.cpp
--- a/Engine/Plugins/Experimental/DynamicMaterial/Source/DynamicMaterialEditor/Private/Material/DynamicMaterialInstanceFactory.cpp
+++ b/Engine/Plugins/Experimental/DynamicMaterial/Source/DynamicMaterialEditor/Private/Material/DynamicMaterialInstanceFactory.cpp
@@ -31,11 +31,6 @@ UObject* UDynamicMaterialInstanceFactory::FactoryCreateNew(UClass* Class, UObjec
 	UDynamicMaterialInstance* NewInstance = NewObject<UDynamicMaterialInstance>(InParent, Class, Name, Flags | RF_Transactional);
 	check(NewInstance);
 
-	if (FEngineAnalytics::IsAvailable())
-	{
-		FEngineAnalytics::GetProvider().RecordEvent(TEXT("Editor.Usage.MaterialDesigner.CreateMaterial"));
-	}
-
 	UDynamicMaterialModelFactory* EditorFactory = NewObject<UDynamicMaterialModelFactory>();
 	check(EditorFactory);
 
@@ -59,6 +54,19 @@ UObject* UDynamicMaterialInstanceFactory::FactoryCreateNew(UClass* Class, UObjec
 	IAssetTools& AssetTools = FModuleManager::LoadModuleChecked<FAssetToolsModule>("AssetTools").Get();
 	AssetTools.OpenEditorForAssets({NewInstance});
 
+	if (FEngineAnalytics::IsAvailable())
+	{
+		static const FString AssetType = TEXT("Asset");
+		static const FString SubobjectType = TEXT("Subobject");
+
+		TArray<FAnalyticsEventAttribute> Attribs;
+		const bool bIsAsset = NewInstance->IsAsset();
+
+		Attribs.Add(FAnalyticsEventAttribute(TEXT("Type"), bIsAsset ? AssetType : SubobjectType));
+
+		FEngineAnalytics::GetProvider().RecordEvent(TEXT("Editor.Usage.MaterialDesigner.InstanceCreated"), Attribs);
+	}
+
 	return NewInstance;
 }
 
diff --git a/Engine/Plugins/Experimental/DynamicMaterial/Source/DynamicMaterialEditor/Private/Menus/DMToolBarMenus.cpp b/Engine/Plugins/Experimental/DynamicMaterial/Source/DynamicMaterialEditor/Private/Menus/DMToolBarMenus.cpp
--- a/Engine/Plugins/Experimental/DynamicMaterial/Source/DynamicMaterialEditor/Private/Menus/DMToolBarMenus.cpp
+++ b/Engine/Plugins/Experimental/DynamicMaterial/Source/DynamicMaterialEditor/Private/Menus/DMToolBarMenus.cpp
@@ -6,6 +6,7 @@
 #include "DMPrivate.h"
 #include "DynamicMaterialEditorModule.h"
 #include "DynamicMaterialEditorSettings.h"
+#include "EngineAnalytics.h"
 #include "IContentBrowserSingleton.h"
 #include "ISinglePropertyView.h"
 #include "Material/DynamicMaterialInstance.h"
@@ -35,6 +36,11 @@ namespace UE::DynamicMaterialEditor::Private
 		{
 			if (UDynamicMaterialModelEditorOnlyData* const ModelEditorOnlyData = UDynamicMaterialModelEditorOnlyData::Get(InMenuContext->GetModel()))
 			{
+				if (FEngineAnalytics::IsAvailable())
+				{
+					FEngineAnalytics::GetProvider().RecordEvent(TEXT("Editor.Usage.MaterialDesigner.OpenedGeneratedMaterial"));
+				}
+
 				ModelEditorOnlyData->OpenMaterialEditor();
 			}
 		}
@@ -54,6 +60,11 @@ namespace UE::DynamicMaterialEditor::Private
 			FString SaveObjectPath = ContentBrowserModule.Get().CreateModalSaveAssetDialog(SaveAssetDialogConfig);
 
 			UDMBlueprintFunctionLibrary::ExportMaterialInstance(MaterialInstance->GetMaterialModel(), SaveObjectPath);
+
+			if (FEngineAnalytics::IsAvailable())
+			{
+				FEngineAnalytics::GetProvider().RecordEvent(TEXT("Editor.Usage.MaterialDesigner.ExportedMaterialInstance"));
+			}
 		}
 	}
 
@@ -84,6 +95,11 @@ namespace UE::DynamicMaterialEditor::Private
 			}
 
 			UDMBlueprintFunctionLibrary::ExportGeneratedMaterial(MaterialModel, SaveObjectPath);
+
+			if (FEngineAnalytics::IsAvailable())
+			{
+				FEngineAnalytics::GetProvider().RecordEvent(TEXT("Editor.Usage.MaterialDesigner.ExportedGeneratedMaterial"));
+			}
 		}
 	}
 
