Material Designer
- Disabled the toast notifications.

diff --git a/Engine/Plugins/Experimental/DynamicMaterial/Source/DynamicMaterialEditor/Private/DMPrivate.h b/Engine/Plugins/Experimental/DynamicMaterial/Source/DynamicMaterialEditor/Private/DMPrivate.h
--- a/Engine/Plugins/Experimental/DynamicMaterial/Source/DynamicMaterialEditor/Private/DMPrivate.h
+++ b/Engine/Plugins/Experimental/DynamicMaterial/Source/DynamicMaterialEditor/Private/DMPrivate.h
@@ -29,7 +29,7 @@ namespace UE::DynamicMaterialEditor::Private
 
 	bool IsCustomMaterialProperty(EDMMaterialPropertyType InMaterialProperty);
 
-	void LogError(const FString& InMessage, bool bInToast = true);
+	void LogError(const FString& InMessage, bool bInToast = false);
 }
 
 struct FDMMaterialLayerReference
