Material Designer
- Fixed a bug causing a crash when a layer's blend mode was changed.

diff --git a/Engine/Plugins/Experimental/DynamicMaterial/Source/DynamicMaterialEditor/Private/Components/DMMaterialStage.cpp b/Engine/Plugins/Experimental/DynamicMaterial/Source/DynamicMaterialEditor/Private/Components/DMMaterialStage.cpp
--- a/Engine/Plugins/Experimental/DynamicMaterial/Source/DynamicMaterialEditor/Private/Components/DMMaterialStage.cpp
+++ b/Engine/Plugins/Experimental/DynamicMaterial/Source/DynamicMaterialEditor/Private/Components/DMMaterialStage.cpp
@@ -1213,7 +1213,7 @@ UDMMaterialStageSource* UDMMaterialStage::ChangeSource(TSubclassOf<UDMMaterialSt
 	check(InSourceClass);
 	check(!(InSourceClass->ClassFlags & (CLASS_Abstract | CLASS_Hidden | CLASS_Deprecated | CLASS_NewerVersionExists)));
 
-	UDMMaterialStageSource* NewSource = NewObject<UDMMaterialStageExpression>(this, InSourceClass, NAME_None, RF_Transactional);
+	UDMMaterialStageSource* NewSource = NewObject<UDMMaterialStageSource>(this, InSourceClass, NAME_None, RF_Transactional);
 	check(NewSource);
 
 	if (InPreInit)
diff --git a/Engine/Plugins/Experimental/DynamicMaterial/Source/DynamicMaterialEditor/Private/Slate/SDMSlot.cpp b/Engine/Plugins/Experimental/DynamicMaterial/Source/DynamicMaterialEditor/Private/Slate/SDMSlot.cpp
--- a/Engine/Plugins/Experimental/DynamicMaterial/Source/DynamicMaterialEditor/Private/Slate/SDMSlot.cpp
+++ b/Engine/Plugins/Experimental/DynamicMaterial/Source/DynamicMaterialEditor/Private/Slate/SDMSlot.cpp
@@ -338,7 +338,7 @@ void SDMSlot::OnSourceBlendTypedSelected(const TSubclassOf<UDMMaterialStageBlend
 				{
 					FScopedTransaction Transaction(LOCTEXT("SetStageBlendMode", "Material Designer Set Blend Mode"));
 					BaseStage->Modify();
-					UDMMaterialStageInputExpression::ChangeStageSource_Expression(BaseStage, InNewItem.Get());
+					BaseStage->ChangeSource<UDMMaterialStageBlend>(InNewItem);
 
 					RefreshSlotSettingsRowWidget();
 				}
