Material Designer
- Fixed property row generator for material stage functions so that a row is always generated even without a function selected.

diff --git a/Engine/Plugins/Experimental/DynamicMaterial/Source/DynamicMaterialEditor/Private/Slate/Properties/Generators/DMMaterialStageFunctionPropertyRowGenerator.cpp b/Engine/Plugins/Experimental/DynamicMaterial/Source/DynamicMaterialEditor/Private/Slate/Properties/Generators/DMMaterialStageFunctionPropertyRowGenerator.cpp
--- a/Engine/Plugins/Experimental/DynamicMaterial/Source/DynamicMaterialEditor/Private/Slate/Properties/Generators/DMMaterialStageFunctionPropertyRowGenerator.cpp
+++ b/Engine/Plugins/Experimental/DynamicMaterial/Source/DynamicMaterialEditor/Private/Slate/Properties/Generators/DMMaterialStageFunctionPropertyRowGenerator.cpp
@@ -34,98 +34,89 @@ void FDMMaterialStageFunctionPropertyRowGenerator::AddComponentProperties(const
 		return;
 	}
 
-	UDMMaterialStageInputFunction* StageInputFunction = Cast<UDMMaterialStageInputFunction>(InComponent);
-
-	if (!StageInputFunction)
-	{
-		return;
-	}
-
-	UDMMaterialStageFunction* StageFunction = StageInputFunction->GetMaterialStageFunction();
-
-	if (!StageFunction)
-	{
-		return;
-	}
-
-	InOutProcessedObjects.Add(InComponent);
-
-	UMaterialFunctionInterface* MaterialFunction = StageFunction->GetMaterialFunction();
-
-	if (!IsValid(MaterialFunction))
-	{
-		return;
-	}
-
-	TArray<FFunctionExpressionInput> Inputs;
-	TArray<FFunctionExpressionOutput> Outputs;
-	MaterialFunction->GetInputsAndOutputs(Inputs, Outputs);
-
-	TArray<UDMMaterialValue*> InputValues = StageFunction->GetInputValues();
-
-	// 1 input, the previous stage, does not have a value.
-	if (Inputs.Num() != (InputValues.Num() + 1))
-	{
-		return;
-	}
-
+	UDMMaterialStageFunction* MaterialStageFunction = nullptr;
 	TArray<FDMPropertyHandle> AllValuePropertyRows;
 
-	for (int32 InputIndex = 0; InputIndex < InputValues.Num(); ++InputIndex)
+	if (UDMMaterialStageInputFunction* StageInputFunction = Cast<UDMMaterialStageInputFunction>(InComponent))
 	{
-		FFunctionExpressionInput& Input = Inputs[InputIndex + 1];
-		UDMMaterialValue* Value = InputValues[InputIndex];
-
-		if (!IsValid(Value))
-		{
-			continue;
-		}
-
-		if (!Input.ExpressionInput)
-		{
-			continue;
-		}
+		MaterialStageFunction = StageInputFunction->GetMaterialStageFunction();
 
-		TArray<FDMPropertyHandle> ValuePropertyRows;
-
-		FDynamicMaterialEditorModule::GeneratorComponentPropertyRows(
-			InComponentEditWidget, 
-			Value, 
-			ValuePropertyRows, 
-			InOutProcessedObjects
-		);
-
-		if (ValuePropertyRows.Num() == 1)
+		if (MaterialStageFunction)
 		{
-			ValuePropertyRows[0].NameOverride = FText::FromName(Input.ExpressionInput->InputName);
-		}
-		else
-		{
-			static const FText NameFormat = LOCTEXT("ValueFormat", "{0}[{1}]");
+			InOutProcessedObjects.Add(InComponent);
 
-			for (int32 ValuePropertyIndex = 0; ValuePropertyIndex < ValuePropertyRows.Num(); ++ValuePropertyIndex)
+			if (UMaterialFunctionInterface* MaterialFunction = MaterialStageFunction->GetMaterialFunction())
 			{
-				ValuePropertyRows[ValuePropertyIndex].NameOverride = FText::Format(
-					NameFormat,
-					FText::FromName(Input.ExpressionInput->InputName),
-					FText::AsNumber(ValuePropertyIndex + 1)
-				);
+				TArray<FFunctionExpressionInput> Inputs;
+				TArray<FFunctionExpressionOutput> Outputs;
+				MaterialFunction->GetInputsAndOutputs(Inputs, Outputs);
+
+				TArray<UDMMaterialValue*> InputValues = MaterialStageFunction->GetInputValues();
+
+				// 1 input, the previous stage, does not have a value.
+				if (Inputs.Num() != (InputValues.Num() + 1))
+				{
+					return;
+				}
+
+				for (int32 InputIndex = 0; InputIndex < InputValues.Num(); ++InputIndex)
+				{
+					FFunctionExpressionInput& Input = Inputs[InputIndex + 1];
+					UDMMaterialValue* Value = InputValues[InputIndex];
+
+					if (!IsValid(Value))
+					{
+						continue;
+					}
+
+					if (!Input.ExpressionInput)
+					{
+						continue;
+					}
+
+					TArray<FDMPropertyHandle> ValuePropertyRows;
+
+					FDynamicMaterialEditorModule::GeneratorComponentPropertyRows(
+						InComponentEditWidget,
+						Value,
+						ValuePropertyRows,
+						InOutProcessedObjects
+					);
+
+					if (ValuePropertyRows.Num() == 1)
+					{
+						ValuePropertyRows[0].NameOverride = FText::FromName(Input.ExpressionInput->InputName);
+					}
+					else
+					{
+						static const FText NameFormat = LOCTEXT("ValueFormat", "{0}[{1}]");
+
+						for (int32 ValuePropertyIndex = 0; ValuePropertyIndex < ValuePropertyRows.Num(); ++ValuePropertyIndex)
+						{
+							ValuePropertyRows[ValuePropertyIndex].NameOverride = FText::Format(
+								NameFormat,
+								FText::FromName(Input.ExpressionInput->InputName),
+								FText::AsNumber(ValuePropertyIndex + 1)
+							);
+						}
+					}
+
+					const FText Description = FText::FromString(Input.ExpressionInput->Description);
+
+					for (FDMPropertyHandle& ValuePropertyRow : ValuePropertyRows)
+					{
+						ValuePropertyRow.NameToolTipOverride = Description;
+					}
+
+					AllValuePropertyRows.Append(ValuePropertyRows);
+				}
 			}
 		}
-
-		const FText Description = FText::FromString(Input.ExpressionInput->Description);
-
-		for (FDMPropertyHandle& ValuePropertyRow : ValuePropertyRows)
-		{
-			ValuePropertyRow.NameToolTipOverride = Description;
-		}
-
-		AllValuePropertyRows.Append(ValuePropertyRows);
 	}
 
 	FDynamicMaterialEditorModule::GeneratorComponentPropertyRows(
 		InComponentEditWidget, 
-		StageFunction, 
+		MaterialStageFunction, 
 		InOutPropertyRows, 
 		InOutProcessedObjects
 	);
