Material Designer
- A new take on the toast problem for errors - now requires a UI interaction to produce one.
- Added a wrapper for scoped transactions that include the ui toast.

diff --git a/Engine/Plugins/Experimental/DynamicMaterial/Source/DynamicMaterialEditor/Private/DMPrivate.cpp b/Engine/Plugins/Experimental/DynamicMaterial/Source/DynamicMaterialEditor/Private/DMPrivate.cpp
--- a/Engine/Plugins/Experimental/DynamicMaterial/Source/DynamicMaterialEditor/Private/DMPrivate.cpp
+++ b/Engine/Plugins/Experimental/DynamicMaterial/Source/DynamicMaterialEditor/Private/DMPrivate.cpp
@@ -69,11 +69,13 @@ namespace UE::DynamicMaterialEditor::Private
 		return (InMaterialProperty >= EDMMaterialPropertyType::Custom1 && InMaterialProperty <= EDMMaterialPropertyType::Custom4);
 	}
 
+	static bool bAllowUIFeedback = false;
+
 	void LogError(const FString& InMessage, bool bInToast)
 	{
 		UE_LOG(LogDynamicMaterialEditor, Error, TEXT("%s"), *InMessage);
 
-		if (bInToast)
+		if (bAllowUIFeedback && bInToast)
 		{
 			FNotificationInfo Info(FText::FromString(InMessage));
 			Info.ExpireDuration = 5.0f;
@@ -136,3 +138,9 @@ bool FDMMaterialLayerReference::IsMaskBeingEdited() const
 
 	return false;
 }
+
+FDMScopedUITransaction::FDMScopedUITransaction(const FText& InSessionName, bool bInShouldActuallyTransact)
+	: Transaction(FScopedTransaction(InSessionName, bInShouldActuallyTransact))
+	, UIFeedbackGuard(TGuardValue<bool>(UE::DynamicMaterialEditor::Private::bAllowUIFeedback, true))
+{
+}
diff --git a/Engine/Plugins/Experimental/DynamicMaterial/Source/DynamicMaterialEditor/Private/DMPrivate.h b/Engine/Plugins/Experimental/DynamicMaterial/Source/DynamicMaterialEditor/Private/DMPrivate.h
--- a/Engine/Plugins/Experimental/DynamicMaterial/Source/DynamicMaterialEditor/Private/DMPrivate.h
+++ b/Engine/Plugins/Experimental/DynamicMaterial/Source/DynamicMaterialEditor/Private/DMPrivate.h
@@ -4,6 +4,7 @@
 
 #include "Containers/Array.h"
 #include "HAL/Platform.h"
+#include "ScopedTransaction.h"
 #include "Slate/SDMSlot.h"
 #include "UObject/WeakObjectPtr.h"
 
@@ -29,7 +30,7 @@ namespace UE::DynamicMaterialEditor::Private
 
 	bool IsCustomMaterialProperty(EDMMaterialPropertyType InMaterialProperty);
 
-	void LogError(const FString& InMessage, bool bInToast = false);
+	void LogError(const FString& InMessage, bool bInToast = true);
 }
 
 struct FDMMaterialLayerReference
@@ -46,3 +47,11 @@ struct FDMMaterialLayerReference
 	bool IsMaskEnabled() const;
 	bool IsMaskBeingEdited() const;
 };
+
+struct FDMScopedUITransaction
+{
+	FScopedTransaction Transaction;
+	TGuardValue<bool> UIFeedbackGuard;
+
+	FDMScopedUITransaction(const FText& InSessionName, bool bInShouldActuallyTransact = true);
+};
diff --git a/Engine/Plugins/Experimental/DynamicMaterial/Source/DynamicMaterialEditor/Private/Slate/Layers/SDMLayerEffectsItem.cpp b/Engine/Plugins/Experimental/DynamicMaterial/Source/DynamicMaterialEditor/Private/Slate/Layers/SDMLayerEffectsItem.cpp
--- a/Engine/Plugins/Experimental/DynamicMaterial/Source/DynamicMaterialEditor/Private/Slate/Layers/SDMLayerEffectsItem.cpp
+++ b/Engine/Plugins/Experimental/DynamicMaterial/Source/DynamicMaterialEditor/Private/Slate/Layers/SDMLayerEffectsItem.cpp
@@ -4,9 +4,9 @@
 #include "Components/DMMaterialEffect.h"
 #include "Components/DMMaterialEffectStack.h"
 #include "Components/DMMaterialStage.h"
+#include "DMPrivate.h"
 #include "DragDrop/DMLayerEffectsDragDropOperation.h"
 #include "DynamicMaterialEditorStyle.h"
-#include "ScopedTransaction.h"
 #include "Slate/SDMStage.h"
 #include "SlateOptMacros.h"
 #include "Widgets/Images/SImage.h"
@@ -343,7 +343,7 @@ FReply SDMLayerEffectsItem::OnLayerBypassButtonClick()
 	{
 		if (UDMMaterialEffect* Effect = LayerItem->MaterialEffectWeak.Get())
 		{
-			FScopedTransaction Transaction(LOCTEXT("ToggleEffectEnabled", "Material Designer Toggle Effect"));
+			FDMScopedUITransaction Transaction(LOCTEXT("ToggleEffectEnabled", "Material Designer Toggle Effect"));
 			Effect->Modify();
 			Effect->SetEnabled(!Effect->IsEnabled());
 		}
@@ -379,7 +379,7 @@ FReply SDMLayerEffectsItem::OnLayerRemoveButtonClick()
 		{
 			if (UDMMaterialEffectStack* EffectStack = Effect->GetEffectStack())
 			{
-				FScopedTransaction Transaction(LOCTEXT("RemoveEffect", "Material Designer Remove Effect"));
+				FDMScopedUITransaction Transaction(LOCTEXT("RemoveEffect", "Material Designer Remove Effect"));
 				EffectStack->Modify();
 				EffectStack->RemoveEffect(Effect);
 			}
diff --git a/Engine/Plugins/Experimental/DynamicMaterial/Source/DynamicMaterialEditor/Private/Slate/Layers/SDMSlotLayerItem.cpp b/Engine/Plugins/Experimental/DynamicMaterial/Source/DynamicMaterialEditor/Private/Slate/Layers/SDMSlotLayerItem.cpp
--- a/Engine/Plugins/Experimental/DynamicMaterial/Source/DynamicMaterialEditor/Private/Slate/Layers/SDMSlotLayerItem.cpp
+++ b/Engine/Plugins/Experimental/DynamicMaterial/Source/DynamicMaterialEditor/Private/Slate/Layers/SDMSlotLayerItem.cpp
@@ -17,7 +17,6 @@
 #include "Materials/MaterialFunctionInterface.h"
 #include "Menus/DMMaterialSlotLayerAddEffectMenus.h"
 #include "Model/DynamicMaterialModelEditorOnlyData.h"
-#include "ScopedTransaction.h"
 #include "Slate/Layers/SDMLayerEffectsView.h"
 #include "Slate/SDMSlot.h"
 #include "Slate/SDMStage.h"
@@ -694,7 +693,7 @@ FReply SDMSlotLayerItem::OnLayerItemAcceptDrop(const FDragDropEvent& InDragDropE
 		}
 
 		{
-			FScopedTransaction Transaction(LOCTEXT("MoveLayer", "Move Material Designer Layer"));
+			FDMScopedUITransaction Transaction(LOCTEXT("MoveLayer", "Move Material Designer Layer"));
 			Slot->Modify();
 			Slot->MoveLayer(DraggedLayer, ThisLayerIndex);
 
@@ -996,7 +995,7 @@ FReply SDMSlotLayerItem::OnStageAcceptDrop(const FDragDropEvent& InDragDropEvent
 	}
 
 	{
-		FScopedTransaction Transaction(LOCTEXT("MoveLayer", "Move Material Designer Layer"));
+		FDMScopedUITransaction Transaction(LOCTEXT("MoveLayer", "Move Material Designer Layer"));
 		Slot->Modify();
 		Slot->MoveLayer(DraggedLayerItem, ThisBaseStageIndex);
 	}
@@ -1030,7 +1029,7 @@ FReply SDMSlotLayerItem::OnCreateLayerBypassButtonClicked()
 {
 	if (UDMMaterialLayerObject* Layer = GetLayer())
 	{
-		FScopedTransaction Transaction(LOCTEXT("ToggledLayerVisibility", "Toggle Material Designer Layer Visibility"));
+		FDMScopedUITransaction Transaction(LOCTEXT("ToggledLayerVisibility", "Toggle Material Designer Layer Visibility"));
 		Layer->Modify();
 		Layer->SetEnabled(!Layer->IsEnabled());
 
@@ -1064,12 +1063,12 @@ FReply SDMSlotLayerItem::OnBaseToggleButtonClicked()
 
 		if (UDMMaterialStage* Stage = Layer->GetStage(StageType))
 		{
-			FScopedTransaction Transaction(LOCTEXT("ToggleBaseStageEnabled", "Toggle Material Designer Base Stage Enabled"));
+			FDMScopedUITransaction Transaction(LOCTEXT("ToggleBaseStageEnabled", "Toggle Material Designer Base Stage Enabled"));
 			Stage->Modify();
 
 			if (!Stage->SetEnabled(!Stage->IsEnabled()))
 			{
-				Transaction.Cancel();
+				Transaction.Transaction.Cancel();
 			}
 		}
 	}
@@ -1098,12 +1097,12 @@ FReply SDMSlotLayerItem::OnLayerMaskToggleButtonClicked()
 
 		if (UDMMaterialStage* Stage = Layer->GetStage(StageType))
 		{
-			FScopedTransaction Transaction(LOCTEXT("ToggleMaskStageEnabled", "Toggle Material Designer Mask Stage Enabled"));
+			FDMScopedUITransaction Transaction(LOCTEXT("ToggleMaskStageEnabled", "Toggle Material Designer Mask Stage Enabled"));
 			Stage->Modify();
 
 			if (!Stage->SetEnabled(!Stage->IsEnabled()))
 			{
-				Transaction.Cancel();
+				Transaction.Transaction.Cancel();
 			}
 		}
 	}
@@ -1128,7 +1127,7 @@ FReply SDMSlotLayerItem::OnLayerLinkToggleButton()
 {
 	if (UDMMaterialLayerObject* Layer = GetLayer())
 	{
-		FScopedTransaction Transaction(LOCTEXT("UVLayerLinkToggle", "Toggle Material Designer Layer UV Link"));
+		FDMScopedUITransaction Transaction(LOCTEXT("UVLayerLinkToggle", "Toggle Material Designer Layer UV Link"));
 		Layer->Modify();
 		Layer->ToggleTextureUVLinkEnabled();
 
@@ -1259,7 +1258,7 @@ TSharedRef<SWidget> SDMSlotLayerItem::CreateLayerHeaderEditableText() const
 						{
 							if (UDMMaterialLayerObject* Layer = Stage->GetLayer())
 							{
-								FScopedTransaction Transaction(LOCTEXT("ChangeLayerName", "Material Designer Change Layer Name"));
+								FDMScopedUITransaction Transaction(LOCTEXT("ChangeLayerName", "Material Designer Change Layer Name"));
 								Layer->Modify();
 								Layer->SetLayerName(InText);
 							}
diff --git a/Engine/Plugins/Experimental/DynamicMaterial/Source/DynamicMaterialEditor/Private/Slate/Layers/SDMSlotLayerView.cpp b/Engine/Plugins/Experimental/DynamicMaterial/Source/DynamicMaterialEditor/Private/Slate/Layers/SDMSlotLayerView.cpp
--- a/Engine/Plugins/Experimental/DynamicMaterial/Source/DynamicMaterialEditor/Private/Slate/Layers/SDMSlotLayerView.cpp
+++ b/Engine/Plugins/Experimental/DynamicMaterial/Source/DynamicMaterialEditor/Private/Slate/Layers/SDMSlotLayerView.cpp
@@ -8,7 +8,6 @@
 #include "DynamicMaterialEditorSettings.h"
 #include "Framework/Application/SlateApplication.h"
 #include "Menus/DMMaterialSlotLayerMenus.h"
-#include "ScopedTransaction.h"
 #include "SDMLayerEffectsItem.h"
 #include "Slate/Layers/SDMSlotLayerItem.h"
 #include "Slate/SDMSlot.h"
@@ -243,7 +242,7 @@ void SDMSlotLayerView::ExecuteMoveLayer(int32 InOffset)
 
 	if (UDMMaterialLayerObject* Layer = Slot->GetLayer(SelectedSlotLayerIndex))
 	{
-		FScopedTransaction Transaction(LOCTEXT("MoveLayer", "Material Designer Move Layer"));
+		FDMScopedUITransaction Transaction(LOCTEXT("MoveLayer", "Material Designer Move Layer"));
 		Slot->Modify();
 
 		if (Slot->MoveLayer(Layer, NewSlotIndex))
@@ -336,7 +335,7 @@ void SDMSlotLayerView::OnLayerItemSelectionChanged(TSharedPtr<FDMMaterialLayerRe
 
 	TSharedPtr<SDMSlot> SlotWidget = SlotWidgetWeak.Pin();
 
-	FScopedTransaction Transaction(LOCTEXT("SelectLayer", "Material Designer Select Layer"));
+	FDMScopedUITransaction Transaction(LOCTEXT("SelectLayer", "Material Designer Select Layer"));
 
 	for (int32 LayerItemIdx = 0; LayerItemIdx < LayerItems.Num(); ++LayerItemIdx)
 	{
@@ -712,7 +711,7 @@ void SDMSlotLayerView::Tick(const FGeometry& AllottedGeometry, const double InCu
 
 	if (bPostRegenSelect && LayerItems.IsValidIndex(PostRegenSelectedLayerIndex))
 	{
-		FScopedTransaction Transaction(LOCTEXT("SelectLayer", "Material Designer Select Layer"));
+		FDMScopedUITransaction Transaction(LOCTEXT("SelectLayer", "Material Designer Select Layer"));
 
 		for (const TWeakObjectPtr<UDMMaterialStage>& SelectedStageWeak : PostRegenSelectedStages)
 		{
diff --git a/Engine/Plugins/Experimental/DynamicMaterial/Source/DynamicMaterialEditor/Private/Slate/Properties/Generators/DMTextureUVPropertyRowGenerator.cpp b/Engine/Plugins/Experimental/DynamicMaterial/Source/DynamicMaterialEditor/Private/Slate/Properties/Generators/DMTextureUVPropertyRowGenerator.cpp
--- a/Engine/Plugins/Experimental/DynamicMaterial/Source/DynamicMaterialEditor/Private/Slate/Properties/Generators/DMTextureUVPropertyRowGenerator.cpp
+++ b/Engine/Plugins/Experimental/DynamicMaterial/Source/DynamicMaterialEditor/Private/Slate/Properties/Generators/DMTextureUVPropertyRowGenerator.cpp
@@ -4,10 +4,10 @@
 #include "Components/DMMaterialComponent.h"
 #include "Components/DMTextureUV.h"
 #include "DMEDefs.h"
+#include "DMPrivate.h"
 #include "Modules/ModuleManager.h"
 #include "PropertyCustomizationHelpers.h"
 #include "PropertyEditorModule.h"
-#include "ScopedTransaction.h"
 #include "Slate/SDMEditor.h"
 #include "Slate/Properties/Editors/SDMPropertyEditFloat.h"
 #include "Slate/Properties/Editors/SDMPropertyEditVector.h"
@@ -361,7 +361,7 @@ TSharedRef<SWidget> UE::DynamicMaterialEditor::Private::CreateMirrorExtensionBut
 						{
 							if (TextureUV->GetMirrorOnX() || TextureUV->GetMirrorOnY())
 							{
-								FScopedTransaction Transaction(LOCTEXT("ResetMirror", "Reset Texture UV Mirror to default."));
+								FDMScopedUITransaction Transaction(LOCTEXT("ResetMirror", "Reset Texture UV Mirror to default."));
 								TextureUV->Modify();
 								TextureUV->SetMirrorOnX(false);
 								TextureUV->SetMirrorOnY(false);
diff --git a/Engine/Plugins/Experimental/DynamicMaterial/Source/DynamicMaterialEditor/Private/Slate/Properties/SDMMaterialParameters.cpp b/Engine/Plugins/Experimental/DynamicMaterial/Source/DynamicMaterialEditor/Private/Slate/Properties/SDMMaterialParameters.cpp
--- a/Engine/Plugins/Experimental/DynamicMaterial/Source/DynamicMaterialEditor/Private/Slate/Properties/SDMMaterialParameters.cpp
+++ b/Engine/Plugins/Experimental/DynamicMaterial/Source/DynamicMaterialEditor/Private/Slate/Properties/SDMMaterialParameters.cpp
@@ -4,12 +4,12 @@
 #include "Components/DMMaterialParameter.h"
 #include "Components/DMMaterialValue.h"
 #include "DetailLayoutBuilder.h"
+#include "DMPrivate.h"
 #include "DMValueDefinition.h"
 #include "DynamicMaterialEditorModule.h"
 #include "DynamicMaterialEditorStyle.h"
 #include "DynamicMaterialModule.h"
 #include "Model/DynamicMaterialModel.h"
-#include "ScopedTransaction.h"
 #include "Slate/Properties/SDMDetailsGrid.h"
 #include "SlateOptMacros.h"
 #include "Styling/StyleColors.h"
@@ -240,7 +240,7 @@ FReply SDMMaterialParameters::OnAddValueButtonClicked(EDMValueType Type)
 		return FReply::Unhandled();
 	}
 
-	FScopedTransaction Transaction(LOCTEXT("AddGlobalValue", "Add Global Value"));
+	FDMScopedUITransaction Transaction(LOCTEXT("AddGlobalValue", "Add Global Value"));
 	MaterialModel->Modify();
 	MaterialModel->AddValue(Type);
 
@@ -269,7 +269,7 @@ FReply SDMMaterialParameters::OnRemoveValueButtonClicked(int32 Index)
 		return FReply::Unhandled();
 	}
 
-	FScopedTransaction Transaction(LOCTEXT("RemoveGlobalValue", "Remove Global Value"));
+	FDMScopedUITransaction Transaction(LOCTEXT("RemoveGlobalValue", "Remove Global Value"));
 	MaterialModel->Modify();
 	MaterialModel->RemoveValueByIndex(Index);
 
@@ -327,7 +327,7 @@ void SDMMaterialParameters::OnAcceptValueNameChanged(const FText& InNewName, int
 
 	const TArray<UDMMaterialValue*>& Values = ModelWeak->GetValues();
 
-	FScopedTransaction Transaction(LOCTEXT("SetGlobalValueName", "Set Global Value Name"));
+	FDMScopedUITransaction Transaction(LOCTEXT("SetGlobalValueName", "Set Global Value Name"));
 	Values[Index]->Modify();
 
 	if (UDMMaterialParameter* Parameter = Values[Index]->GetParameter())
diff --git a/Engine/Plugins/Experimental/DynamicMaterial/Source/DynamicMaterialEditor/Private/Slate/Properties/SDMPropertyEdit.cpp b/Engine/Plugins/Experimental/DynamicMaterial/Source/DynamicMaterialEditor/Private/Slate/Properties/SDMPropertyEdit.cpp
--- a/Engine/Plugins/Experimental/DynamicMaterial/Source/DynamicMaterialEditor/Private/Slate/Properties/SDMPropertyEdit.cpp
+++ b/Engine/Plugins/Experimental/DynamicMaterial/Source/DynamicMaterialEditor/Private/Slate/Properties/SDMPropertyEdit.cpp
@@ -2,6 +2,7 @@
 
 #include "Slate/Properties/SDMPropertyEdit.h"
 #include "Components/DMMaterialValue.h"
+#include "DMPrivate.h"
 #include "DMWorldSubsystem.h"
 #include "DetailLayoutBuilder.h"
 #include "DetailRowMenuContext.h"
@@ -13,7 +14,6 @@
 #include "IDetailKeyframeHandler.h"
 #include "PropertyCustomizationHelpers.h"
 #include "PropertyHandle.h"
-#include "ScopedTransaction.h"
 #include "Slate/SDMEditor.h"
 #include "Slate/SDMSlot.h"
 #include "Slate/SDMComponentEdit.h"
@@ -409,7 +409,7 @@ void SDMPropertyEdit::StartTransaction(FText InDescription)
 		}
 	}
 
-	ScrubbingTransaction = MakeShared<FScopedTransaction>(InDescription);
+	ScrubbingTransaction = MakeShared<FDMScopedUITransaction>(InDescription);
 
 	if (IsValid(Object))
 	{
diff --git a/Engine/Plugins/Experimental/DynamicMaterial/Source/DynamicMaterialEditor/Private/Slate/Properties/SDMPropertyEdit.h b/Engine/Plugins/Experimental/DynamicMaterial/Source/DynamicMaterialEditor/Private/Slate/Properties/SDMPropertyEdit.h
--- a/Engine/Plugins/Experimental/DynamicMaterial/Source/DynamicMaterialEditor/Private/Slate/Properties/SDMPropertyEdit.h
+++ b/Engine/Plugins/Experimental/DynamicMaterial/Source/DynamicMaterialEditor/Private/Slate/Properties/SDMPropertyEdit.h
@@ -9,11 +9,11 @@
 #include "Widgets/Input/SSpinBox.h"
 #include "Widgets/SCompoundWidget.h"
 
-class FScopedTransaction;
 class SDMComponentEdit;
 class SHorizontalBox;
 class UDMMaterialValue;
 enum class ECheckBoxState : uint8;
+struct FDMScopedUITransaction;
 
 /**
  * Material Designer Property Edit
@@ -52,7 +52,7 @@ protected:
 	TWeakObjectPtr<UDMMaterialValue> PropertyMaterialValue;
 
 	FProperty* Property = nullptr;
-	TSharedPtr<FScopedTransaction> ScrubbingTransaction;
+	TSharedPtr<FDMScopedUITransaction> ScrubbingTransaction;
 
 	static void CreateKey(TWeakObjectPtr<UDMMaterialValue> InValueWeak);
 
diff --git a/Engine/Plugins/Experimental/DynamicMaterial/Source/DynamicMaterialEditor/Private/Slate/SDMComponentEdit.cpp b/Engine/Plugins/Experimental/DynamicMaterial/Source/DynamicMaterialEditor/Private/Slate/SDMComponentEdit.cpp
--- a/Engine/Plugins/Experimental/DynamicMaterial/Source/DynamicMaterialEditor/Private/Slate/SDMComponentEdit.cpp
+++ b/Engine/Plugins/Experimental/DynamicMaterial/Source/DynamicMaterialEditor/Private/Slate/SDMComponentEdit.cpp
@@ -15,6 +15,7 @@
 #include "CustomDetailsViewArgs.h"
 #include "CustomDetailsViewModule.h"
 #include "CustomDetailsViewSequencer.h"
+#include "DMPrivate.h"
 #include "DetailLayoutBuilder.h"
 #include "DMEDefs.h"
 #include "DMWorldSubsystem.h"
@@ -35,7 +36,6 @@
 #include "Model/DynamicMaterialModelEditorOnlyData.h"
 #include "PropertyCustomizationHelpers.h"
 #include "PropertyEditorModule.h"
-#include "ScopedTransaction.h"
 #include "Slate/Properties/Editors/SDMPropertyEditBool.h"
 #include "Slate/Properties/Editors/SDMPropertyEditColor.h"
 #include "Slate/Properties/Editors/SDMPropertyEditEnum.h"
@@ -451,7 +451,7 @@ TSharedRef<SWidget> SDMComponentEdit::CreateExtensionButtons(const TSharedPtr<SW
 
 					if (Outers.IsEmpty() == false && IsValid(Outers[0]))
 					{
-						FScopedTransaction Transaction(LOCTEXT("ResetValue", "Reset Value to default."));
+						FDMScopedUITransaction Transaction(LOCTEXT("ResetValue", "Reset Value to default."));
 
 						if (FProperty* Property = PropertyHandle->GetProperty())
 						{
diff --git a/Engine/Plugins/Experimental/DynamicMaterial/Source/DynamicMaterialEditor/Private/Slate/SDMEditor.cpp b/Engine/Plugins/Experimental/DynamicMaterial/Source/DynamicMaterialEditor/Private/Slate/SDMEditor.cpp
--- a/Engine/Plugins/Experimental/DynamicMaterial/Source/DynamicMaterialEditor/Private/Slate/SDMEditor.cpp
+++ b/Engine/Plugins/Experimental/DynamicMaterial/Source/DynamicMaterialEditor/Private/Slate/SDMEditor.cpp
@@ -12,6 +12,7 @@
 #include "Components/PrimitiveComponent.h"
 #include "DetailLayoutBuilder.h"
 #include "DMBlueprintFunctionLibrary.h"
+#include "DMPrivate.h"
 #include "DMWorldSubsystem.h"
 #include "DynamicMaterialEditorCommands.h"
 #include "DynamicMaterialEditorModule.h"
@@ -32,7 +33,6 @@
 #include "Model/DynamicMaterialModelEditorOnlyData.h"
 #include "PropertyCustomizationHelpers.h"
 #include "PropertyHandle.h"
-#include "ScopedTransaction.h"
 #include "SEnumCombo.h"
 #include "Slate/Properties/Editors/SDMPropertyEditOpacity.h"
 #include "Slate/Properties/SDMBlendMode.h"
@@ -1074,7 +1074,7 @@ void SDMEditor::OnOpacitySlotCheckStateChanged_NoSlot(ECheckBoxState InCheckStat
 		{
 			if (UDMMaterialLayerObject* Layer = NewSlot->GetLayer(0))
 			{
-				FScopedTransaction Transaction(LOCTEXT("AddOpacitySlot", "Material Designer Add Opacity Slot"));
+				FDMScopedUITransaction Transaction(LOCTEXT("AddOpacitySlot", "Material Designer Add Opacity Slot"));
 				Layer->Modify();
 
 				switch (ModelEditorOnlyData->GetBlendMode())
@@ -1166,7 +1166,7 @@ void SDMEditor::InsertNewLayer()
 			if (UDMMaterialSlot* Slot = SelectedLayer->GetSlot())
 			{
 				// Added here because stuff is done after the layer is added
-				FScopedTransaction Transaction(LOCTEXT("InsertNewLayer", "Material Designer Insert Layer"));
+				FDMScopedUITransaction Transaction(LOCTEXT("InsertNewLayer", "Material Designer Insert Layer"));
 				Slot->Modify();
 
 				SlotWidgets[ActiveSlotIndex]->AddNewLayer_Expression(
@@ -1236,7 +1236,7 @@ void SDMEditor::PasteLayer()
 
 			if (UDMMaterialLayerObject* PastedLayer = UDMMaterialLayerObject::DeserializeFromString(Slot, SerializeString))
 			{
-				FScopedTransaction Transaction(LOCTEXT("PasteLayer", "Material Designer Paste Layer"));
+				FDMScopedUITransaction Transaction(LOCTEXT("PasteLayer", "Material Designer Paste Layer"));
 				Slot->Modify();
 				Slot->PasteLayer(PastedLayer);
 
@@ -1258,7 +1258,7 @@ void SDMEditor::DuplicateSelectedLayer()
 	FPlatformApplicationMisc::ClipboardPaste(PastedText);
 
 	// Added here to set the transaction description
-	FScopedTransaction Transaction(LOCTEXT("DuplicateLayer", "Material Designer Duplicate Layer"));
+	FDMScopedUITransaction Transaction(LOCTEXT("DuplicateLayer", "Material Designer Duplicate Layer"));
 
 	CopySelectedLayer();
 	PasteLayer();
@@ -1387,7 +1387,7 @@ void SDMEditor::OnDomainChanged(const EMaterialDomain InDomain)
 {
 	if (UDynamicMaterialModelEditorOnlyData* ModelEditorOnlyData = UDynamicMaterialModelEditorOnlyData::Get(MaterialModelWeak))
 	{
-		FScopedTransaction Transaction(LOCTEXT("ChangeDomain", "Material Designer Change Domain"));
+		FDMScopedUITransaction Transaction(LOCTEXT("ChangeDomain", "Material Designer Change Domain"));
 		ModelEditorOnlyData->Modify();
 		ModelEditorOnlyData->SetDomain(InDomain);
 
@@ -1419,7 +1419,7 @@ void SDMEditor::OnBlendModeChanged(const EBlendMode InBlendMode)
 {
 	if (UDynamicMaterialModelEditorOnlyData* ModelEditorOnlyData = UDynamicMaterialModelEditorOnlyData::Get(MaterialModelWeak))
 	{
-		FScopedTransaction Transaction(LOCTEXT("ChangeBlendMode", "Material Designer Change Blend Mode"));
+		FDMScopedUITransaction Transaction(LOCTEXT("ChangeBlendMode", "Material Designer Change Blend Mode"));
 		ModelEditorOnlyData->Modify();
 		ModelEditorOnlyData->SetBlendMode(InBlendMode);
 
@@ -1457,7 +1457,7 @@ void SDMEditor::OnMaterialUnlitChanged(const ECheckBoxState InNewCheckState)
 {
 	if (UDynamicMaterialModelEditorOnlyData* ModelEditorOnlyData = UDynamicMaterialModelEditorOnlyData::Get(MaterialModelWeak))
 	{
-		FScopedTransaction Transaction(LOCTEXT("ChangeShadingModel", "Material Designer Change Shading Model"));
+		FDMScopedUITransaction Transaction(LOCTEXT("ChangeShadingModel", "Material Designer Change Shading Model"));
 		ModelEditorOnlyData->Modify();
 		ModelEditorOnlyData->SetShadingModel(InNewCheckState == ECheckBoxState::Checked
 			? EDMMaterialShadingModel::Unlit : EDMMaterialShadingModel::DefaultLit);
@@ -1507,7 +1507,7 @@ void SDMEditor::OnMaterialAnimatedChanged(const ECheckBoxState InNewCheckState)
 {
 	if (UDynamicMaterialModelEditorOnlyData* ModelEditorOnlyData = UDynamicMaterialModelEditorOnlyData::Get(MaterialModelWeak))
 	{
-		FScopedTransaction Transaction(LOCTEXT("ToggleAnimated", "Material Designer Toggle Animated"));
+		FDMScopedUITransaction Transaction(LOCTEXT("ToggleAnimated", "Material Designer Toggle Animated"));
 		ModelEditorOnlyData->Modify();
 		ModelEditorOnlyData->SetPixelAnimationFlag(InNewCheckState == ECheckBoxState::Checked);
 	}
diff --git a/Engine/Plugins/Experimental/DynamicMaterial/Source/DynamicMaterialEditor/Private/Slate/SDMSlot.cpp b/Engine/Plugins/Experimental/DynamicMaterial/Source/DynamicMaterialEditor/Private/Slate/SDMSlot.cpp
--- a/Engine/Plugins/Experimental/DynamicMaterial/Source/DynamicMaterialEditor/Private/Slate/SDMSlot.cpp
+++ b/Engine/Plugins/Experimental/DynamicMaterial/Source/DynamicMaterialEditor/Private/Slate/SDMSlot.cpp
@@ -32,7 +32,6 @@
 #include "Model/DynamicMaterialModel.h"
 #include "Model/DynamicMaterialModelEditorOnlyData.h"
 #include "Properties/Editors/SDMPropertyEditOpacity.h"
-#include "ScopedTransaction.h"
 #include "Slate/Layers/SDMSlotLayerItem.h"
 #include "Slate/Layers/SDMSlotLayerView.h"
 #include "Slate/Properties/SDMMaterialProperty.h"
@@ -336,7 +335,7 @@ void SDMSlot::OnSourceBlendTypedSelected(const TSubclassOf<UDMMaterialStageBlend
 			{
 				if (UDMMaterialStage* BaseStage = BaseStageWidget->GetStage())
 				{
-					FScopedTransaction Transaction(LOCTEXT("SetStageBlendMode", "Material Designer Set Blend Mode"));
+					FDMScopedUITransaction Transaction(LOCTEXT("SetStageBlendMode", "Material Designer Set Blend Mode"));
 					BaseStage->Modify();
 					BaseStage->ChangeSource<UDMMaterialStageBlend>(InNewItem);
 
@@ -982,7 +981,7 @@ void SDMSlot::StartTransaction(const FText InDescription)
 		return;
 	}
 
-	ScrubbingTransaction = MakeShared<FScopedTransaction>(InDescription);
+	ScrubbingTransaction = MakeShared<FDMScopedUITransaction>(InDescription);
 }
 
 void SDMSlot::EndTransaction()
@@ -1171,7 +1170,7 @@ void SDMSlot::RemoveLayerByIndex(const int32 InLayerIndex, const bool bInSelectN
 		return;
 	}
 
-	FScopedTransaction Transaction(LOCTEXT("RemoveLayer", "Material Designer Remove Layer"));
+	FDMScopedUITransaction Transaction(LOCTEXT("RemoveLayer", "Material Designer Remove Layer"));
 	Slot->Modify();
 	Slot->RemoveLayer(LayerToRemove);
 
@@ -1218,7 +1217,7 @@ void SDMSlot::RemoveSlot()
 	{
 		if (UDynamicMaterialModelEditorOnlyData* ModelEditorOnlyData = Slot->GetMaterialModelEditorOnlyData())
 		{
-			FScopedTransaction Transaction(LOCTEXT("RemoveSlot", "Material Designer Remove Slot"));
+			FDMScopedUITransaction Transaction(LOCTEXT("RemoveSlot", "Material Designer Remove Slot"));
 			ModelEditorOnlyData->Modify();
 			ModelEditorOnlyData->RemoveSlot(Slot->GetIndex());
 
@@ -1246,7 +1245,7 @@ void SDMSlot::AddPropertyToSlot(EDMMaterialPropertyType Property)
 		return;
 	}
 
-	FScopedTransaction Transaction(LOCTEXT("AssignPropertyToSlot", "Assign Property To Slot"));
+	FDMScopedUITransaction Transaction(LOCTEXT("AssignPropertyToSlot", "Assign Property To Slot"));
 	ModelEditorOnlyData->Modify();
 	ModelEditorOnlyData->AssignMaterialPropertyToSlot(Property, Slot);
 }
@@ -1275,7 +1274,7 @@ UDMMaterialLayerObject* SDMSlot::AddNewLayer(UDMMaterialStage* InNewBaseStage, U
 	{
 		const FDMUpdateGuard Guard;
 
-		FScopedTransaction Transaction(LOCTEXT("AddLayer", "Material Designer Add Layer"));
+		FDMScopedUITransaction Transaction(LOCTEXT("AddLayer", "Material Designer Add Layer"));
 		Slot->Modify();
 
 		EDMMaterialPropertyType MaterialProperty = EDMMaterialPropertyType::None;
@@ -1450,7 +1449,7 @@ void SDMSlot::AddNewLayer_Expression(TSubclassOf<UDMMaterialStageExpression> InE
 	}
 
 	// Extra transaction here because there are changes to the slot afterwards.
-	FScopedTransaction Transaction(LOCTEXT("AddLayerExpression", "Material Designer Add Layer (Expression)"));
+	FDMScopedUITransaction Transaction(LOCTEXT("AddLayerExpression", "Material Designer Add Layer (Expression)"));
 	Slot->Modify();
 
 	UDMMaterialStage* NewStage = UDMMaterialStageBlend::CreateStage(UDMMaterialStageBlendNormal::StaticClass());
diff --git a/Engine/Plugins/Experimental/DynamicMaterial/Source/DynamicMaterialEditor/Private/Slate/SDMSlot.h b/Engine/Plugins/Experimental/DynamicMaterial/Source/DynamicMaterialEditor/Private/Slate/SDMSlot.h
--- a/Engine/Plugins/Experimental/DynamicMaterial/Source/DynamicMaterialEditor/Private/Slate/SDMSlot.h
+++ b/Engine/Plugins/Experimental/DynamicMaterial/Source/DynamicMaterialEditor/Private/Slate/SDMSlot.h
@@ -24,6 +24,7 @@ class UDMMaterialStageInputValue;
 class UDMMaterialValue;
 class UDMMaterialValueFloat1;
 enum class EDMMaterialLayerStage : uint8;
+struct FDMScopedUITransaction;
 struct FPropertyChangedEvent;
 
 class SDMSlot : public SCompoundWidget
@@ -103,7 +104,7 @@ protected:
 
 	TSharedPtr<SDMSlotLayerView> LayerView;
 
-	TSharedPtr<FScopedTransaction> ScrubbingTransaction;
+	TSharedPtr<FDMScopedUITransaction> ScrubbingTransaction;
 	float OriginalOpacity;
 
 	TWeakObjectPtr<UDMMaterialStageInputValue> SelectedOpacityStageInputValue;
diff --git a/Engine/Plugins/Experimental/DynamicMaterial/Source/DynamicMaterialEditor/Private/Slate/SDMStage.cpp b/Engine/Plugins/Experimental/DynamicMaterial/Source/DynamicMaterialEditor/Private/Slate/SDMStage.cpp
--- a/Engine/Plugins/Experimental/DynamicMaterial/Source/DynamicMaterialEditor/Private/Slate/SDMStage.cpp
+++ b/Engine/Plugins/Experimental/DynamicMaterial/Source/DynamicMaterialEditor/Private/Slate/SDMStage.cpp
@@ -12,12 +12,12 @@
 #include "Components/MaterialStageInputs/DMMSIExpression.h"
 #include "Components/MaterialStageInputs/DMMSIValue.h"
 #include "Components/MaterialValues/DMMaterialValueTexture.h"
+#include "DMPrivate.h"
 #include "DragAndDrop/AssetDragDropOp.h"
 #include "DragDrop/DMStageDragDropOperation.h"
 #include "DynamicMaterialEditorSettings.h"
 #include "DynamicMaterialEditorStyle.h"
 #include "Engine/Texture.h"
-#include "ScopedTransaction.h"
 #include "SDMEditor.h"
 #include "Slate/Previews/SDMStagePreview.h"
 #include "Slate/SMaterialToolTip.h"
@@ -557,7 +557,7 @@ void SDMStage::HandleAssetDragDropOperation(FAssetDragDropOp& AssetDragDropOpera
 		return;
 	}
 
-	FScopedTransaction Transaction(LOCTEXT("DragTextureOntoStage", "Material Designer Drag Texture onto Stage"), !FDMInitializationGuard::IsInitializing());
+	FDMScopedUITransaction Transaction(LOCTEXT("DragTextureOntoStage", "Material Designer Drag Texture onto Stage"), !FDMInitializationGuard::IsInitializing());
 	UDMMaterialValueTexture* TextureValue = nullptr;
 
 	DraggedOverStage->Modify();
