Material Designer
- Effects list toggle button is no longer visible on layers without effects.

diff --git a/Engine/Plugins/Experimental/DynamicMaterial/Source/DynamicMaterialEditor/Private/Slate/Layers/SDMSlotLayerItem.cpp b/Engine/Plugins/Experimental/DynamicMaterial/Source/DynamicMaterialEditor/Private/Slate/Layers/SDMSlotLayerItem.cpp
--- a/Engine/Plugins/Experimental/DynamicMaterial/Source/DynamicMaterialEditor/Private/Slate/Layers/SDMSlotLayerItem.cpp
+++ b/Engine/Plugins/Experimental/DynamicMaterial/Source/DynamicMaterialEditor/Private/Slate/Layers/SDMSlotLayerItem.cpp
@@ -439,10 +439,10 @@ TSharedRef<SWidget> SDMSlotLayerItem::CreateLayerLinkToggleButton()
 		.ButtonStyle(FDynamicMaterialEditorStyle::Get(), "HoverHintOnly")
 		.ToolTipText(LOCTEXT("MaterialStageLinkTooltip", "Click to toggle UV Link on and off."))
 		.OnClicked(this, &SDMSlotLayerItem::OnLayerLinkToggleButton)
+		.Visibility(this, &SDMSlotLayerItem::GetLayerLinkToggleButtonVisibility)
 		[
 			SNew(SImage)
 			.Image(this, &SDMSlotLayerItem::GetLayerLinkToggleButtonImage)
-			.Visibility(this, &SDMSlotLayerItem::GetLayerLinkToggleButtonVisibility)
 		];
 }
 
@@ -454,6 +454,7 @@ TSharedRef<SWidget> SDMSlotLayerItem::CreateEffectsToggleButton()
 		.ButtonStyle(FDynamicMaterialEditorStyle::Get(), "HoverHintOnly")
 		.ToolTipText(LOCTEXT("MaterialLayerFxTooltip", "Toggle the visibility of this slot layer's effects.\n\nNot implemented."))
 		.OnClicked(this, &SDMSlotLayerItem::OnEffectsToggleButtonClicked)
+		.Visibility(this, &SDMSlotLayerItem::GetEffectsToggleButtonVisibility)
 		[
 			SNew(SImage)
 			.Image(this, &SDMSlotLayerItem::GetEffectsToggleButtonImage)
@@ -1195,6 +1196,25 @@ const FSlateBrush* SDMSlotLayerItem::GetEffectsToggleButtonImage() const
 	return Hidden;
 }
 
+EVisibility SDMSlotLayerItem::GetEffectsToggleButtonVisibility() const
+{
+	if (LayerItem.IsValid())
+	{
+		if (UDMMaterialLayerObject* Layer = LayerItem->GetLayer())
+		{
+			if (UDMMaterialEffectStack* EffectStack = Layer->GetEffectStack())
+			{
+				if (!EffectStack->GetEffects().IsEmpty())
+				{
+					return EVisibility::Visible;
+				}
+			}
+		}
+	}
+
+	return EVisibility::Hidden;
+}
+
 FVector2D SDMSlotLayerItem::GetStagePreviewSize() const
 {
 	return FVector2D(PreviewSize.Get());
diff --git a/Engine/Plugins/Experimental/DynamicMaterial/Source/DynamicMaterialEditor/Private/Slate/Layers/SDMSlotLayerItem.h b/Engine/Plugins/Experimental/DynamicMaterial/Source/DynamicMaterialEditor/Private/Slate/Layers/SDMSlotLayerItem.h
--- a/Engine/Plugins/Experimental/DynamicMaterial/Source/DynamicMaterialEditor/Private/Slate/Layers/SDMSlotLayerItem.h
+++ b/Engine/Plugins/Experimental/DynamicMaterial/Source/DynamicMaterialEditor/Private/Slate/Layers/SDMSlotLayerItem.h
@@ -156,6 +156,7 @@ protected:
 
 	FReply OnEffectsToggleButtonClicked();
 	const FSlateBrush* GetEffectsToggleButtonImage() const;
+	EVisibility GetEffectsToggleButtonVisibility() const;
 
 	FVector2D GetStagePreviewSize() const;
 
