Material Designer
- Export material to texture (hdr, exr and png)

diff --git a/Engine/Plugins/Experimental/DynamicMaterial/DynamicMaterial.uplugin b/Engine/Plugins/Experimental/DynamicMaterial/DynamicMaterial.uplugin
--- a/Engine/Plugins/Experimental/DynamicMaterial/DynamicMaterial.uplugin
+++ b/Engine/Plugins/Experimental/DynamicMaterial/DynamicMaterial.uplugin
@@ -24,6 +24,11 @@
 			"Name": "DynamicMaterialEditor",
 			"Type": "Editor",
 			"LoadingPhase": "Default"
+		},
+		{
+			"Name": "DynamicMaterialShadersEditor",
+			"Type": "Editor",
+			"LoadingPhase": "PostConfigInit"
 		}
 	],
 	"Plugins": [
diff --git a/Engine/Plugins/Experimental/DynamicMaterial/Shaders/DMAlphaOneMinusPS.usf b/Engine/Plugins/Experimental/DynamicMaterial/Shaders/DMAlphaOneMinusPS.usf
new file mode 100644
--- /dev/null
+++ b/Engine/Plugins/Experimental/DynamicMaterial/Shaders/DMAlphaOneMinusPS.usf
@@ -0,0 +1,17 @@
+// Copyright Epic Games, Inc. All Rights Reserved.
+
+#include "/Engine/Private/Common.ush"
+
+Texture2D InputTexture;
+SamplerState InputSampler;
+
+// Shader to apply 1-x to the alpha channel
+void MainPS(
+	noperspective float4 UVAndScreenPos : TEXCOORD0,
+	out float4 OutColor : SV_Target0
+)
+{
+	float2 UV = UVAndScreenPos.xy;
+	OutColor = InputTexture.Sample(InputSampler, UV);
+	OutColor.a = 1.f - OutColor.a;
+}
diff --git a/Engine/Plugins/Experimental/DynamicMaterial/Source/DynamicMaterialEditor/DynamicMaterialEditor.Build.cs b/Engine/Plugins/Experimental/DynamicMaterial/Source/DynamicMaterialEditor/DynamicMaterialEditor.Build.cs
--- a/Engine/Plugins/Experimental/DynamicMaterial/Source/DynamicMaterialEditor/DynamicMaterialEditor.Build.cs
+++ b/Engine/Plugins/Experimental/DynamicMaterial/Source/DynamicMaterialEditor/DynamicMaterialEditor.Build.cs
@@ -7,7 +7,6 @@ public class DynamicMaterialEditor : ModuleRules
 	public DynamicMaterialEditor(ReadOnlyTargetRules Target) : base(Target)
 	{
 		PCHUsage = ModuleRules.PCHUsageMode.UseExplicitOrSharedPCHs;
-		bUseUnity = false;
 
 		PublicDependencyModuleNames.AddRange(
 			new string[]
@@ -28,6 +27,7 @@ public class DynamicMaterialEditor : ModuleRules
 				"ContentBrowserData",
 				"CustomDetailsView",
 				"DeveloperSettings",
+				"DynamicMaterialShadersEditor",
 				"EditorWidgets",
 				"Engine",
 				"InputCore",
@@ -37,6 +37,8 @@ public class DynamicMaterialEditor : ModuleRules
 				"Projects",
 				"PropertyEditor",
 				"RenderCore",
+				"Renderer",
+				"RHI",
 				"Slate",
 				"SlateCore",
 				"ToolMenus",
diff --git a/Engine/Plugins/Experimental/DynamicMaterial/Source/DynamicMaterialEditor/Private/DynamicMaterialEditorModule.cpp b/Engine/Plugins/Experimental/DynamicMaterial/Source/DynamicMaterialEditor/Private/DynamicMaterialEditorModule.cpp
--- a/Engine/Plugins/Experimental/DynamicMaterial/Source/DynamicMaterialEditor/Private/DynamicMaterialEditorModule.cpp
+++ b/Engine/Plugins/Experimental/DynamicMaterial/Source/DynamicMaterialEditor/Private/DynamicMaterialEditorModule.cpp
@@ -1,10 +1,8 @@
 // Copyright Epic Games, Inc. All Rights Reserved.
 
 #include "DynamicMaterialEditorModule.h"
-#include "AssetToolsModule.h"
 #include "Components/DMMaterialComponent.h"
 #include "Components/DMMaterialEffectFunction.h"
-#include "Components/DMMaterialStageSource.h"
 #include "Components/DMMaterialStageThroughput.h"
 #include "Components/DMTextureUV.h"
 #include "Components/MaterialStageInputs/DMMSIFunction.h"
@@ -18,12 +16,11 @@
 #include "Components/MaterialValues/DMMaterialValueFloat4.h"
 #include "Components/MaterialValues/DMMaterialValueTexture.h"
 #include "Components/PrimitiveComponent.h"
-#include "DMMaterialFunctionLibrary.h"
-#include "DMWorldSubsystem.h"
-#include "DetailLayoutBuilder.h"
 #include "DetailsPanel/DMMaterialInterfaceTypeCustomizer.h"
 #include "DetailsPanel/DMPropertyTypeCustomizer.h"
 #include "DetailsPanel/Slate/SDMMaterialListExtensionWidget.h"
+#include "DMMaterialFunctionLibrary.h"
+#include "DMWorldSubsystem.h"
 #include "DynamicMaterialEditorCommands.h"
 #include "DynamicMaterialEditorSettings.h"
 #include "DynamicMaterialEditorStyle.h"
diff --git a/Engine/Plugins/Experimental/DynamicMaterial/Source/DynamicMaterialEditor/Private/Helpers/DMMaterialSnapshotLibrary.cpp b/Engine/Plugins/Experimental/DynamicMaterial/Source/DynamicMaterialEditor/Private/Helpers/DMMaterialSnapshotLibrary.cpp
new file mode 100644
--- /dev/null
+++ b/Engine/Plugins/Experimental/DynamicMaterial/Source/DynamicMaterialEditor/Private/Helpers/DMMaterialSnapshotLibrary.cpp
@@ -0,0 +1,174 @@
+// Copyright Epic Games, Inc. All Rights Reserved.
+
+#include "Helpers/DMMaterialSnapshotLibrary.h"
+#include "DMAlphaOneMinusPS.h"
+#include "DynamicMaterialEditorModule.h"
+#include "Engine/Engine.h"
+#include "Engine/Canvas.h"
+#include "Framework/Notifications/NotificationManager.h"
+#include "ImageUtils.h"
+#include "Kismet/KismetRenderingLibrary.h"
+#include "Materials/MaterialInterface.h"
+#include "RenderGraphBuilder.h"
+#include "ScreenPass.h"
+#include "TextureResource.h"
+#include "Widgets/Notifications/SNotificationList.h"
+#include "UObject/Object.h"
+#include "UObject/UObjectGlobals.h"
+#include "UObject/Package.h"
+
+#define LOCTEXT_NAMESPACE "DMMaterialShapshotLibrary"
+
+namespace UE::DynamicMaterial::Private
+{
+	void RenderMaterialToRenderTarget(UMaterialInterface* InMaterial, UTextureRenderTarget2D* InRenderTarget)
+	{
+		InMaterial->EnsureIsComplete();
+
+		FTextureRenderTargetResource* RenderTargetResource = InRenderTarget->GameThread_GetRenderTargetResource();
+
+		UCanvas* Canvas = NewObject<UCanvas>(GetTransientPackage());
+
+		FCanvas RenderCanvas(
+			RenderTargetResource,
+			nullptr,
+			FGameTime::CreateUndilated(0, 0),
+			GEngine->GetDefaultWorldFeatureLevel()
+		);
+
+		Canvas->Init(InRenderTarget->SizeX, InRenderTarget->SizeY, nullptr, &RenderCanvas);
+
+		{
+			SCOPED_DRAW_EVENTF_GAMETHREAD(DrawMaterialToRenderTarget, TEXT("DrawMaterialToRenderTarget: %s"), InRenderTarget->GetFName());
+
+			ENQUEUE_RENDER_COMMAND(FlushDeferredResourceUpdateCommand)(
+				[RenderTargetResource](FRHICommandListImmediate& RHICmdList)
+				{
+					RenderTargetResource->FlushDeferredResourceUpdate(RHICmdList);
+				});
+
+			Canvas->K2_DrawMaterial(InMaterial, FVector2D(0, 0), FVector2D(InRenderTarget->SizeX, InRenderTarget->SizeY), FVector2D(0, 0));
+
+			RenderCanvas.Flush_GameThread();
+			Canvas->Canvas = nullptr;
+
+			//UpdateResourceImmediate must be called here to ensure mips are generated.
+			InRenderTarget->UpdateResourceImmediate(false);
+
+			ENQUEUE_RENDER_COMMAND(ResetSceneTextureExtentHistory)(
+				[RenderTargetResource](FRHICommandListImmediate& RHICmdList)
+				{
+					RenderTargetResource->ResetSceneTextureExtentsHistory();
+				});
+		}
+	}
+
+	UTextureRenderTarget2D* CreateSnapshotRenderTarget(const FIntPoint& InTextureSize)
+	{
+		UTextureRenderTarget2D* RenderTarget = NewObject<UTextureRenderTarget2D>(GetTransientPackage());
+		check(RenderTarget);
+		RenderTarget->RenderTargetFormat = RTF_RGBA32f;
+		RenderTarget->ClearColor = FLinearColor::Black;
+		RenderTarget->bAutoGenerateMips = false;
+		RenderTarget->bCanCreateUAV = false;
+		RenderTarget->InitAutoFormat(InTextureSize.X, InTextureSize.Y);
+		RenderTarget->UpdateResourceImmediate(true);
+
+		return RenderTarget;
+	}
+
+	void ApplyAlphaOneMinusShader(FRHICommandListImmediate& InRHICmdList, FTextureRenderTargetResource* InSourceTextureResource, FTextureRenderTargetResource* InDestTargetResource)
+	{
+		FRDGBuilder GraphBuilder(InRHICmdList);
+
+		const FTextureRHIRef& SourceTexture = InSourceTextureResource->GetRenderTargetTexture();
+		const FTextureRHIRef& DestTarget = InDestTargetResource->GetRenderTargetTexture();
+
+		FRDGTexture* InputTexture = GraphBuilder.RegisterExternalTexture(CreateRenderTarget(SourceTexture, TEXT("SourceTexture")));
+		FRDGTexture* OutputTexture = GraphBuilder.RegisterExternalTexture(CreateRenderTarget(DestTarget, TEXT("DestTarget")));
+		const FGlobalShaderMap* GlobalShaderMap = GetGlobalShaderMap(GMaxRHIFeatureLevel);
+
+		// The formats or sizes differ to pixel shader stuff
+		// Configure source/output viewport to get the right UV scaling from source texture to output texture
+		const FScreenPassTextureViewport InputViewport(InputTexture);
+		const FScreenPassTextureViewport OutputViewport(OutputTexture);
+
+		const TShaderMapRef<FScreenPassVS> VertexShader(GlobalShaderMap);
+
+		// Rectangle area to use from source
+		const FIntRect ViewRect(FIntPoint(0, 0), InputTexture->Desc.Extent);
+
+		//Dummy ViewFamily/ViewInfo created to use built in Draw Screen/Texture Pass
+		const FSceneViewFamily ViewFamily(FSceneViewFamily::ConstructionValues(nullptr, nullptr, FEngineShowFlags(ESFIM_Game))
+			.SetTime(FGameTime()));
+		FSceneViewInitOptions ViewInitOptions;
+		ViewInitOptions.ViewFamily = &ViewFamily;
+		ViewInitOptions.SetViewRectangle(ViewRect);
+		ViewInitOptions.ViewOrigin = FVector::ZeroVector;
+		ViewInitOptions.ViewRotationMatrix = FMatrix::Identity;
+		ViewInitOptions.ProjectionMatrix = FMatrix::Identity;
+		const FSceneView ViewInfo(ViewInitOptions);
+
+		const TShaderMapRef<FDMAlphaOneMinusPS> PixelShader(GlobalShaderMap);
+		FDMAlphaOneMinusPS::FParameters* Parameters = PixelShader->AllocateAndSetParameters(GraphBuilder, InputTexture, OutputTexture);
+		AddDrawScreenPass(GraphBuilder, RDG_EVENT_NAME("ApplyAlphaOneMinusShader"), ViewInfo, OutputViewport, InputViewport, VertexShader, PixelShader, Parameters);
+
+		GraphBuilder.Execute();
+	}
+
+	UTextureRenderTarget2D* ApplyAlphaOneMinusShader(UTextureRenderTarget2D* InRenderTarget)
+	{
+		FTextureRenderTargetResource* Source = InRenderTarget->GameThread_GetRenderTargetResource();
+
+		UTextureRenderTarget2D* OutRenderTarget = CreateSnapshotRenderTarget({InRenderTarget->SizeX, InRenderTarget->SizeY});
+		FTextureRenderTargetResource* Target = OutRenderTarget->GameThread_GetRenderTargetResource();
+
+		ENQUEUE_RENDER_COMMAND(ApplyAlphaOneMinusCommand)(
+			[Source, Target](FRHICommandListImmediate& InRHICmdList)
+			{
+				ApplyAlphaOneMinusShader(InRHICmdList, Source, Target);
+			});
+
+		return OutRenderTarget;
+	}
+}
+
+bool FDMMaterialShapshotLibrary::SnapshotMaterial(UMaterialInterface* InMaterial, const FIntPoint& InTextureSize, const FString& InSavePath)
+{
+	UTextureRenderTarget2D* RenderTarget = UE::DynamicMaterial::Private::CreateSnapshotRenderTarget(InTextureSize);
+
+	if (!RenderTarget)
+	{
+		UE_LOG(LogDynamicMaterialEditor, Warning, TEXT("Unable to create render target for material snapshot."));
+		return false;
+	}
+
+	UE::DynamicMaterial::Private::RenderMaterialToRenderTarget(InMaterial, RenderTarget);
+
+	// Render target has inverted alpha - fix that.
+	UTextureRenderTarget2D* FixedAlphaRenderTarget = UE::DynamicMaterial::Private::ApplyAlphaOneMinusShader(RenderTarget);
+
+	FImage Image;
+	FImageUtils::GetRenderTargetImage(FixedAlphaRenderTarget, Image);
+
+	FImageUtils::SaveImageByExtension(*InSavePath, Image);
+
+	// Make a toast!
+	const FString HyperLinkText = FPaths::ConvertRelativePathToFull(InSavePath);
+
+	FNotificationInfo Info(LOCTEXT("SnapshotCreated", "Snapshot created."));
+
+	Info.Hyperlink = FSimpleDelegate::CreateStatic([](FString SourceFilePath)
+		{
+			FPlatformProcess::ExploreFolder(*(FPaths::GetPath(SourceFilePath)));
+		}, HyperLinkText);
+
+	Info.HyperlinkText = FText::FromString(HyperLinkText);
+	Info.ExpireDuration = 3.0f;
+
+	FSlateNotificationManager::Get().AddNotification(Info);
+
+	return true;
+}
+
+#undef LOCTEXT_NAMESPACE
diff --git a/Engine/Plugins/Experimental/DynamicMaterial/Source/DynamicMaterialEditor/Private/Helpers/DMMaterialSnapshotLibrary.h b/Engine/Plugins/Experimental/DynamicMaterial/Source/DynamicMaterialEditor/Private/Helpers/DMMaterialSnapshotLibrary.h
new file mode 100644
--- /dev/null
+++ b/Engine/Plugins/Experimental/DynamicMaterial/Source/DynamicMaterialEditor/Private/Helpers/DMMaterialSnapshotLibrary.h
@@ -0,0 +1,14 @@
+// Copyright Epic Games, Inc. All Rights Reserved.
+
+#pragma once
+
+#include "Math/MathFwd.h"
+#include "Containers/UnrealString.h"
+
+class UMaterialInterface;
+
+class FDMMaterialShapshotLibrary
+{
+public:
+	static bool SnapshotMaterial(UMaterialInterface* InMaterial, const FIntPoint& InTextureSize, const FString& InSavePath);
+};
diff --git a/Engine/Plugins/Experimental/DynamicMaterial/Source/DynamicMaterialEditor/Private/Menus/DMToolBarMenus.cpp b/Engine/Plugins/Experimental/DynamicMaterial/Source/DynamicMaterialEditor/Private/Menus/DMToolBarMenus.cpp
--- a/Engine/Plugins/Experimental/DynamicMaterial/Source/DynamicMaterialEditor/Private/Menus/DMToolBarMenus.cpp
+++ b/Engine/Plugins/Experimental/DynamicMaterial/Source/DynamicMaterialEditor/Private/Menus/DMToolBarMenus.cpp
@@ -2,12 +2,17 @@
 
 #include "DMToolBarMenus.h"
 #include "ContentBrowserModule.h"
+#include "DesktopPlatformModule.h"
 #include "DMBlueprintFunctionLibrary.h"
 #include "DMPrivate.h"
 #include "DynamicMaterialEditorModule.h"
 #include "DynamicMaterialEditorSettings.h"
+#include "Engine/Engine.h"
+#include "Engine/Texture2D.h"
 #include "EngineAnalytics.h"
+#include "Framework/Application/SlateApplication.h"
 #include "IContentBrowserSingleton.h"
+#include "IDesktopPlatform.h"
 #include "ISinglePropertyView.h"
 #include "Material/DynamicMaterialInstance.h"
 #include "Materials/Material.h"
@@ -19,6 +24,7 @@
 #include "Slate/SDMSlot.h"
 #include "ToolMenu.h"
 #include "ToolMenus.h"
+#include "Helpers/DMMaterialSnapshotLibrary.h"
 #include "Widgets/SWidget.h"
 
 #define LOCTEXT_NAMESPACE "FDMToolBarMenus"
@@ -103,6 +109,62 @@ namespace UE::DynamicMaterialEditor::Private
 		}
 	}
 
+	void SnapshotMaterial(TWeakObjectPtr<UDynamicMaterialModel> InMaterialModelWeak, FIntPoint InTextureSize)
+	{
+		UDynamicMaterialModel* MaterialModel = InMaterialModelWeak.Get();
+
+		if (!IsValid(MaterialModel))
+		{
+			return;
+		}
+
+		UMaterialInterface* Material = MaterialModel->GetGeneratedMaterial();
+
+		if (UDynamicMaterialInstance* MaterialInstance = MaterialModel->GetDynamicMaterialInstance())
+		{
+			if (!IsValid(MaterialInstance->Parent.Get()))
+			{
+				UE_LOG(LogDynamicMaterialEditor, Warning, TEXT("Unable to find world to find material instance parent."));
+				return;
+			}
+
+			Material = MaterialInstance;
+		}
+
+		if (!Material)
+		{
+			UE_LOG(LogDynamicMaterialEditor, Warning, TEXT("Unable to find material to snapshot."));
+			return;
+		}
+
+		TArray<FString> OutFilenames;
+
+		if (IDesktopPlatform* DesktopPlatform = FDesktopPlatformModule::Get())
+		{
+			DesktopPlatform->SaveFileDialog(
+				FSlateApplication::Get().FindBestParentWindowHandleForDialogs(nullptr),
+				LOCTEXT("SaveSnapshotAs", "Save Snapshot As").ToString(),
+				FPaths::ProjectSavedDir(),
+				Material->GetName() + "_Snapshot_" + FString::FromInt(InTextureSize.X) + "x" + FString::FromInt(InTextureSize.Y),
+				TEXT("HDR File (*.hdr)|*.hdr|EXR File (*.exr)|*.exr|PNG File (*.png)|*.png"),
+				EFileDialogFlags::None,
+				OutFilenames
+			);
+		}
+
+		if (OutFilenames.Num() == 0)
+		{
+			return;
+		}
+
+		FDMMaterialShapshotLibrary::SnapshotMaterial(Material, InTextureSize, OutFilenames[0]);
+
+		if (FEngineAnalytics::IsAvailable())
+		{
+			FEngineAnalytics::GetProvider().RecordEvent(TEXT("Editor.Usage.MaterialDesigner.SnapshotMaterial"));
+		}
+	}
+
 	void AddToolBarBoolOptionMenuEntry(FToolMenuSection& InSection, const FName& InPropertyName, const FUIAction InAction)
 	{
 		const FProperty* const OptionProperty = UDynamicMaterialEditorSettings::StaticClass()->FindPropertyByName(InPropertyName);
@@ -161,6 +223,74 @@ namespace UE::DynamicMaterialEditor::Private
 		);
 	}
 
+	void CreateSnapshotMaterialMenu(UToolMenu* InMenu)
+	{
+		const UDMMenuContext* const MenuContext = InMenu->FindContext<UDMMenuContext>();
+
+		if (!MenuContext)
+		{
+			return;
+		}
+
+		UDynamicMaterialModel* const MaterialModel = MenuContext->GetModel();
+
+		if (!MaterialModel)
+		{
+			return;
+		}
+
+		TWeakObjectPtr<UDynamicMaterialModel> MaterialModelWeak = MaterialModel;
+		FToolMenuSection& NewSection = InMenu->AddSection("SnapshotMaterial", LOCTEXT("SnapshotMaterial", "Snapshop Material"));
+
+		NewSection.AddMenuEntry(
+			NAME_None,
+			LOCTEXT("SnapshotMaterial512", "512x512"),
+			LOCTEXT("SnapshotMaterial512Tooltip", "Take a snapshot of the material with the current settings and export it as a texture with a resolution of 512x512 pixels."),
+			FSlateIcon(),
+			FUIAction(FExecuteAction::CreateStatic(
+				&UE::DynamicMaterialEditor::Private::SnapshotMaterial,
+				MaterialModelWeak,
+				FIntPoint(512, 512)
+			))
+		);
+
+		NewSection.AddMenuEntry(
+			NAME_None,
+			LOCTEXT("SnapshotMaterial1024", "1024x1024"),
+			LOCTEXT("SnapshotMaterial1024Tooltip", "Take a snapshot of the material with the current settings and export it as a texture with a resolution of 1024x1024 pixels."),
+			FSlateIcon(),
+			FUIAction(FExecuteAction::CreateStatic(
+				&UE::DynamicMaterialEditor::Private::SnapshotMaterial,
+				MaterialModelWeak,
+				FIntPoint(1024, 1024)
+			))
+		);
+
+		NewSection.AddMenuEntry(
+			NAME_None,
+			LOCTEXT("SnapshotMaterial2048", "2048x2048"),
+			LOCTEXT("SnapshotMaterial2048Tooltip", "Take a snapshot of the material with the current settings and export it as a texture with a resolution of 2048x2048 pixels."),
+			FSlateIcon(),
+			FUIAction(FExecuteAction::CreateStatic(
+				&UE::DynamicMaterialEditor::Private::SnapshotMaterial,
+				MaterialModelWeak,
+				FIntPoint(2048, 2048)
+			))
+		);
+
+		NewSection.AddMenuEntry(
+			NAME_None,
+			LOCTEXT("SnapshotMaterial4096", "4096x4096"),
+			LOCTEXT("SnapshotMaterial4096Tooltip", "Take a snapshot of the material with the current settings and export it as a texture with a resolution of 4096x4096 pixels."),
+			FSlateIcon(),
+			FUIAction(FExecuteAction::CreateStatic(
+				&UE::DynamicMaterialEditor::Private::SnapshotMaterial,
+				MaterialModelWeak,
+				FIntPoint(4096, 4096)
+			))
+		);
+	}
+
 	void AddToolbarExportMenu(UToolMenu* InMenu)
 	{
 		if (!IsValid(InMenu) || InMenu->ContainsSection(ToolBarExportSectionName))
@@ -197,7 +327,7 @@ namespace UE::DynamicMaterialEditor::Private
 			return;
 		}
 
-		FToolMenuSection& NewSection = InMenu->AddSection("Export", LOCTEXT("ExportSection", "Export Material"));
+		FToolMenuSection& NewSection = InMenu->AddSection("Export", LOCTEXT("ExportSection", "Export"));
 
 		if (bAllowInstanceExport)
 		{
@@ -220,10 +350,19 @@ namespace UE::DynamicMaterialEditor::Private
 				FSlateIcon(),
 				FUIAction(FExecuteAction::CreateStatic(
 					&UE::DynamicMaterialEditor::Private::ExportMaterialModelFromModel,
-					TWeakObjectPtr<UDynamicMaterialModel>(MaterialModel))
-				)
+					TWeakObjectPtr<UDynamicMaterialModel>(MaterialModel)
+				))
 			);
 		}
+
+		NewSection.AddSubMenu(
+			NAME_None,
+			LOCTEXT("SnapshotMaterial", "Snapshop Material"),
+			LOCTEXT("SnapshotMaterialTooltip", "Take a snapshot of the material with the current settings and export it as a texture."),
+			FNewToolMenuChoice(FNewToolMenuDelegate::CreateStatic(
+				&UE::DynamicMaterialEditor::Private::CreateSnapshotMaterialMenu
+			))
+		);
 	}
 
 	void AddToolBarTooltipOptionsSection(UToolMenu* InMenu)
@@ -314,7 +453,7 @@ namespace UE::DynamicMaterialEditor::Private
 	{
 		AddToolbarExportMenu(InMenu);
 
-		FToolMenuSection& NewSection = InMenu->AddSection("MaterialDesigner", LOCTEXT("MaterialDesignerSection", "MaterialDesigner"));
+		FToolMenuSection& NewSection = InMenu->AddSection("MaterialDesigner", LOCTEXT("MaterialDesignerSection", "Material Designer"));
 
 		UDMMenuContext* MenuContext = InMenu->FindContext<UDMMenuContext>();
 
diff --git a/Engine/Plugins/Experimental/DynamicMaterial/Source/DynamicMaterialShadersEditor/DynamicMaterialShadersEditor.Build.cs b/Engine/Plugins/Experimental/DynamicMaterial/Source/DynamicMaterialShadersEditor/DynamicMaterialShadersEditor.Build.cs
new file mode 100644
--- /dev/null
+++ b/Engine/Plugins/Experimental/DynamicMaterial/Source/DynamicMaterialShadersEditor/DynamicMaterialShadersEditor.Build.cs
@@ -0,0 +1,32 @@
+// Copyright Epic Games, Inc. All Rights Reserved.
+
+using UnrealBuildTool;
+
+public class DynamicMaterialShadersEditor : ModuleRules
+{
+	public DynamicMaterialShadersEditor(ReadOnlyTargetRules Target) : base(Target)
+	{
+		PCHUsage = ModuleRules.PCHUsageMode.UseExplicitOrSharedPCHs;
+
+		PublicDependencyModuleNames.AddRange(
+			new string[]
+			{
+				"Core",
+				"CoreUObject",
+				"RenderCore"
+			}
+		);
+
+		PrivateDependencyModuleNames.AddRange(
+			new string[]
+			{
+				"Engine",
+				"Projects",
+				"Renderer",
+				"RHI"
+			}
+		);
+
+		ShortName = "DynMatEdShdrs";
+	}
+}
diff --git a/Engine/Plugins/Experimental/DynamicMaterial/Source/DynamicMaterialShadersEditor/Private/DMAlphaOneMinusPS.cpp b/Engine/Plugins/Experimental/DynamicMaterial/Source/DynamicMaterialShadersEditor/Private/DMAlphaOneMinusPS.cpp
new file mode 100644
--- /dev/null
+++ b/Engine/Plugins/Experimental/DynamicMaterial/Source/DynamicMaterialShadersEditor/Private/DMAlphaOneMinusPS.cpp
@@ -0,0 +1,26 @@
+// Copyright Epic Games, Inc. All Rights Reserved.
+
+#include "DMAlphaOneMinusPS.h"
+#include "DataDrivenShaderPlatformInfo.h"
+#include "DynamicMaterialShadersEditorModule.h"
+#include "RenderGraphBuilder.h"
+#include "RHIStaticStates.h"
+#include "UObject/Package.h"
+
+const FString FDMAlphaOneMinusPS::ShaderPath = FString(UE::DynamicMaterialShadersEditor::Internal::VirtualShaderMountPoint) + TEXT("/DMAlphaOneMinusPS.usf");
+
+bool FDMAlphaOneMinusPS::ShouldCompilePermutation(const FGlobalShaderPermutationParameters& InParameters)
+{
+	return IsFeatureLevelSupported(InParameters.Platform, ERHIFeatureLevel::ES3_1);
+}
+
+FDMAlphaOneMinusPS::FParameters* FDMAlphaOneMinusPS::AllocateAndSetParameters(FRDGBuilder& InGraphBuilder, FRDGTextureRef InRGBATexture, FRDGTextureRef InOutputTexture)
+{
+	FParameters* Parameters = InGraphBuilder.AllocParameters<FParameters>();
+	Parameters->InputTexture = InRGBATexture;
+	Parameters->InputSampler = TStaticSamplerState<SF_Point>::GetRHI();
+	Parameters->RenderTargets[0] = FRenderTargetBinding{InOutputTexture, ERenderTargetLoadAction::ENoAction};
+	return Parameters;
+}
+
+IMPLEMENT_SHADER_TYPE(, FDMAlphaOneMinusPS, *FDMAlphaOneMinusPS::ShaderPath, TEXT("MainPS"), SF_Pixel);
diff --git a/Engine/Plugins/Experimental/DynamicMaterial/Source/DynamicMaterialShadersEditor/Private/DynamicMaterialShadersEditorModule.cpp b/Engine/Plugins/Experimental/DynamicMaterial/Source/DynamicMaterialShadersEditor/Private/DynamicMaterialShadersEditorModule.cpp
new file mode 100644
--- /dev/null
+++ b/Engine/Plugins/Experimental/DynamicMaterial/Source/DynamicMaterialShadersEditor/Private/DynamicMaterialShadersEditorModule.cpp
@@ -0,0 +1,20 @@
+// Copyright Epic Games, Inc. All Rights Reserved.
+
+#include "DynamicMaterialShadersEditorModule.h"
+#include "DMAlphaOneMinusPS.h"
+#include "Interfaces/IPluginManager.h"
+#include "Modules/ModuleManager.h"
+
+void FDynamicMaterialShadersEditorModule::StartupModule()
+{
+	const TSharedPtr<IPlugin> Plugin = IPluginManager::Get().FindPlugin(UE_PLUGIN_NAME);
+	check(Plugin.IsValid());
+
+	const FString PluginShaderDir = FPaths::Combine(Plugin->GetBaseDir(), TEXT("Shaders"));
+	AddShaderSourceDirectoryMapping(UE::DynamicMaterialShadersEditor::Internal::VirtualShaderMountPoint, PluginShaderDir);
+
+	// Force init
+	FDMAlphaOneMinusPS::GetStaticType();
+}
+
+IMPLEMENT_MODULE(FDynamicMaterialShadersEditorModule, DynamicMaterialShadersEditor)
diff --git a/Engine/Plugins/Experimental/DynamicMaterial/Source/DynamicMaterialShadersEditor/Private/DynamicMaterialShadersEditorModule.h b/Engine/Plugins/Experimental/DynamicMaterial/Source/DynamicMaterialShadersEditor/Private/DynamicMaterialShadersEditorModule.h
new file mode 100644
--- /dev/null
+++ b/Engine/Plugins/Experimental/DynamicMaterial/Source/DynamicMaterialShadersEditor/Private/DynamicMaterialShadersEditorModule.h
@@ -0,0 +1,17 @@
+// Copyright Epic Games, Inc. All Rights Reserved.
+
+#pragma once
+
+#include "Modules/ModuleInterface.h"
+
+namespace UE::DynamicMaterialShadersEditor::Internal
+{
+	static inline const TCHAR* VirtualShaderMountPoint = TEXT("/Plugin/MaterialDesigner");
+}
+
+class FDynamicMaterialShadersEditorModule : public IModuleInterface
+{
+	//~ Begin IModuleInterface interface
+	virtual void StartupModule() override;
+	//~ End IModuleInterface interface
+};
diff --git a/Engine/Plugins/Experimental/DynamicMaterial/Source/DynamicMaterialShadersEditor/Public/DMAlphaOneMinusPS.h b/Engine/Plugins/Experimental/DynamicMaterial/Source/DynamicMaterialShadersEditor/Public/DMAlphaOneMinusPS.h
new file mode 100644
--- /dev/null
+++ b/Engine/Plugins/Experimental/DynamicMaterial/Source/DynamicMaterialShadersEditor/Public/DMAlphaOneMinusPS.h
@@ -0,0 +1,29 @@
+// Copyright Epic Games, Inc. All Rights Reserved.
+
+#pragma once
+
+#include "GlobalShader.h"
+#include "ShaderParameterStruct.h"
+
+class FRDGBuilder;
+struct FGlobalShaderPermutationParameters;
+using FRDGTextureRef = class FRDGTexture*;
+
+class FDMAlphaOneMinusPS : public FGlobalShader
+{
+	DECLARE_EXPORTED_GLOBAL_SHADER(FDMAlphaOneMinusPS, DYNAMICMATERIALSHADERSEDITOR_API);
+	SHADER_USE_PARAMETER_STRUCT(FDMAlphaOneMinusPS, FGlobalShader);
+
+	BEGIN_SHADER_PARAMETER_STRUCT(FParameters, )
+		SHADER_PARAMETER_RDG_TEXTURE(Texture2D, InputTexture)
+		SHADER_PARAMETER_SAMPLER(SamplerState, InputSampler)
+		RENDER_TARGET_BINDING_SLOTS()
+	END_SHADER_PARAMETER_STRUCT()
+
+public:
+	static const FString ShaderPath;
+
+	static DYNAMICMATERIALSHADERSEDITOR_API bool ShouldCompilePermutation(const FGlobalShaderPermutationParameters& InParameters);
+
+	DYNAMICMATERIALSHADERSEDITOR_API FParameters* AllocateAndSetParameters(FRDGBuilder& InGraphBuilder, FRDGTextureRef InRGBATexture, FRDGTextureRef InOutputTexture);
+};
