Material Designer
- Fixed a bug where materials were not compiling straight away.

diff --git a/Engine/Plugins/Experimental/DynamicMaterial/Source/DynamicMaterialEditor/Private/DMBlueprintFunctionLibrary.cpp b/Engine/Plugins/Experimental/DynamicMaterial/Source/DynamicMaterialEditor/Private/DMBlueprintFunctionLibrary.cpp
--- a/Engine/Plugins/Experimental/DynamicMaterial/Source/DynamicMaterialEditor/Private/DMBlueprintFunctionLibrary.cpp
+++ b/Engine/Plugins/Experimental/DynamicMaterial/Source/DynamicMaterialEditor/Private/DMBlueprintFunctionLibrary.cpp
@@ -237,7 +237,17 @@ UDynamicMaterialModel* UDMBlueprintFunctionLibrary::CreateDynamicMaterialInObjec
 		InMaterialProperty.SetMaterial(NewInstance);
 	}
 
-	return NewInstance->GetMaterialModel();
+	if (UDynamicMaterialModel* MaterialModel = NewInstance->GetMaterialModel())
+	{
+		if (IDynamicMaterialModelEditorOnlyDataInterface* EditorOnlyData = MaterialModel->GetEditorOnlyData())
+		{
+			EditorOnlyData->RequestMaterialBuild();
+		}
+
+		return MaterialModel;
+	}
+
+	return nullptr;
 }
 
 bool UDMBlueprintFunctionLibrary::ExportMaterialInstance(UDynamicMaterialModel* InMaterialModel, const FString& InSavePath)
