Material Designer
- Added release resource to screenshot render targets.

diff --git a/Engine/Plugins/Experimental/DynamicMaterial/Source/DynamicMaterialEditor/Private/Helpers/DMMaterialSnapshotLibrary.cpp b/Engine/Plugins/Experimental/DynamicMaterial/Source/DynamicMaterialEditor/Private/Helpers/DMMaterialSnapshotLibrary.cpp
--- a/Engine/Plugins/Experimental/DynamicMaterial/Source/DynamicMaterialEditor/Private/Helpers/DMMaterialSnapshotLibrary.cpp
+++ b/Engine/Plugins/Experimental/DynamicMaterial/Source/DynamicMaterialEditor/Private/Helpers/DMMaterialSnapshotLibrary.cpp
@@ -153,6 +153,9 @@ bool FDMMaterialShapshotLibrary::SnapshotMaterial(UMaterialInterface* InMaterial
 
 	FImageUtils::SaveImageByExtension(*InSavePath, Image);
 
+	RenderTarget->ReleaseResource();
+	FixedAlphaRenderTarget->ReleaseResource();
+
 	// Make a toast!
 	const FString HyperLinkText = FPaths::ConvertRelativePathToFull(InSavePath);
 
