Material Designer: Fix crash when Motion Design launches in -game

diff --git a/Engine/Plugins/Experimental/DynamicMaterial/Source/DynamicMaterialEditor/Private/Components/MaterialStageInputs/DMMSITextureUV.cpp b/Engine/Plugins/Experimental/DynamicMaterial/Source/DynamicMaterialEditor/Private/Components/MaterialStageInputs/DMMSITextureUV.cpp
--- a/Engine/Plugins/Experimental/DynamicMaterial/Source/DynamicMaterialEditor/Private/Components/MaterialStageInputs/DMMSITextureUV.cpp
+++ b/Engine/Plugins/Experimental/DynamicMaterial/Source/DynamicMaterialEditor/Private/Components/MaterialStageInputs/DMMSITextureUV.cpp
@@ -336,7 +336,7 @@ TArray<UMaterialExpression*> UDMMaterialStageInputTextureUV::CreateTextureUVExpr
 
 	if (InTextureUV->GetMirrorOnX() == false && InTextureUV->GetMirrorOnY() == false)
 	{
-		TextureUVFunc = FDynamicMaterialEditorModule::Get().GetFunctionLibrary()->MakeExpression(
+		TextureUVFunc = FDMMaterialFunctionLibrary::Get().MakeExpression(
 			InBuildState->GetDynamicMaterial(),
 			"MF_DM_TextureUVFunc",
 			MaterialFunc_Name_TextureUV_Mirror_None,
@@ -346,7 +346,7 @@ TArray<UMaterialExpression*> UDMMaterialStageInputTextureUV::CreateTextureUVExpr
 	// GetMirrorOnX() == true
 	else if (InTextureUV->GetMirrorOnY() == false)
 	{
-		TextureUVFunc = FDynamicMaterialEditorModule::Get().GetFunctionLibrary()->MakeExpression(
+		TextureUVFunc = FDMMaterialFunctionLibrary::Get().MakeExpression(
 			InBuildState->GetDynamicMaterial(),
 			"MF_DM_TextureUV_Mirror_X",
 			MaterialFunc_Name_TextureUV_Mirror_X,
@@ -356,7 +356,7 @@ TArray<UMaterialExpression*> UDMMaterialStageInputTextureUV::CreateTextureUVExpr
 	// GetMirrorOnY() == true
 	else if (InTextureUV->GetMirrorOnX() == false)
 	{
-		TextureUVFunc = FDynamicMaterialEditorModule::Get().GetFunctionLibrary()->MakeExpression(
+		TextureUVFunc = FDMMaterialFunctionLibrary::Get().MakeExpression(
 			InBuildState->GetDynamicMaterial(),
 			"MF_DM_TextureUV_Mirror_Y",
 			MaterialFunc_Name_TextureUV_Mirror_Y,
@@ -366,7 +366,7 @@ TArray<UMaterialExpression*> UDMMaterialStageInputTextureUV::CreateTextureUVExpr
 	// GetMirrorOnX() == true && // GetMirrorOnY() == true
 	else
 	{
-		TextureUVFunc = FDynamicMaterialEditorModule::Get().GetFunctionLibrary()->MakeExpression(
+		TextureUVFunc = FDMMaterialFunctionLibrary::Get().MakeExpression(
 			InBuildState->GetDynamicMaterial(),
 			"MF_DM_TextureUV_Mirror_XY",
 			MaterialFunc_Name_TextureUV_Mirror_XY,
diff --git a/Engine/Plugins/Experimental/DynamicMaterial/Source/DynamicMaterialEditor/Private/DMMaterialFunctionLibrary.cpp b/Engine/Plugins/Experimental/DynamicMaterial/Source/DynamicMaterialEditor/Private/DMMaterialFunctionLibrary.cpp
--- a/Engine/Plugins/Experimental/DynamicMaterial/Source/DynamicMaterialEditor/Private/DMMaterialFunctionLibrary.cpp
+++ b/Engine/Plugins/Experimental/DynamicMaterial/Source/DynamicMaterialEditor/Private/DMMaterialFunctionLibrary.cpp
@@ -59,7 +59,8 @@ UMaterialExpressionMaterialFunctionCall* FDMMaterialFunctionLibrary::MakeExpress
 
 FDMMaterialFunctionLibrary& FDMMaterialFunctionLibrary::Get()
 {
-	return *(FDynamicMaterialEditorModule::Get().GetFunctionLibrary().Get());
+	static FDMMaterialFunctionLibrary FunctionLibrary;
+	return FunctionLibrary;
 }
 
 UMaterialExpressionMaterialFunctionCall* FDMMaterialFunctionLibrary::GetBreakOutFloat2Components(UMaterial* Parent, const FString& InComment)
diff --git a/Engine/Plugins/Experimental/DynamicMaterial/Source/DynamicMaterialEditor/Private/DynamicMaterialEditorModule.cpp b/Engine/Plugins/Experimental/DynamicMaterial/Source/DynamicMaterialEditor/Private/DynamicMaterialEditorModule.cpp
--- a/Engine/Plugins/Experimental/DynamicMaterial/Source/DynamicMaterialEditor/Private/DynamicMaterialEditorModule.cpp
+++ b/Engine/Plugins/Experimental/DynamicMaterial/Source/DynamicMaterialEditor/Private/DynamicMaterialEditorModule.cpp
@@ -211,8 +211,6 @@ void FDynamicMaterialEditorModule::StartupModule()
 	FDynamicMaterialEditorCommands::Register();
 	MapCommands();
 
-	FunctionLibrary = MakeShared<FDMMaterialFunctionLibrary>();
-
 	FPropertyEditorModule& PropertyModule = FModuleManager::GetModuleChecked<FPropertyEditorModule>("PropertyEditor");
 
 	PropertyModule.RegisterCustomPropertyTypeLayout(UDynamicMaterialModelEditorOnlyData::StaticClass()->GetFName()
@@ -263,8 +261,6 @@ void FDynamicMaterialEditorModule::ShutdownModule()
 	FDynamicMaterialEditorCommands::Unregister();
 	FDynamicMaterialEditorStyle::Shutdown();
 
-	FunctionLibrary.Reset();
-
 	if (UObjectInitialized() && !IsEngineExitRequested() && FModuleManager::Get().IsModuleLoaded("PropertyEditor"))
 	{
 		FPropertyEditorModule& PropertyModule = FModuleManager::GetModuleChecked<FPropertyEditorModule>("PropertyEditor");
@@ -281,11 +277,6 @@ void FDynamicMaterialEditorModule::ShutdownModule()
 	UDMMaterialValueTexture::GetDefaultRGBTexture.Unbind();
 }
 
-const TSharedPtr<FDMMaterialFunctionLibrary>& FDynamicMaterialEditorModule::GetFunctionLibrary()
-{
-	return FunctionLibrary;
-}
-
 void FDynamicMaterialEditorModule::SetDynamicMaterialModel(UDynamicMaterialModel* InMaterialModel, UWorld* InWorld, bool bInInvokeTab)
 {
 	using namespace UE::DynamicMaterialEditor::Private;
diff --git a/Engine/Plugins/Experimental/DynamicMaterial/Source/DynamicMaterialEditor/Private/DynamicMaterialEditorModule.h b/Engine/Plugins/Experimental/DynamicMaterial/Source/DynamicMaterialEditor/Private/DynamicMaterialEditorModule.h
--- a/Engine/Plugins/Experimental/DynamicMaterial/Source/DynamicMaterialEditor/Private/DynamicMaterialEditorModule.h
+++ b/Engine/Plugins/Experimental/DynamicMaterial/Source/DynamicMaterialEditor/Private/DynamicMaterialEditorModule.h
@@ -97,8 +97,6 @@ public:
 	virtual void ShutdownModule() override;
 	//~ End IModuleInterface
 
-	const TSharedPtr<FDMMaterialFunctionLibrary>& GetFunctionLibrary();
-
 	void SetDynamicMaterialModel(UDynamicMaterialModel* InMaterialModel, UWorld* InWorld, bool bInInvokeTab);
 	void SetDynamicMaterialObjectProperty(const FDMObjectMaterialProperty& InObjectProperty, UWorld* InWorld, bool bInInvokeTab);
 	void SetDynamicMaterialInstance(UDynamicMaterialInstance* InInstance, UWorld* InWorld, bool bInInvokeTab);
@@ -121,7 +119,6 @@ protected:
 	static TMap<UClass*, FDMGetObjectMaterialPropertiesDelegate> CustomMaterialPropertyGenerators;
 	static FDMOnUIValueUpdate OnUIValueUpdate;
 
-	TSharedPtr<FDMMaterialFunctionLibrary> FunctionLibrary;
 	TSet<FDMBuildRequestEntry> BuildRequestList;
 	TSharedRef<FUICommandList> CommandList;
 
