Material Designer
- Added text-based stage source

diff --git a/Engine/Plugins/Experimental/DynamicMaterial/Source/DynamicMaterialEditor/DynamicMaterialEditor.Build.cs b/Engine/Plugins/Experimental/DynamicMaterial/Source/DynamicMaterialEditor/DynamicMaterialEditor.Build.cs
--- a/Engine/Plugins/Experimental/DynamicMaterial/Source/DynamicMaterialEditor/DynamicMaterialEditor.Build.cs
+++ b/Engine/Plugins/Experimental/DynamicMaterial/Source/DynamicMaterialEditor/DynamicMaterialEditor.Build.cs
@@ -13,7 +13,8 @@ public class DynamicMaterialEditor : ModuleRules
 			{
 				"Core",
 				"CoreUObject",
-				"DynamicMaterial"
+				"DynamicMaterial",
+				"Engine"
 			}
 		);
 
@@ -29,7 +30,6 @@ public class DynamicMaterialEditor : ModuleRules
 				"DeveloperSettings",
 				"DynamicMaterialShadersEditor",
 				"EditorWidgets",
-				"Engine",
 				"InputCore",
 				"Json",
 				"JsonUtilities",
diff --git a/Engine/Plugins/Experimental/DynamicMaterial/Source/DynamicMaterialEditor/Private/Components/DMMaterialStageSource.cpp b/Engine/Plugins/Experimental/DynamicMaterial/Source/DynamicMaterialEditor/Private/Components/DMMaterialStageSource.cpp
--- a/Engine/Plugins/Experimental/DynamicMaterial/Source/DynamicMaterialEditor/Private/Components/DMMaterialStageSource.cpp
+++ b/Engine/Plugins/Experimental/DynamicMaterial/Source/DynamicMaterialEditor/Private/Components/DMMaterialStageSource.cpp
@@ -337,12 +337,6 @@ bool UDMMaterialStageSource::UpdateStagePreviewMaterial(UDMMaterialStage* InStag
 
 void UDMMaterialStageSource::NotifyPostChange(const FPropertyChangedEvent& InPropertyChangedEvent, class FEditPropertyChain* InPropertyThatChanged)
 {
-	if (!IsComponentValid())
-	{
-		return;
-	}
-
-	Update(EDMUpdateType::Structure);
 }
 
 void UDMMaterialStageSource::PostEditUndo()
diff --git a/Engine/Plugins/Experimental/DynamicMaterial/Source/DynamicMaterialEditor/Private/Components/MaterialStageExpressions/DMMSEText.cpp b/Engine/Plugins/Experimental/DynamicMaterial/Source/DynamicMaterialEditor/Private/Components/MaterialStageExpressions/DMMSEText.cpp
new file mode 100644
--- /dev/null
+++ b/Engine/Plugins/Experimental/DynamicMaterial/Source/DynamicMaterialEditor/Private/Components/MaterialStageExpressions/DMMSEText.cpp
@@ -0,0 +1,822 @@
+// Copyright Epic Games, Inc. All Rights Reserved.
+
+#include "Components/MaterialStageExpressions/DMMSEText.h"
+
+#include "CanvasItem.h"
+#include "CanvasTypes.h"
+#include "Components/DMMaterialStage.h"
+#include "Components/MaterialStageInputs/DMMSIValue.h"
+#include "Components/MaterialValues/DMMaterialValueTexture.h"
+#include "DynamicMaterialEditorModule.h"
+#include "Engine/Canvas.h"
+#include "Engine/Engine.h"
+#include "Engine/TextureRenderTarget2D.h"
+#include "Materials/MaterialExpressionTextureSample.h"
+#include "TextureResource.h"
+
+#define LOCTEXT_NAMESPACE "DMMaterialStageExpressionText"
+
+namespace UE::DynamicMaterialEditor::Private
+{
+	const FIntPoint MinimumTextTextureSize = FIntPoint(5, 10);
+	const FIntPoint MaxTexScale = FIntPoint(64, 64);
+}
+
+struct FDMMaterialStageExpressionText
+{
+	static const inline FName FontName = GET_MEMBER_NAME_CHECKED(UDMMaterialStageExpressionText, Font);
+	static const inline FName TextName = GET_MEMBER_NAME_CHECKED(UDMMaterialStageExpressionText, Text);
+	static const inline FName ColorName = GET_MEMBER_NAME_CHECKED(UDMMaterialStageExpressionText, Color);
+	static const inline FName BackgroundColorName = GET_MEMBER_NAME_CHECKED(UDMMaterialStageExpressionText, BackgroundColor);
+	static const inline FName JustifyName = GET_MEMBER_NAME_CHECKED(UDMMaterialStageExpressionText, Justify);
+	static const inline FName KerningName = GET_MEMBER_NAME_CHECKED(UDMMaterialStageExpressionText, Kerning);
+	static const inline FName LineHeightName = GET_MEMBER_NAME_CHECKED(UDMMaterialStageExpressionText, LineHeight);
+	static const inline FName PaddingLeftName = GET_MEMBER_NAME_CHECKED(UDMMaterialStageExpressionText, PaddingLeft);
+	static const inline FName PaddingRightName = GET_MEMBER_NAME_CHECKED(UDMMaterialStageExpressionText, PaddingRight);
+	static const inline FName PaddingTopName = GET_MEMBER_NAME_CHECKED(UDMMaterialStageExpressionText, PaddingTop);
+	static const inline FName PaddingBottomName = GET_MEMBER_NAME_CHECKED(UDMMaterialStageExpressionText, PaddingBottom);
+	static const inline FName TextScaleName = GET_MEMBER_NAME_CHECKED(UDMMaterialStageExpressionText, TextScale);
+	static const inline FName TextureSizeOverrideName = GET_MEMBER_NAME_CHECKED(UDMMaterialStageExpressionText, TextureSizeOverride);
+	static const inline FName bOutlineName = GET_MEMBER_NAME_CHECKED(UDMMaterialStageExpressionText, bOutline);
+	static const inline FName OutlineColorName = GET_MEMBER_NAME_CHECKED(UDMMaterialStageExpressionText, OutlineColor);
+	static const inline FName bShadowName = GET_MEMBER_NAME_CHECKED(UDMMaterialStageExpressionText, bShadow);
+	static const inline FName ShadowColorName = GET_MEMBER_NAME_CHECKED(UDMMaterialStageExpressionText, ShadowColor);
+	static const inline FName ShadowOffsetName = GET_MEMBER_NAME_CHECKED(UDMMaterialStageExpressionText, ShadowOffset);
+	static const inline FName bGlowName = GET_MEMBER_NAME_CHECKED(UDMMaterialStageExpressionText, bGlow);
+	static const inline FName GlowColorName = GET_MEMBER_NAME_CHECKED(UDMMaterialStageExpressionText, GlowColor);
+	static const inline FName GlowInnerRadiusName = GET_MEMBER_NAME_CHECKED(UDMMaterialStageExpressionText, GlowInnerRadius);
+	static const inline FName GlowOuterRadiusName = GET_MEMBER_NAME_CHECKED(UDMMaterialStageExpressionText, GlowOuterRadius);
+
+	static const inline TSet<FName> PropertyNames = {
+		FontName, TextName, ColorName, BackgroundColorName, JustifyName, KerningName, LineHeightName, PaddingLeftName, PaddingRightName,
+		PaddingTopName, PaddingBottomName, TextScaleName, TextureSizeOverrideName, bOutlineName, OutlineColorName, bShadowName,
+		ShadowColorName, ShadowOffsetName, bGlowName, GlowColorName, GlowInnerRadiusName, GlowOuterRadiusName
+	};
+};
+
+UDMMaterialStageExpressionText::UDMMaterialStageExpressionText()
+	: UDMMaterialStageExpressionTextureSampleBase(
+		LOCTEXT("Text", "Text"),
+		UMaterialExpressionTextureSample::StaticClass()
+	)
+{
+	Menus.Add(EDMExpressionMenu::Other);
+
+	EditableProperties.Add(FDMMaterialStageExpressionText::FontName);
+	EditableProperties.Add(FDMMaterialStageExpressionText::TextName);
+	EditableProperties.Add(FDMMaterialStageExpressionText::ColorName);
+	EditableProperties.Add(FDMMaterialStageExpressionText::BackgroundColorName);
+	EditableProperties.Add(FDMMaterialStageExpressionText::JustifyName);
+	EditableProperties.Add(FDMMaterialStageExpressionText::KerningName);
+	EditableProperties.Add(FDMMaterialStageExpressionText::LineHeightName);
+	EditableProperties.Add(FDMMaterialStageExpressionText::PaddingLeftName);
+	EditableProperties.Add(FDMMaterialStageExpressionText::PaddingRightName);
+	EditableProperties.Add(FDMMaterialStageExpressionText::PaddingTopName);
+	EditableProperties.Add(FDMMaterialStageExpressionText::PaddingBottomName);
+	EditableProperties.Add(FDMMaterialStageExpressionText::TextScaleName);
+	EditableProperties.Add(FDMMaterialStageExpressionText::TextureSizeOverrideName);
+	EditableProperties.Add(FDMMaterialStageExpressionText::bOutlineName);
+	EditableProperties.Add(FDMMaterialStageExpressionText::OutlineColorName);
+	EditableProperties.Add(FDMMaterialStageExpressionText::bShadowName);
+	EditableProperties.Add(FDMMaterialStageExpressionText::ShadowColorName);
+	EditableProperties.Add(FDMMaterialStageExpressionText::ShadowOffsetName);
+	EditableProperties.Add(FDMMaterialStageExpressionText::bGlowName);
+	EditableProperties.Add(FDMMaterialStageExpressionText::GlowColorName);
+	EditableProperties.Add(FDMMaterialStageExpressionText::GlowInnerRadiusName);
+	EditableProperties.Add(FDMMaterialStageExpressionText::GlowOuterRadiusName);
+}
+
+bool UDMMaterialStageExpressionText::IsInputVisible(int32 InInputIndex) const
+{
+	return InInputIndex != 0;
+}
+
+void UDMMaterialStageExpressionText::AddDefaultInput(int32 InInputIndex) const
+{
+	if (InInputIndex == 0)
+	{
+		UDMMaterialStage* Stage = GetStage();
+		check(Stage);
+
+		UDMMaterialStageInputValue* NewInput = UDMMaterialStageInputValue::ChangeStageInput_NewLocalValue(
+			Stage,
+			InInputIndex,
+			FDMMaterialStageConnectorChannel::WHOLE_CHANNEL,
+			EDMValueType::VT_Texture,
+			FDMMaterialStageConnectorChannel::WHOLE_CHANNEL
+		);
+		check(NewInput);
+
+		const_cast<UDMMaterialStageExpressionText*>(this)->TextureValue = Cast<UDMMaterialValueTexture>(NewInput->GetValue());
+		check(TextureValue);
+
+		FIntPoint RequiredSize = GetRequiredTextureSize();
+		RequiredSize.X = FMath::Max(5, RequiredSize.X);
+		RequiredSize.Y = FMath::Max(10, RequiredSize.Y);
+		const_cast<UDMMaterialStageExpressionText*>(this)->RenderTarget = CreateRenderTarget(TextureValue, RequiredSize, FLinearColor::Black);
+
+		TextureValue->SetValue(RenderTarget);
+		return;
+	}
+
+	Super::AddDefaultInput(InInputIndex);
+}
+
+void UDMMaterialStageExpressionText::PostLoad()
+{
+	Super::PostLoad();
+
+	RequestUpdateTextTexture(/* Needs Resize */ true);
+}
+
+UFont* UDMMaterialStageExpressionText::GetFont() const
+{
+	return Font;
+}
+
+void UDMMaterialStageExpressionText::SetFont(UFont* InFont)
+{
+	if (Font == InFont)
+	{
+		return;
+	}
+
+	Font = InFont;
+
+	RequestUpdateTextTexture(/* Needs Resize */ true);
+}
+
+const FText& UDMMaterialStageExpressionText::GetText() const
+{
+	return Text;
+}
+
+void UDMMaterialStageExpressionText::SetText(const FText& InText)
+{
+	if (Text.EqualTo(InText, ETextComparisonLevel::Default))
+	{
+		return;
+	}
+
+	Text = InText;
+
+	UpdateTextLines();
+	RequestUpdateTextTexture(/* Needs Resize */ true);
+}
+
+const FLinearColor& UDMMaterialStageExpressionText::GetColor() const
+{
+	return Color;
+}
+
+void UDMMaterialStageExpressionText::SetColor(const FLinearColor& InColor)
+{
+	if (Color == InColor)
+	{
+		return;
+	}
+
+	Color = InColor;
+
+	RequestUpdateTextTexture(/* Needs Resize */ false);
+}
+
+const FLinearColor& UDMMaterialStageExpressionText::GetBackgroundColor() const
+{
+	return BackgroundColor;
+}
+
+void UDMMaterialStageExpressionText::SetBackgroundColor(const FLinearColor& InBackgroundColor)
+{
+	if (BackgroundColor == InBackgroundColor)
+	{
+		return;
+	}
+
+	BackgroundColor = InBackgroundColor;
+
+	RequestUpdateTextTexture(/* Needs Resize */ false);
+}
+
+ETextJustify::Type UDMMaterialStageExpressionText::GetJustify() const
+{
+	return Justify;
+}
+
+void UDMMaterialStageExpressionText::SetJustify(ETextJustify::Type InJustify)
+{
+	if (Justify == InJustify)
+	{
+		return;
+	}
+
+	Justify = InJustify;
+
+	RequestUpdateTextTexture(/* Needs Resize */ false);
+}
+
+float UDMMaterialStageExpressionText::GetKerning() const
+{
+	return Kerning;
+}
+
+void UDMMaterialStageExpressionText::SetKerning(float InKerning)
+{
+	if (Kerning == InKerning)
+	{
+		return;
+	}
+
+	Kerning = InKerning;
+
+	RequestUpdateTextTexture(/* Needs Resize */ true);
+}
+
+const FLinearColor& UDMMaterialStageExpressionText::GetOutlineColor() const
+{
+	return OutlineColor;
+}
+
+void UDMMaterialStageExpressionText::SetOutlineColor(const FLinearColor& InOutlineColor)
+{
+	if (OutlineColor == InOutlineColor)
+	{
+		return;
+	}
+
+	OutlineColor = InOutlineColor;
+
+	RequestUpdateTextTexture(/* Needs Resize */ false);
+}
+
+float UDMMaterialStageExpressionText::GetLineHeight() const
+{
+	return LineHeight;
+}
+
+void UDMMaterialStageExpressionText::SetLineHeight(float InLineHeight)
+{
+	if (LineHeight == InLineHeight)
+	{
+		return;
+	}
+
+	LineHeight = InLineHeight;
+
+	RequestUpdateTextTexture(/* Needs Resize */ true);
+}
+
+float UDMMaterialStageExpressionText::GetPaddingLeft() const
+{
+	return PaddingLeft;
+}
+
+void UDMMaterialStageExpressionText::SetPaddingLeft(float InPaddingLeft)
+{
+	if (PaddingLeft == InPaddingLeft)
+	{
+		return;
+	}
+
+	PaddingLeft = InPaddingLeft;
+
+	RequestUpdateTextTexture(/* Needs Resize */ true);
+}
+
+float UDMMaterialStageExpressionText::GetPaddingRight() const
+{
+	return PaddingRight;
+}
+
+void UDMMaterialStageExpressionText::SetPaddingRight(float InPaddingRight)
+{
+	if (PaddingRight == InPaddingRight)
+	{
+		return;
+	}
+
+	PaddingRight = InPaddingRight;
+
+	RequestUpdateTextTexture(/* Needs Resize */ true);
+}
+
+float UDMMaterialStageExpressionText::GetPaddingTop() const
+{
+	return PaddingTop;
+}
+
+void UDMMaterialStageExpressionText::SetPaddingTop(float InPaddingTop)
+{
+	if (PaddingTop == InPaddingTop)
+	{
+		return;
+	}
+
+	PaddingTop = InPaddingTop;
+
+	RequestUpdateTextTexture(/* Needs Resize */ true);
+}
+
+float UDMMaterialStageExpressionText::GetPaddingBottom() const
+{
+	return PaddingBottom;
+}
+
+void UDMMaterialStageExpressionText::SetPaddingBottom(float InPaddingBottom)
+{
+	if (PaddingBottom == InPaddingBottom)
+	{
+		return;
+	}
+
+	PaddingBottom = InPaddingBottom;
+
+	RequestUpdateTextTexture(/* Needs Resize */ true);
+}
+
+const FIntPoint& UDMMaterialStageExpressionText::GetTextScale() const
+{
+	return TextScale;
+}
+
+void UDMMaterialStageExpressionText::SetTextScale(const FIntPoint& InTextScale)
+{
+	FIntPoint NewTextScale = InTextScale;
+	NewTextScale.X = FMath::Clamp(InTextScale.X, 1, UE::DynamicMaterialEditor::Private::MaxTexScale.X);
+	NewTextScale.Y = FMath::Clamp(InTextScale.Y, 1, UE::DynamicMaterialEditor::Private::MaxTexScale.Y);
+
+	if (TextScale == NewTextScale)
+	{
+		return;
+	}
+
+	TextScale = NewTextScale;
+
+	RequestUpdateTextTexture(/* Needs Resize */ true);
+}
+
+const FIntPoint& UDMMaterialStageExpressionText::GetTextureSizeOverride() const
+{
+	return TextureSizeOverride;
+}
+
+void UDMMaterialStageExpressionText::SetTextureSizeOverride(const FIntPoint& InTextureSizeOverride)
+{
+	if (TextureSizeOverride == InTextureSizeOverride)
+	{
+		return;
+	}
+
+	TextureSizeOverride = InTextureSizeOverride;
+
+	RequestUpdateTextTexture(/* Needs Resize */ true);
+}
+
+bool UDMMaterialStageExpressionText::GetHasOutline() const
+{
+	return bOutline;
+}
+
+void UDMMaterialStageExpressionText::SetHasOutline(bool bInHasOutline)
+{
+	if (bOutline == bInHasOutline)
+	{
+		return;
+	}
+
+	bOutline = bInHasOutline;
+
+	RequestUpdateTextTexture(/* Needs Resize */ false);
+}
+
+bool UDMMaterialStageExpressionText::GetHasShadow() const
+{
+	return bShadow;
+}
+
+void UDMMaterialStageExpressionText::SetHasShadow(bool bInHasShadow)
+{
+	if (bShadow == bInHasShadow)
+	{
+		return;
+	}
+
+	bShadow = bInHasShadow;
+
+	RequestUpdateTextTexture(/* Needs Resize */ false);
+}
+
+const FLinearColor& UDMMaterialStageExpressionText::GetShadowColor() const
+{
+	return ShadowColor;
+}
+
+void UDMMaterialStageExpressionText::SetShadowColor(const FLinearColor& InShadowColor)
+{
+	if (ShadowColor == InShadowColor)
+	{
+		return;
+	}
+
+	ShadowColor = InShadowColor;
+
+	RequestUpdateTextTexture(/* Needs Resize */ false);
+}
+
+const FVector2D& UDMMaterialStageExpressionText::GetShadowOffset() const
+{
+	return ShadowOffset;
+}
+
+void UDMMaterialStageExpressionText::SetShadowOffset(const FVector2D& InShadowOffset)
+{
+	if (ShadowOffset == InShadowOffset)
+	{
+		return;
+	}
+
+	ShadowOffset = InShadowOffset;
+
+	RequestUpdateTextTexture(/* Needs Resize */ false);
+}
+
+bool UDMMaterialStageExpressionText::GetHasGlow() const
+{
+	return bGlow;
+}
+
+void UDMMaterialStageExpressionText::SetHasGlow(bool bInHasGlow)
+{
+	if (bGlow == bInHasGlow)
+	{
+		return;
+	}
+
+	bGlow = bInHasGlow;
+
+	RequestUpdateTextTexture(/* Needs Resize */ false);
+}
+
+const FLinearColor& UDMMaterialStageExpressionText::GetGlowColor() const
+{
+	return GlowColor;
+}
+
+void UDMMaterialStageExpressionText::SetGlowColor(const FLinearColor& InGlowColor)
+{
+	if (GlowColor == InGlowColor)
+	{
+		return;
+	}
+
+	GlowColor = InGlowColor;
+
+	RequestUpdateTextTexture(/* Needs Resize */ false);
+}
+
+const FVector2D& UDMMaterialStageExpressionText::GetGlowInnerRadius() const
+{
+	return GlowInnerRadius;
+}
+
+void UDMMaterialStageExpressionText::SetGlowInnerRadius(const FVector2D& InGlowInnerRadius)
+{
+	if (GlowInnerRadius == InGlowInnerRadius)
+	{
+		return;
+	}
+
+	GlowInnerRadius = InGlowInnerRadius;
+
+	RequestUpdateTextTexture(/* Needs Resize */ false);
+}
+
+const FVector2D& UDMMaterialStageExpressionText::GetGlowOuterRadius() const
+{
+	return GlowOuterRadius;
+}
+
+void UDMMaterialStageExpressionText::SetGlowOuterRadius(const FVector2D& InGlowOuterRadius)
+{
+	if (GlowOuterRadius == InGlowOuterRadius)
+	{
+		return;
+	}
+
+	GlowOuterRadius = InGlowOuterRadius;
+
+	RequestUpdateTextTexture(/* Needs Resize */ false);
+}
+
+void UDMMaterialStageExpressionText::PostEditChangeProperty(FPropertyChangedEvent& InPropertyChangedEvent)
+{
+	Super::PostEditChangeProperty(InPropertyChangedEvent);
+
+	const FName PropertyName = InPropertyChangedEvent.GetMemberPropertyName();
+
+	if (PropertyName == FDMMaterialStageExpressionText::TextName)
+	{
+		UpdateTextLines();
+		RequestUpdateTextTexture(/* Needs Resize */ true);
+	}
+	else if (PropertyName == FDMMaterialStageExpressionText::FontName
+		|| PropertyName == FDMMaterialStageExpressionText::KerningName
+		|| PropertyName == FDMMaterialStageExpressionText::LineHeightName
+		|| PropertyName == FDMMaterialStageExpressionText::PaddingLeftName
+		|| PropertyName == FDMMaterialStageExpressionText::PaddingRightName
+		|| PropertyName == FDMMaterialStageExpressionText::PaddingTopName
+		|| PropertyName == FDMMaterialStageExpressionText::PaddingBottomName
+		|| PropertyName == FDMMaterialStageExpressionText::TextScaleName
+		|| PropertyName == FDMMaterialStageExpressionText::TextureSizeOverrideName)
+	{
+		TextScale.X = FMath::Clamp(TextScale.X, 1, UE::DynamicMaterialEditor::Private::MaxTexScale.X);
+		TextScale.Y = FMath::Clamp(TextScale.Y, 1, UE::DynamicMaterialEditor::Private::MaxTexScale.Y);
+		RequestUpdateTextTexture(/* Needs Resize */ true);
+	}
+	else if (PropertyName == FDMMaterialStageExpressionText::bGlowName
+		|| PropertyName == FDMMaterialStageExpressionText::bShadowName
+		|| PropertyName == FDMMaterialStageExpressionText::bOutlineName)
+	{
+		// Cause details panel refresh.
+		Update(EDMUpdateType::Structure);
+		RequestUpdateTextTexture(/* Needs Resize */ false);
+	}
+	else if (FDMMaterialStageExpressionText::PropertyNames.Contains(PropertyName))
+	{
+		RequestUpdateTextTexture(/* Needs Resize */ false);
+	}
+}
+
+UTextureRenderTarget2D* UDMMaterialStageExpressionText::CreateRenderTarget(UObject* InOuter, const FIntPoint& InSize, 
+	const FLinearColor& InClearColor)
+{
+	UTextureRenderTarget2D* RenderTarget = NewObject<UTextureRenderTarget2D>(InOuter, NAME_None, 
+		EObjectFlags::RF_Transactional | EObjectFlags::RF_DuplicateTransient | EObjectFlags::RF_TextExportTransient
+	);
+
+	check(RenderTarget);
+	RenderTarget->RenderTargetFormat = RTF_RGBA16f;
+	RenderTarget->ClearColor = InClearColor;
+	RenderTarget->bAutoGenerateMips = false;
+	RenderTarget->bCanCreateUAV = false;
+	RenderTarget->InitAutoFormat(InSize.X, InSize.Y);
+	RenderTarget->UpdateResourceImmediate(true);
+
+	return RenderTarget;
+}
+
+void UDMMaterialStageExpressionText::CalculateLineSizes()
+{
+	FTextSizingParameters Params(0, 0, 0, 0, Font);
+	Params.Scaling = FVector2D(1, 1);
+
+	for (FDMTextLine& Line : Lines)
+	{
+		Params.DrawXL = 0;
+		Params.DrawYL = 0;
+
+		UCanvas::CanvasStringSize(Params, *Line.Line);
+
+		int32 RequiredWidth = Params.DrawXL;
+		RequiredWidth += Line.Line.Len() > 0 ? ((Line.Line.Len() - 1) * Kerning) : 0;
+
+		Line.Width = RequiredWidth;
+	}
+}
+
+FIntPoint UDMMaterialStageExpressionText::GetRequiredTextureSize() const
+{
+	if (TextureSizeOverride.X > 0 && TextureSizeOverride.Y > 0)
+	{
+		return TextureSizeOverride;
+	}
+
+	if (!IsValid(Font) || Lines.IsEmpty())
+	{
+		return UE::DynamicMaterialEditor::Private::MinimumTextTextureSize;
+	}
+
+	const float MaxCharHeight = Font->GetMaxCharHeight();
+
+	if (MaxCharHeight <= 0)
+	{
+		return UE::DynamicMaterialEditor::Private::MinimumTextTextureSize;
+	}
+
+	FVector2f Size = FVector2f::ZeroVector;
+
+	for (const FDMTextLine& Line : Lines)
+	{
+		Size.X = FMath::Max(Size.X, Line.Width);
+		Size.Y += MaxCharHeight * LineHeight;
+	}
+
+	Size.X += PaddingLeft + PaddingRight;
+	Size.Y += PaddingTop + PaddingBottom;
+
+	Size.X = FMath::Max(Size.X, UE::DynamicMaterialEditor::Private::MinimumTextTextureSize.X);
+	Size.Y = FMath::Max(Size.Y, UE::DynamicMaterialEditor::Private::MinimumTextTextureSize.Y);
+
+	FIntPoint RequiredSize = FIntPoint::ZeroValue;
+	RequiredSize.X = FMath::CeilToInt(Size.X * TextScale.X);
+	RequiredSize.Y = FMath::CeilToInt(Size.Y * TextScale.Y);
+
+	return RequiredSize;
+}
+
+void UDMMaterialStageExpressionText::UpdateTextLines()
+{
+	const FString TextStr = Text.ToString();
+	TArray<FString> NewLines;
+	TextStr.ParseIntoArray(NewLines, TEXT("\n"), /* Cull empty */ false);
+
+	Lines.Empty();
+	Lines.Reserve(NewLines.Num());
+
+	for (FString& NewLine : NewLines)
+	{
+		if (NewLine.EndsWith(TEXT("\r")))
+		{
+			NewLine.LeftChopInline(1);
+		}
+
+		Lines.Add({NewLine, -1.f});
+	}
+}
+
+void UDMMaterialStageExpressionText::RequestUpdateTextTexture(bool bInNeedsResize)
+{
+	bNeedsResize = bInNeedsResize;
+
+	if (EndOfFrameDelegateHandle.IsValid())
+	{
+		return;
+	}
+
+	EndOfFrameDelegateHandle = FCoreDelegates::OnEndFrame.AddUObject(this, &UDMMaterialStageExpressionText::UpdateTextTexture);
+}
+
+void UDMMaterialStageExpressionText::UpdateTextTexture()
+{
+	if (EndOfFrameDelegateHandle.IsValid())
+	{
+		FCoreDelegates::OnEndFrame.Remove(EndOfFrameDelegateHandle);
+		EndOfFrameDelegateHandle.Reset();
+	}
+
+	if (!IsValid(Font) || Lines.IsEmpty())
+	{
+		return;
+	}
+
+	if (bNeedsResize)
+	{
+		CalculateLineSizes();
+	}
+		
+	const bool bValidRenderTarget = IsValid(RenderTarget.Get());
+
+	if (!bValidRenderTarget)
+	{
+		bNeedsResize = true;
+	}
+	else if (bNeedsResize && TextureSizeOverride.X > 0 && TextureSizeOverride.X == RenderTarget->SizeX
+		&& TextureSizeOverride.Y > 0 && TextureSizeOverride.Y == RenderTarget->SizeY)
+	{
+		bNeedsResize = false;
+	}
+
+	if (bNeedsResize)
+	{
+		FIntPoint RequiredSize = GetRequiredTextureSize();
+		RequiredSize.X = FMath::Max(RequiredSize.X, UE::DynamicMaterialEditor::Private::MinimumTextTextureSize.X);
+		RequiredSize.Y = FMath::Max(RequiredSize.Y, UE::DynamicMaterialEditor::Private::MinimumTextTextureSize.Y);
+
+		if (bValidRenderTarget && RequiredSize.X == RenderTarget->SizeX && RequiredSize.X == RenderTarget->SizeY)
+		{
+			bNeedsResize = false;
+		}
+
+		if (bNeedsResize)
+		{
+			RenderTarget = CreateRenderTarget(TextureValue, RequiredSize, FLinearColor::Black);
+			TextureValue->SetValue(RenderTarget);
+			bNeedsResize = false;
+		}
+	}
+
+	RenderTarget->UpdateResourceImmediate(true);
+
+	FTextureRenderTargetResource* RenderTargetResource = RenderTarget->GameThread_GetRenderTargetResource();
+
+	if (!RenderTargetResource)
+	{
+		return;
+	}
+
+	static const FGameTime GameTime = FGameTime::CreateUndilated(0, 0);
+
+	FCanvas RenderCanvas(
+		RenderTargetResource,
+		nullptr,
+		GameTime,
+		GEngine->GetDefaultWorldFeatureLevel(),
+		FCanvas::ECanvasDrawMode::CDM_ImmediateDrawing
+	);
+
+	const FVector2D RenderTargetSize = FVector2D(RenderTarget->SizeX, RenderTarget->SizeY);
+
+	FCanvasTileItem TileItem = FCanvasTileItem(FVector2D::ZeroVector, RenderTargetSize, BackgroundColor);
+	RenderCanvas.DrawItem(TileItem);
+
+	const float MaxCharHeight = Font->GetMaxCharHeight();
+
+	FCanvasTextItem TextItem = CreateTextItem(FVector2D::ZeroVector, FText::GetEmpty());
+
+	const float LeftOffset = PaddingLeft * static_cast<float>(TextScale.X);
+	const float RightOffset = PaddingRight * static_cast<float>(TextScale.X);
+
+	TextItem.Position.X = LeftOffset;
+	TextItem.Position.Y = PaddingTop * static_cast<float>(TextScale.Y);
+	TextItem.Scale.X = static_cast<float>(TextScale.X);
+	TextItem.Scale.Y = static_cast<float>(TextScale.Y);
+
+	if (!FMath::IsNearlyEqual(LineHeight, 1.f))
+	{
+		TextItem.Position.Y += MaxCharHeight * (LineHeight - 1.f) * 0.5f * static_cast<float>(TextScale.Y);
+	}
+
+	const float ActualLineHeight = MaxCharHeight * LineHeight * static_cast<float>(TextScale.Y);
+
+	for (const FDMTextLine& Line : Lines)
+	{
+		switch (Justify)
+		{
+			default:
+			case ETextJustify::Left:
+			case ETextJustify::InvariantLeft:
+				TextItem.Position.X = LeftOffset;
+				break;
+
+			case ETextJustify::Center:
+				TextItem.Position.X = (RenderTarget->SizeX - (Line.Width * TextItem.Scale.X)) * 0.5f;
+				break;
+
+			case ETextJustify::Right:
+			case ETextJustify::InvariantRight:
+				TextItem.Position.X = RenderTarget->SizeX - (Line.Width * TextItem.Scale.X) - RightOffset;
+				break;
+		}
+
+		TextItem.Text = FText::FromString(Line.Line);
+		RenderCanvas.DrawItem(TextItem);
+		TextItem.Position.Y += ActualLineHeight;
+	}
+
+	RenderCanvas.Flush_GameThread();
+
+	RenderTarget->UpdateResourceImmediate(false);
+}
+
+FCanvasTextItem UDMMaterialStageExpressionText::CreateTextItem(const FVector2D& InPosition, const FText& InText) const
+{
+	FCanvasTextItem TextItem(InPosition, InText, Font, Color);
+
+	TextItem.BlendMode = SE_BLEND_Translucent; // Must be translucent for fonts!
+	TextItem.HorizSpacingAdjust = Kerning;
+
+	if (bOutline)
+	{
+		TextItem.bOutlined = true;
+		TextItem.OutlineColor = OutlineColor;
+	}
+	else
+	{
+		TextItem.bOutlined = false;
+	}
+
+	if (bShadow)
+	{
+		TextItem.EnableShadow(ShadowColor, ShadowOffset);
+	}
+	else
+	{
+		TextItem.DisableShadow();
+	}
+
+	if (bGlow)
+	{
+		TextItem.FontRenderInfo.GlowInfo.bEnableGlow = true;
+		TextItem.FontRenderInfo.GlowInfo.GlowColor = GlowColor;
+		TextItem.FontRenderInfo.GlowInfo.GlowInnerRadius = GlowInnerRadius;
+		TextItem.FontRenderInfo.GlowInfo.GlowOuterRadius = GlowOuterRadius;
+	}
+	else
+	{
+		TextItem.FontRenderInfo.GlowInfo.bEnableGlow = false;
+	}
+
+	return TextItem;
+}
+
+#undef LOCTEXT_NAMESPACE
diff --git a/Engine/Plugins/Experimental/DynamicMaterial/Source/DynamicMaterialEditor/Private/Menus/DMMaterialStageSourceMenus.cpp b/Engine/Plugins/Experimental/DynamicMaterial/Source/DynamicMaterialEditor/Private/Menus/DMMaterialStageSourceMenus.cpp
--- a/Engine/Plugins/Experimental/DynamicMaterial/Source/DynamicMaterialEditor/Private/Menus/DMMaterialStageSourceMenus.cpp
+++ b/Engine/Plugins/Experimental/DynamicMaterial/Source/DynamicMaterialEditor/Private/Menus/DMMaterialStageSourceMenus.cpp
@@ -14,6 +14,7 @@
 #include "Components/DMMaterialStageThroughputLayerBlend.h"
 #include "Components/DMMaterialValue.h"
 #include "Components/MaterialStageExpressions/DMMSESceneTexture.h"
+#include "Components/MaterialStageExpressions/DMMSEText.h"
 #include "Components/MaterialStageExpressions/DMMSETextureSample.h"
 #include "Components/MaterialStageExpressions/DMMSETextureSampleEdgeColor.h"
 #include "Components/MaterialStageInputs/DMMSIExpression.h"
@@ -1165,6 +1166,79 @@ namespace UE::DynamicMaterialEditor::Private
 		}
 	}
 
+	void ChangeSourceToTextFromContext(UDMMenuContext* InMenuContext)
+	{
+		if (!IsValid(InMenuContext))
+		{
+			return;
+		}
+
+		UDMMaterialStageSource* const StageSource = InMenuContext->GetStageSource();
+
+		if (!StageSource)
+		{
+			return;
+		}
+
+		UDMMaterialStage* const Stage = InMenuContext->GetStage();
+
+		if (!Stage)
+		{
+			return;
+		}
+
+		UDMMaterialLayerObject* const Layer = Stage->GetLayer();
+
+		if (!Layer)
+		{
+			return;
+		}
+
+		if (StageSource->IsA<UDMMaterialStageBlend>())
+		{
+			FScopedTransaction Transaction(LOCTEXT("SetStageInputBase", "Set Material Designer Base Source"));
+			Stage->Modify();
+
+			UDMMaterialStageInputExpression::ChangeStageInput_Expression(
+				Stage,
+				UDMMaterialStageExpressionText::StaticClass(),
+				UDMMaterialStageBlend::InputB,
+				FDMMaterialStageConnectorChannel::WHOLE_CHANNEL,
+				0,
+				FDMMaterialStageConnectorChannel::THREE_CHANNELS
+			);
+		}
+		else if (StageSource->IsA<UDMMaterialStageThroughputLayerBlend>())
+		{
+			FScopedTransaction Transaction(LOCTEXT("SetStageInputMask", "Set Material Designer Mask Source"));
+			Stage->Modify();
+
+			UDMMaterialStageInputExpression::ChangeStageInput_Expression(
+				Stage,
+				UDMMaterialStageExpressionText::StaticClass(),
+				UDMMaterialStageThroughputLayerBlend::InputMaskSource,
+				FDMMaterialStageConnectorChannel::WHOLE_CHANNEL,
+				0,
+				FDMMaterialStageConnectorChannel::FIRST_CHANNEL
+			);
+		}
+		else
+		{
+			FScopedTransaction Transaction(LOCTEXT("SetStageInput", "Set Material Designer Source"));
+			Stage->Modify();
+
+			UDMMaterialStageInputExpression::ChangeStageSource_Expression(
+				Stage,
+				UDMMaterialStageExpressionText::StaticClass()
+			);
+		}
+
+		if (TSharedPtr<SDMSlot> SlotWidget = InMenuContext->GetSlotWidget().Pin())
+		{
+			SlotWidget->InvalidateComponentEditWidget();
+		}
+	}
+
 	void CreateChangeMaterialStageSource(FToolMenuSection& InSection)
 	{
 		UDMMenuContext* MenuContext = InSection.FindContext<UDMMenuContext>();
@@ -1326,6 +1400,18 @@ namespace UE::DynamicMaterialEditor::Private
 			)
 		);
 
+		InSection.AddMenuEntry("Text",
+			LOCTEXT("ChangeSourceText", "Text"),
+			LOCTEXT("ChangeSourceTextTooltip", "Change the source of this stage to a Text Renderer."),
+			FSlateIcon(),
+			FUIAction(
+				FExecuteAction::CreateStatic(
+					&ChangeSourceToTextFromContext,
+					MenuContext
+				)
+			)
+		);
+
 		if constexpr (UE::DynamicMaterialEditor::bAdvancedSlotsEnabled)
 		{
 			if (bHasValidSlot)
diff --git a/Engine/Plugins/Experimental/DynamicMaterial/Source/DynamicMaterialEditor/Private/Slate/SDMComponentEdit.cpp b/Engine/Plugins/Experimental/DynamicMaterial/Source/DynamicMaterialEditor/Private/Slate/SDMComponentEdit.cpp
--- a/Engine/Plugins/Experimental/DynamicMaterial/Source/DynamicMaterialEditor/Private/Slate/SDMComponentEdit.cpp
+++ b/Engine/Plugins/Experimental/DynamicMaterial/Source/DynamicMaterialEditor/Private/Slate/SDMComponentEdit.cpp
@@ -308,6 +308,11 @@ TSharedRef<SWidget> SDMComponentEdit::CreateEditWidget()
 
 	for (const FDMPropertyHandle& EditRow : EditRows)
 	{
+		if (!EditRow.DetailTreeNode)
+		{
+			continue;
+		}
+
 		ECustomDetailsTreeInsertPosition Position = ECustomDetailsTreeInsertPosition::Child;
 
 		if (EditRow.DetailTreeNode->CreatePropertyHandle()->HasMetaData("HighPriority"))
diff --git a/Engine/Plugins/Experimental/DynamicMaterial/Source/DynamicMaterialEditor/Private/Slate/SDMEditor.cpp b/Engine/Plugins/Experimental/DynamicMaterial/Source/DynamicMaterialEditor/Private/Slate/SDMEditor.cpp
--- a/Engine/Plugins/Experimental/DynamicMaterial/Source/DynamicMaterialEditor/Private/Slate/SDMEditor.cpp
+++ b/Engine/Plugins/Experimental/DynamicMaterial/Source/DynamicMaterialEditor/Private/Slate/SDMEditor.cpp
@@ -56,6 +56,64 @@
 
 #define LOCTEXT_NAMESPACE "SDMEditor"
 
+namespace UE::DynamicMaterialEditor::Private
+{
+	TSharedPtr<IDetailTreeNode> SearchGeneratorForNode(const TSharedRef<IPropertyRowGenerator>& InGenerator, FName InPropertyName)
+	{
+		for (const TSharedRef<IDetailTreeNode>& CategoryNode : InGenerator->GetRootTreeNodes())
+		{
+			if (CategoryNode->GetNodeName() != TEXT("Material Designer"))
+			{
+				continue;
+			}
+
+			TArray<TSharedRef<IDetailTreeNode>> ChildNodes;
+			CategoryNode->GetChildren(ChildNodes);
+
+			for (const TSharedRef<IDetailTreeNode>& ChildNode : ChildNodes)
+			{
+				if (ChildNode->GetNodeType() != EDetailNodeType::Item)
+				{
+					continue;
+				}
+
+				if (ChildNode->GetNodeName() != InPropertyName)
+				{
+					continue;
+				}
+
+				return ChildNode;
+			}
+		}
+
+		return nullptr;
+	}
+
+	TSharedPtr<IPropertyRowGenerator> SearchForGenerator(const TArray<FDMPropertyHandle>& InPropertyHandles, UObject* InObject)
+	{
+		if (!InObject)
+		{
+			return nullptr;
+		}
+
+		for (const FDMPropertyHandle& PropertyHandle : InPropertyHandles)
+		{
+			if (PropertyHandle.PropertyRowGenerator.IsValid())
+			{
+				for (const TWeakObjectPtr<UObject>& WeakObject : PropertyHandle.PropertyRowGenerator->GetSelectedObjects())
+				{
+					if (WeakObject.Get() == InObject)
+					{
+						return PropertyHandle.PropertyRowGenerator;
+					}
+				}
+			}
+		}
+
+		return nullptr;
+	}
+}
+
 TSharedPtr<FAssetThumbnailPool> SDMEditor::ThumbnailPool = nullptr;
 TMap<const SWidget*, TArray<FDMPropertyHandle>> SDMEditor::PropertyHandleMap;
 
@@ -158,7 +216,7 @@ FDMPropertyHandle SDMEditor::GetPropertyHandle(const SWidget* InOwner, UDMMateri
 
 	for (const FDMPropertyHandle& ExistingHandle : PropertyHandles)
 	{
-		if (ExistingHandle.PropertyHandle->GetProperty()->GetFName() == InPropertyName)
+		if (ExistingHandle.PropertyHandle && ExistingHandle.PropertyHandle->GetProperty()->GetFName() == InPropertyName)
 		{
 			TArray<UObject*> Outers;
 			ExistingHandle.PropertyHandle->GetOuterObjects(Outers);
@@ -170,6 +228,21 @@ FDMPropertyHandle SDMEditor::GetPropertyHandle(const SWidget* InOwner, UDMMateri
 		}
 	}
 
+	if (TSharedPtr<IPropertyRowGenerator> PropertyRowGenerator = UE::DynamicMaterialEditor::Private::SearchForGenerator(PropertyHandles, InComponent))
+	{
+		FDMPropertyHandle PropertyHandle;
+		PropertyHandle.PropertyRowGenerator = PropertyRowGenerator;
+
+		if (TSharedPtr<IDetailTreeNode> DetailTreeNode = UE::DynamicMaterialEditor::Private::SearchGeneratorForNode(PropertyRowGenerator.ToSharedRef(), InPropertyName))
+		{
+			PropertyHandle.DetailTreeNode = DetailTreeNode;
+			PropertyHandle.PropertyHandle = DetailTreeNode->CreatePropertyHandle();
+			return PropertyHandle;
+		}
+
+		return PropertyHandle;
+	}
+
 	FDMPropertyHandle NewHandle = CreatePropertyHandle(InOwner, InComponent, InPropertyName);
 	PropertyHandles.Add(NewHandle);
 
@@ -194,32 +267,12 @@ FDMPropertyHandle SDMEditor::CreatePropertyHandle(const void* InOwner, UDMMateri
 	PropertyHandle.PropertyRowGenerator = PropertyEditor.CreatePropertyRowGenerator(RowGeneratorArgs);
 	PropertyHandle.PropertyRowGenerator->SetObjects({InComponent});
 
-	for (const TSharedRef<IDetailTreeNode>& CategoryNode : PropertyHandle.PropertyRowGenerator->GetRootTreeNodes())
+	if (const TSharedPtr<IDetailTreeNode> FoundTreeNode = UE::DynamicMaterialEditor::Private::SearchGeneratorForNode(
+		PropertyHandle.PropertyRowGenerator.ToSharedRef(), InPropertyName))
 	{
-		if (CategoryNode->GetNodeName() != TEXT("Material Designer"))
-		{
-			continue;
-		}
-
-		TArray<TSharedRef<IDetailTreeNode>> ChildNodes;
-		CategoryNode->GetChildren(ChildNodes);
-
-		for (const TSharedRef<IDetailTreeNode>& ChildNode : ChildNodes)
-		{
-			if (ChildNode->GetNodeType() != EDetailNodeType::Item)
-			{
-				continue;
-			}
-
-			if (ChildNode->GetNodeName() != InPropertyName)
-			{
-				continue;
-			}
-
-			PropertyHandle.DetailTreeNode = ChildNode;
-			PropertyHandle.PropertyHandle = ChildNode->CreatePropertyHandle();
-			return PropertyHandle;
-		}
+		PropertyHandle.DetailTreeNode = FoundTreeNode;
+		PropertyHandle.PropertyHandle = FoundTreeNode->CreatePropertyHandle();
+		return PropertyHandle;
 	}
 
 	return PropertyHandle;
diff --git a/Engine/Plugins/Experimental/DynamicMaterial/Source/DynamicMaterialEditor/Public/Components/MaterialStageExpressions/DMMSEText.h b/Engine/Plugins/Experimental/DynamicMaterial/Source/DynamicMaterialEditor/Public/Components/MaterialStageExpressions/DMMSEText.h
new file mode 100644
--- /dev/null
+++ b/Engine/Plugins/Experimental/DynamicMaterial/Source/DynamicMaterialEditor/Public/Components/MaterialStageExpressions/DMMSEText.h
@@ -0,0 +1,240 @@
+// Copyright Epic Games, Inc. All Rights Reserved.
+
+#pragma once
+
+#include "Components/MaterialStageExpressions/DMMSETextureSampleBase.h"
+#include "Delegates/IDelegateInstance.h"
+#include "Framework/Text/TextLayout.h"
+#include "Internationalization/Text.h"
+#include "Math/Color.h"
+#include "Math/Vector2D.h"
+#include "DMMSEText.generated.h"
+
+class FCanvasTextItem;
+class UCanvas;
+class UDMMaterialValueTexture;
+class UFont;
+class UTextureRenderTarget2D;
+
+USTRUCT()
+struct FDMTextLine
+{
+	GENERATED_BODY()
+
+	UPROPERTY()
+	FString Line;
+
+	UPROPERTY()
+	float Width;
+};
+
+UCLASS(BlueprintType, ClassGroup = "Material Designer")
+class DYNAMICMATERIALEDITOR_API UDMMaterialStageExpressionText : public UDMMaterialStageExpressionTextureSampleBase
+{
+	GENERATED_BODY()
+
+	friend class SDMComponentEdit;
+	friend struct FDMMaterialStageExpressionText;
+
+public:
+	UDMMaterialStageExpressionText();
+
+	//~ Begin UDMMaterialStageThroughput
+	virtual bool IsInputVisible(int32 InInputIndex) const override;
+	virtual void AddDefaultInput(int32 InInputIndex) const override;
+	//~ End UDMMaterialStageThroughput
+
+	//~ Begin UObject
+	virtual void PostLoad() override;
+	virtual void PostEditChangeProperty(FPropertyChangedEvent& InPropertyChangedEvent) override;
+	//~ End UObject
+
+	UFont* GetFont() const;
+	void SetFont(UFont* InFont);
+
+	const FText& GetText() const;
+	void SetText(const FText& InText);
+
+	const FLinearColor& GetColor() const;
+	void SetColor(const FLinearColor& InColor);
+
+	const FLinearColor& GetBackgroundColor() const;
+	void SetBackgroundColor(const FLinearColor& InBackgroundColor);
+
+	ETextJustify::Type GetJustify() const;
+	void SetJustify(ETextJustify::Type InJustify);
+
+	float GetKerning() const;
+	void SetKerning(float InKerning);
+
+	float GetLineHeight() const;
+	void SetLineHeight(float InLineHeight);
+
+	float GetPaddingLeft() const;
+	void SetPaddingLeft(float InPadding);
+
+	float GetPaddingRight() const;
+	void SetPaddingRight(float InPadding);
+
+	float GetPaddingTop() const;
+	void SetPaddingTop(float InPadding);
+
+	float GetPaddingBottom() const;
+	void SetPaddingBottom(float InPadding);
+
+	const FIntPoint& GetTextScale() const;
+	void SetTextScale(const FIntPoint& InTextScale);
+
+	const FIntPoint& GetTextureSizeOverride() const;
+	void SetTextureSizeOverride(const FIntPoint& InTextureSizeOverride);
+
+	bool GetHasOutline() const;
+	void SetHasOutline(bool bInHasOutline);
+
+	const FLinearColor& GetOutlineColor() const;
+	void SetOutlineColor(const FLinearColor& InOutlineColor);
+
+	bool GetHasShadow() const;
+	void SetHasShadow(bool bInHasShadow);
+
+	const FLinearColor& GetShadowColor() const;
+	void SetShadowColor(const FLinearColor& InShadowColor);
+
+	const FVector2D& GetShadowOffset() const;
+	void SetShadowOffset(const FVector2D& InShadowOffset);
+
+	bool GetHasGlow() const;
+	void SetHasGlow(bool bInHasGlow);
+
+	const FLinearColor& GetGlowColor() const;
+	void SetGlowColor(const FLinearColor& InGlowColor);
+
+	const FVector2D& GetGlowInnerRadius() const;
+	void SetGlowInnerRadius(const FVector2D& InGlowInnerRadius);
+
+	const FVector2D& GetGlowOuterRadius() const;
+	void SetGlowOuterRadius(const FVector2D& InGlowOuterRadius);
+	
+protected:
+	UPROPERTY(EditInstanceOnly, BlueprintReadOnly, Setter = SetFont, Category = "Material Designer",
+		meta = (NotKeyframeable, AllowPrivateAccess = "true"))
+	UFont* Font = nullptr;
+
+	UPROPERTY(EditInstanceOnly, BlueprintReadOnly, Setter = SetText, Category = "Material Designer",
+		meta = (MultiLine, NotKeyframeable, AllowPrivateAccess = "true"))
+	FText Text = FText::GetEmpty();
+
+	UPROPERTY()
+	TArray<FDMTextLine> Lines = {};
+
+	UPROPERTY(EditInstanceOnly, BlueprintReadOnly, Setter = SetColor, Category = "Material Designer",
+		meta = (AllowPrivateAccess = "true"))
+	FLinearColor Color = FLinearColor::White;
+
+	/** Clear color of the render target. */
+	UPROPERTY(EditInstanceOnly, BlueprintReadOnly, Setter = SetBackgroundColor, Category = "Material Designer",
+		meta = (AllowPrivateAccess = "true"))
+	FLinearColor BackgroundColor = FLinearColor::Black;
+
+	UPROPERTY(EditInstanceOnly, BlueprintReadOnly, Setter = SetJustify, Category = "Material Designer",
+		meta = (NotKeyframeable, AllowPrivateAccess = "true"))
+	TEnumAsByte<ETextJustify::Type> Justify = ETextJustify::Left;
+
+	UPROPERTY(EditInstanceOnly, BlueprintReadOnly, Setter = SetKerning, Category = "Material Designer",
+		meta = (AllowPrivateAccess = "true"))
+	float Kerning = 0.f;
+
+	/** Multiplier on the base font height. */
+	UPROPERTY(EditInstanceOnly, BlueprintReadOnly, Setter = SetLineHeight, Category = "Material Designer",
+		meta = (NotKeyframeable, AllowPrivateAccess = "true", UIMin = 0.1, UIMax = 2.0))
+	float LineHeight = 1.f;
+
+	/** Extra space adding beyond the edge of the glyphs. Useful for shadows, glows, etc. */
+	UPROPERTY(EditInstanceOnly, BlueprintReadOnly, Setter = SetPaddingLeft, Category = "Material Designer",
+		meta = (NotKeyframeable, AllowPrivateAccess = "true"))
+	float PaddingLeft = 0.f;
+
+	/** Extra space adding beyond the edge of the glyphs. Useful for shadows, glows, etc. */
+	UPROPERTY(EditInstanceOnly, BlueprintReadOnly, Setter = SetPaddingRight, Category = "Material Designer",
+		meta = (NotKeyframeable, AllowPrivateAccess = "true"))
+	float PaddingRight = 0.f;
+
+	/** Extra space adding beyond the edge of the glyphs. Useful for shadows, glows, etc. */
+	UPROPERTY(EditInstanceOnly, BlueprintReadOnly, Setter = SetPaddingTop, Category = "Material Designer",
+		meta = (NotKeyframeable, AllowPrivateAccess = "true"))
+	float PaddingTop = 0.f;
+
+	/** Extra space adding beyond the edge of the glyphs. Useful for shadows, glows, etc. */
+	UPROPERTY(EditInstanceOnly, BlueprintReadOnly, Setter = SetPaddingBottom, Category = "Material Designer",
+		meta = (NotKeyframeable, AllowPrivateAccess = "true"))
+	float PaddingBottom = 0.f;
+
+	UPROPERTY(EditInstanceOnly, BlueprintReadOnly, Setter = SetTextScale, Category = "Material Designer",
+		meta = (NotKeyframeable, AllowPrivateAccess = "true", ClampMin = 1, UIMin = 1, ClampMax = 64, UIMax = 64,
+			ToolTip = "Increasing this will increase the resolution of the text."))
+	FIntPoint TextScale = FIntPoint(10, 10);
+
+	UPROPERTY(EditInstanceOnly, BlueprintReadOnly, Setter = SetTextureSizeOverride, Category = "Material Designer",
+		meta = (NotKeyframeable, AllowPrivateAccess = "true", ToolTip = "Override the size of the render target. Set to 0 to disable."))
+	FIntPoint TextureSizeOverride = FIntPoint(0, 0);
+
+	UPROPERTY(EditInstanceOnly, BlueprintReadOnly, Setter = SetHasOutline, Category = "Material Designer",
+		meta = (AllowPrivateAccess = "true"))
+	bool bOutline = false;
+
+	UPROPERTY(EditInstanceOnly, BlueprintReadOnly, Setter = SetOutlineColor, Category = "Material Designer",
+		meta = (AllowPrivateAccess = "true", EditCondition = "bOutline", EditConditionHides = true))
+	FLinearColor OutlineColor = FLinearColor::Black;
+
+	UPROPERTY(EditInstanceOnly, BlueprintReadOnly, Setter = SetHasShadow, Category = "Material Designer",
+		meta = (AllowPrivateAccess = "true"))
+	bool bShadow = false;
+
+	UPROPERTY(EditInstanceOnly, BlueprintReadOnly, Setter = SetShadowColor, Category = "Material Designer",
+		meta = (AllowPrivateAccess = "true", EditCondition = "bShadow", EditConditionHides = true))
+	FLinearColor ShadowColor = FLinearColor::Black;
+
+	UPROPERTY(EditInstanceOnly, BlueprintReadOnly, Setter = SetShadowOffset, Category = "Material Designer",
+		meta = (AllowPrivateAccess = "true", EditCondition="bShadow", EditConditionHides = true))
+	FVector2D ShadowOffset = FVector2D(1.0, 1.0);
+
+	UPROPERTY(EditInstanceOnly, BlueprintReadOnly, Setter = SetHasGlow, Category = "Material Designer",
+		meta = (AllowPrivateAccess = "true"))
+	bool bGlow = false;
+
+	UPROPERTY(EditInstanceOnly, BlueprintReadOnly, Setter = SetGlowColor, Category = "Material Designer",
+		meta = (AllowPrivateAccess = "true", EditCondition = "bGlow", EditConditionHides = true))
+	FLinearColor GlowColor = FLinearColor::White;
+
+	UPROPERTY(EditInstanceOnly, BlueprintReadOnly, Setter = SetGlowInnerRadius, Category = "Material Designer",
+		meta = (AllowPrivateAccess = "true", EditCondition = "bGlow", EditConditionHides = true, ClampMin = 0, UIMin = 0, ClampMax = 1, UIMax = 1))
+	FVector2D GlowInnerRadius = FVector2D(0.5, 0.5);
+
+	UPROPERTY(EditInstanceOnly, BlueprintReadOnly, Setter = SetGlowOuterRadius, Category = "Material Designer",
+		meta = (AllowPrivateAccess = "true", EditCondition = "bGlow", EditConditionHides = true, ClampMin = 0, UIMin = 0, ClampMax = 1, UIMax = 1))
+	FVector2D GlowOuterRadius = FVector2D(0.5, 0.5);
+
+	UPROPERTY(Instanced)
+	TObjectPtr<UDMMaterialValueTexture> TextureValue = nullptr;
+
+	UPROPERTY(Instanced)
+	TObjectPtr<UTextureRenderTarget2D> RenderTarget = nullptr;
+
+	FDelegateHandle EndOfFrameDelegateHandle;
+
+	bool bNeedsResize = false;
+
+	static UTextureRenderTarget2D* CreateRenderTarget(UObject* InOuter, const FIntPoint& InSize, const FLinearColor& InClearColor);
+
+	void CalculateLineSizes();
+
+	FIntPoint GetRequiredTextureSize() const;
+
+	void UpdateTextLines();
+
+	void RequestUpdateTextTexture(bool bInNeedsResize);
+
+	void UpdateTextTexture();
+
+	FCanvasTextItem CreateTextItem(const FVector2D& InPosition, const FText& InText) const;
+};
