Material Design
- Change to prevent object gc at runtime.

diff --git a/Engine/Plugins/Experimental/DynamicMaterial/Source/DynamicMaterial/Private/Components/DMMaterialValue.cpp b/Engine/Plugins/Experimental/DynamicMaterial/Source/DynamicMaterial/Private/Components/DMMaterialValue.cpp
--- a/Engine/Plugins/Experimental/DynamicMaterial/Source/DynamicMaterial/Private/Components/DMMaterialValue.cpp
+++ b/Engine/Plugins/Experimental/DynamicMaterial/Source/DynamicMaterial/Private/Components/DMMaterialValue.cpp
@@ -117,6 +117,14 @@ void UDMMaterialValue::PostLoad()
 
 		Parameter->SetParentComponent(this);
 	}
+
+	if (bLocal)
+	{
+		if (UDynamicMaterialModel* MaterialModel = GetMaterialModel())
+		{
+			MaterialModel->AddRuntimeComponentReference(this);
+		}
+	}
 }
 
 void UDMMaterialValue::PostEditImport()
@@ -203,6 +211,19 @@ bool UDMMaterialValue::SetParameterName(FName InBaseName)
 	return true;
 }
 
+void UDMMaterialValue::OnComponentAdded()
+{
+	Super::OnComponentAdded();
+
+	if (bLocal)
+	{
+		if (UDynamicMaterialModel* MaterialModel = GetMaterialModel())
+		{
+			MaterialModel->AddRuntimeComponentReference(this);
+		}
+	}
+}
+
 void UDMMaterialValue::OnComponentRemoved()
 {
 	if (Parameter)
@@ -214,6 +235,14 @@ void UDMMaterialValue::OnComponentRemoved()
 
 		Parameter->SetComponentState(EDMComponentLifetimeState::Removed);
 	}
+
+	if (bLocal)
+	{
+		if (UDynamicMaterialModel* MaterialModel = GetMaterialModel())
+		{
+			MaterialModel->RemoveRuntimeComponentReference(this);
+		}
+	}
  
 	Super::OnComponentRemoved();
 }
@@ -590,7 +619,7 @@ void UDMMaterialValue::UpdatePreviewMaterial()
  
 	BuildState->GetBuildUtils().UpdatePreviewMaterial(ValueExpression, 0, FDMMaterialStageConnectorChannel::WHOLE_CHANNEL, 32);
 }
- 
+
 int32 UDMMaterialValue::GetInnateMaskOutput(int32 OutputChannels) const
 {
 	return INDEX_NONE;
diff --git a/Engine/Plugins/Experimental/DynamicMaterial/Source/DynamicMaterial/Private/Components/DMTextureUV.cpp b/Engine/Plugins/Experimental/DynamicMaterial/Source/DynamicMaterial/Private/Components/DMTextureUV.cpp
--- a/Engine/Plugins/Experimental/DynamicMaterial/Source/DynamicMaterial/Private/Components/DMTextureUV.cpp
+++ b/Engine/Plugins/Experimental/DynamicMaterial/Source/DynamicMaterial/Private/Components/DMTextureUV.cpp
@@ -570,6 +570,11 @@ void UDMTextureUV::OnComponentAdded()
 	}
 
 	CreateParameterNames();
+
+	if (UDynamicMaterialModel* MaterialModel = GetMaterialModel())
+	{
+		MaterialModel->AddRuntimeComponentReference(this);
+	}
 	
 	Super::OnComponentAdded();
 }
@@ -578,6 +583,11 @@ void UDMTextureUV::OnComponentRemoved()
 {
 	RemoveParameterNames();
 
+	if (UDynamicMaterialModel* MaterialModel = GetMaterialModel())
+	{
+		MaterialModel->RemoveRuntimeComponentReference(this);
+	}
+
 	Super::OnComponentRemoved();
 }
 
@@ -644,6 +654,11 @@ void UDMTextureUV::PostLoad()
 		return;
 	}
 
+	if (UDynamicMaterialModel* MaterialModel = GetMaterialModel())
+	{
+		MaterialModel->AddRuntimeComponentReference(this);
+	}
+
 	if (MaterialParameters.IsEmpty())
 	{
 		CreateParameterNames();
diff --git a/Engine/Plugins/Experimental/DynamicMaterial/Source/DynamicMaterial/Private/Model/DynamicMaterialModel.cpp b/Engine/Plugins/Experimental/DynamicMaterial/Source/DynamicMaterial/Private/Model/DynamicMaterialModel.cpp
--- a/Engine/Plugins/Experimental/DynamicMaterial/Source/DynamicMaterial/Private/Model/DynamicMaterialModel.cpp
+++ b/Engine/Plugins/Experimental/DynamicMaterial/Source/DynamicMaterial/Private/Model/DynamicMaterialModel.cpp
@@ -163,6 +163,16 @@ UDMMaterialValue* UDynamicMaterialModel::AddValue(EDMValueType InType)
 	return NewValue;
 }
 
+void UDynamicMaterialModel::AddRuntimeComponentReference(UDMMaterialComponent* InValue)
+{
+	RuntimeComponents.Add(InValue);
+}
+
+void UDynamicMaterialModel::RemoveRuntimeComponentReference(UDMMaterialComponent* InValue)
+{
+	RuntimeComponents.Remove(InValue);
+}
+
 void UDynamicMaterialModel::RemoveValueByName(FName InName)
 {
 	int32 FoundIndex = Values.IndexOfByPredicate([InName](UDMMaterialValue* Value)
diff --git a/Engine/Plugins/Experimental/DynamicMaterial/Source/DynamicMaterial/Public/Components/DMMaterialValue.h b/Engine/Plugins/Experimental/DynamicMaterial/Source/DynamicMaterial/Public/Components/DMMaterialValue.h
--- a/Engine/Plugins/Experimental/DynamicMaterial/Source/DynamicMaterial/Public/Components/DMMaterialValue.h
+++ b/Engine/Plugins/Experimental/DynamicMaterial/Source/DynamicMaterial/Public/Components/DMMaterialValue.h
@@ -190,6 +190,7 @@ protected:
 	void UpdatePreviewMaterial();
  
 	//~ Begin UDMMaterialComponent
+	virtual void OnComponentAdded() override;
 	virtual void OnComponentRemoved() override;
 	//~ End UDMMaterialComponent
  
diff --git a/Engine/Plugins/Experimental/DynamicMaterial/Source/DynamicMaterial/Public/Model/DynamicMaterialModel.h b/Engine/Plugins/Experimental/DynamicMaterial/Source/DynamicMaterial/Public/Model/DynamicMaterialModel.h
--- a/Engine/Plugins/Experimental/DynamicMaterial/Source/DynamicMaterial/Public/Model/DynamicMaterialModel.h
+++ b/Engine/Plugins/Experimental/DynamicMaterial/Source/DynamicMaterial/Public/Model/DynamicMaterialModel.h
@@ -3,7 +3,8 @@
 #pragma once
 
 #include "UObject/Object.h"
-#include "Containers/ContainersFwd.h"
+#include "Containers/Array.h"
+#include "Containers/Set.h"
 
 #if WITH_EDITOR
 #include "UObject/ScriptInterface.h"
@@ -87,6 +88,12 @@ public:
 	UFUNCTION(BlueprintCallable, Category = "Material Designer")
 	UDMMaterialValue* AddValue(EDMValueType InType);
 
+	UFUNCTION(BlueprintCallable, Category = "Material Designer")
+	void AddRuntimeComponentReference(UDMMaterialComponent* InValue);
+
+	UFUNCTION(BlueprintCallable, Category = "Material Designer")
+	void RemoveRuntimeComponentReference(UDMMaterialComponent* InValue);
+
 	UFUNCTION(BlueprintCallable, Category = "Material Designer")
 	void RemoveValueByName(FName InName);
 
@@ -129,9 +136,14 @@ public:
 #endif
 
 protected:
+	/** Global values */
 	UPROPERTY(VisibleInstanceOnly, BlueprintReadOnly, Instanced, Category = "Material Designer")
 	TArray<TObjectPtr<UDMMaterialValue>> Values;
 
+	/** References to runtime components outered to this model which are not otherwise referenced. */
+	UPROPERTY(VisibleInstanceOnly, BlueprintReadOnly, Instanced, Category = "Material Designer")
+	TSet<TObjectPtr<UDMMaterialComponent>> RuntimeComponents;
+
 	UPROPERTY(VisibleInstanceOnly, BlueprintReadOnly, Category = "Material Designer")
 	TObjectPtr<UDMMaterialValueFloat1> GlobalOpacityValue;
 
