Material Designer
- Enabled float ranges on float values.

diff --git a/Engine/Plugins/Experimental/DynamicMaterial/Source/DynamicMaterialEditor/Private/Slate/Properties/Generators/DMMaterialValuePropertyRowGenerator.cpp b/Engine/Plugins/Experimental/DynamicMaterial/Source/DynamicMaterialEditor/Private/Slate/Properties/Generators/DMMaterialValuePropertyRowGenerator.cpp
--- a/Engine/Plugins/Experimental/DynamicMaterial/Source/DynamicMaterialEditor/Private/Slate/Properties/Generators/DMMaterialValuePropertyRowGenerator.cpp
+++ b/Engine/Plugins/Experimental/DynamicMaterial/Source/DynamicMaterialEditor/Private/Slate/Properties/Generators/DMMaterialValuePropertyRowGenerator.cpp
@@ -5,6 +5,7 @@
 #include "IDetailPropertyRow.h"
 #include "Components/DMMaterialComponent.h"
 #include "Components/DMMaterialValue.h"
+#include "Components/MaterialValues/DMMaterialValueFloat1.h"
 #include "Slate/SDMEditor.h"
 #include "Slate/SDMComponentEdit.h"
 
@@ -49,6 +50,22 @@ void FDMMaterialValuePropertyRowGenerator::AddComponentProperties(const TSharedR
 		FResetToDefaultHandler::CreateUObject(Value, &UDMMaterialValue::ResetToDefault)
 	);
 
+	if (UDMMaterialValueFloat* FloatValue = Cast<UDMMaterialValueFloat>(Value))
+	{
+		if (FloatValue->HasValueRange())
+		{
+			static const FName UIMin = FName("UIMin");
+			static const FName UIMax = FName("UIMax");
+			static const FName ClampMin = FName("ClampMin");
+			static const FName ClampMax = FName("ClampMax");
+
+			Handle.PropertyHandle->SetInstanceMetaData(UIMin, FString::SanitizeFloat(FloatValue->GetValueRange().Min));
+			Handle.PropertyHandle->SetInstanceMetaData(ClampMin, FString::SanitizeFloat(FloatValue->GetValueRange().Min));
+			Handle.PropertyHandle->SetInstanceMetaData(UIMax, FString::SanitizeFloat(FloatValue->GetValueRange().Max));
+			Handle.PropertyHandle->SetInstanceMetaData(ClampMax, FString::SanitizeFloat(FloatValue->GetValueRange().Max));
+		}
+	}
+
 	InOutPropertyRows.Add(Handle);
 
 	const TArray<FName>& Properties = Value->GetEditableProperties();
