Material Designer
- Added DeferredDecal as a material domain option.
- Swapped the position of "Unlit" and "Animated" in the UI.

diff --git a/Engine/Plugins/Experimental/DynamicMaterial/Source/DynamicMaterialEditor/Private/Slate/Properties/SDMDomain.cpp b/Engine/Plugins/Experimental/DynamicMaterial/Source/DynamicMaterialEditor/Private/Slate/Properties/SDMDomain.cpp
--- a/Engine/Plugins/Experimental/DynamicMaterial/Source/DynamicMaterialEditor/Private/Slate/Properties/SDMDomain.cpp
+++ b/Engine/Plugins/Experimental/DynamicMaterial/Source/DynamicMaterialEditor/Private/Slate/Properties/SDMDomain.cpp
@@ -11,7 +11,8 @@ namespace UE::DynamicMaterialEditor::Private
 	static const TArray<EMaterialDomain> SupportedDomains =
 	{
 		EMaterialDomain::MD_Surface,
-		EMaterialDomain::MD_PostProcess
+		EMaterialDomain::MD_PostProcess,
+		EMaterialDomain::MD_DeferredDecal
 	};
 
 	const TArray<FName>& SupportedDomainNames()
diff --git a/Engine/Plugins/Experimental/DynamicMaterial/Source/DynamicMaterialEditor/Private/Slate/SDMEditor.cpp b/Engine/Plugins/Experimental/DynamicMaterial/Source/DynamicMaterialEditor/Private/Slate/SDMEditor.cpp
--- a/Engine/Plugins/Experimental/DynamicMaterial/Source/DynamicMaterialEditor/Private/Slate/SDMEditor.cpp
+++ b/Engine/Plugins/Experimental/DynamicMaterial/Source/DynamicMaterialEditor/Private/Slate/SDMEditor.cpp
@@ -648,21 +648,18 @@ TSharedRef<SWidget> SDMEditor::CreateMaterialSettingsRow()
 		.AutoWidth()
 		.HAlign(HAlign_Right)
 		.VAlign(VAlign_Center)
-		.Padding(10.0f, 0.0f, 0.0f, 2.0f)
+		.Padding(10.f, 0.0f, 5.0f, 2.0f)
 		[
 			SNew(SHorizontalBox)
-			.ToolTipText(LOCTEXT("MaterialDesignerInstanceTypeTooltip", "Enables the material's TSR pixel animation flag. Not available in translucent blend modes."))
+			.ToolTipText(LOCTEXT("MaterialDesignerInstanceUnlitTooltip", "Toggle between Default Lit and Unlit for this Material Designer Instance."))
 			+ SHorizontalBox::Slot()
 			.AutoWidth()
 			.HAlign(HAlign_Left)
 			.VAlign(VAlign_Center)
+			.Padding(0.0f, 0.0f, 5.0f, 0.0f)
 			[
-				SNew(SBox)
-				.Padding(0.f, 0.f, 5.f, 0.f)
-				[
-					SNew(STextBlock)
-					.Text(LOCTEXT("MaterialDesignerInstanceAnimaed", "Animated"))
-				]
+				SNew(STextBlock)
+				.Text(LOCTEXT("MaterialDesignerInstanceUnlit", "Unlit"))
 			]
 			+ SHorizontalBox::Slot()
 			.AutoWidth()
@@ -670,27 +667,30 @@ TSharedRef<SWidget> SDMEditor::CreateMaterialSettingsRow()
 			.VAlign(VAlign_Center)
 			[
 				SNew(SCheckBox)
-				.IsChecked(this, &SDMEditor::IsMaterialAnimated)
-				.IsEnabled(this, &SDMEditor::CanMaterialBeAnimated)
-				.OnCheckStateChanged(this, &SDMEditor::OnMaterialAnimatedChanged)
+				.IsChecked(this, &SDMEditor::IsMaterialUnlit)
+				.OnCheckStateChanged(this, &SDMEditor::OnMaterialUnlitChanged)
+				.IsEnabled(this, &SDMEditor::CanChangeMaterialShadingModel)
 			]
 		]
 		+ SHorizontalBox::Slot()
 		.AutoWidth()
 		.HAlign(HAlign_Right)
 		.VAlign(VAlign_Center)
-		.Padding(10.f, 0.0f, 5.0f, 2.0f)
+		.Padding(10.0f, 0.0f, 0.0f, 2.0f)
 		[
 			SNew(SHorizontalBox)
-			.ToolTipText(LOCTEXT("MaterialDesignerInstanceUnlitTooltip", "Toggle between Default Lit and Unlit for this Material Designer Instance."))
+			.ToolTipText(LOCTEXT("MaterialDesignerInstanceTypeTooltip", "Enables the material's TSR pixel animation flag. Not available in translucent blend modes."))
 			+ SHorizontalBox::Slot()
 			.AutoWidth()
 			.HAlign(HAlign_Left)
 			.VAlign(VAlign_Center)
-			.Padding(0.0f, 0.0f, 5.0f, 0.0f)
 			[
-				SNew(STextBlock)
-				.Text(LOCTEXT("MaterialDesignerInstanceUnlit", "Unlit"))
+				SNew(SBox)
+				.Padding(0.f, 0.f, 5.f, 0.f)
+				[
+					SNew(STextBlock)
+					.Text(LOCTEXT("MaterialDesignerInstanceAnimaed", "Animated"))
+				]
 			]
 			+ SHorizontalBox::Slot()
 			.AutoWidth()
@@ -698,9 +698,9 @@ TSharedRef<SWidget> SDMEditor::CreateMaterialSettingsRow()
 			.VAlign(VAlign_Center)
 			[
 				SNew(SCheckBox)
-				.IsChecked(this, &SDMEditor::IsMaterialUnlit)
-				.OnCheckStateChanged(this, &SDMEditor::OnMaterialUnlitChanged)
-				.IsEnabled(this, &SDMEditor::CanChangeMaterialShadingModel)
+				.IsChecked(this, &SDMEditor::IsMaterialAnimated)
+				.IsEnabled(this, &SDMEditor::CanMaterialBeAnimated)
+				.OnCheckStateChanged(this, &SDMEditor::OnMaterialAnimatedChanged)
 			]
 		];
 }
@@ -1489,7 +1489,12 @@ bool SDMEditor::CanChangeBlendType() const
 {
 	if (UDynamicMaterialModelEditorOnlyData* ModelEditorOnlyData = UDynamicMaterialModelEditorOnlyData::Get(MaterialModelWeak))
 	{
-		return ModelEditorOnlyData->GetDomain() != EMaterialDomain::MD_PostProcess;
+		switch (ModelEditorOnlyData->GetDomain())
+		{
+			case EMaterialDomain::MD_Surface:
+			case EMaterialDomain::MD_DeferredDecal:
+				return true;
+		}
 	}
 
 	return false;
@@ -1539,7 +1544,11 @@ bool SDMEditor::CanChangeMaterialShadingModel() const
 {
 	if (UDynamicMaterialModelEditorOnlyData* ModelEditorOnlyData = UDynamicMaterialModelEditorOnlyData::Get(MaterialModelWeak))
 	{
-		return ModelEditorOnlyData->GetDomain() != EMaterialDomain::MD_PostProcess;
+		switch (ModelEditorOnlyData->GetDomain())
+		{
+			case EMaterialDomain::MD_Surface:
+				return true;
+		}
 	}
 
 	return false;
