Material Designer
- Fix for the text renderer so it works at runtime.

diff --git a/Engine/Plugins/Experimental/DynamicMaterial/Source/DynamicMaterial/DynamicMaterial.Build.cs b/Engine/Plugins/Experimental/DynamicMaterial/Source/DynamicMaterial/DynamicMaterial.Build.cs
--- a/Engine/Plugins/Experimental/DynamicMaterial/Source/DynamicMaterial/DynamicMaterial.Build.cs
+++ b/Engine/Plugins/Experimental/DynamicMaterial/Source/DynamicMaterial/DynamicMaterial.Build.cs
@@ -12,7 +12,8 @@ public class DynamicMaterial : ModuleRules
 			new string[]
 			{
 				"Core",
-				"CoreUObject"
+				"CoreUObject",
+				"Slate"
 			}
 		);
 
diff --git a/Engine/Plugins/Experimental/DynamicMaterial/Source/DynamicMaterial/Private/Components/MaterialValues/DMMaterialValueText.cpp b/Engine/Plugins/Experimental/DynamicMaterial/Source/DynamicMaterial/Private/Components/MaterialValues/DMMaterialValueText.cpp
new file mode 100644
--- /dev/null
+++ b/Engine/Plugins/Experimental/DynamicMaterial/Source/DynamicMaterial/Private/Components/MaterialValues/DMMaterialValueText.cpp
@@ -0,0 +1,784 @@
+// Copyright Epic Games, Inc. All Rights Reserved.
+
+#include "Components/MaterialValues/DMMaterialValueText.h"
+
+#include "CanvasItem.h"
+#include "CanvasTypes.h"
+#include "Engine/Canvas.h"
+#include "Engine/Engine.h"
+#include "Engine/TextureRenderTarget2D.h"
+#include "TextureResource.h"
+
+#define LOCTEXT_NAMESPACE "DMMaterialStageExpressionText"
+
+namespace UE::DynamicMaterial::Private
+{
+	const FIntPoint MinimumTextTextureSize = FIntPoint(5, 10);
+	const FIntPoint MaxTexScale = FIntPoint(64, 64);
+}
+
+#if WITH_EDITOR
+struct FDMMaterialValueText
+{
+	static const inline FName FontName = GET_MEMBER_NAME_CHECKED(UDMMaterialValueText, Font);
+	static const inline FName TextName = GET_MEMBER_NAME_CHECKED(UDMMaterialValueText, Text);
+	static const inline FName ColorName = GET_MEMBER_NAME_CHECKED(UDMMaterialValueText, Color);
+	static const inline FName BackgroundColorName = GET_MEMBER_NAME_CHECKED(UDMMaterialValueText, BackgroundColor);
+	static const inline FName JustifyName = GET_MEMBER_NAME_CHECKED(UDMMaterialValueText, Justify);
+	static const inline FName KerningName = GET_MEMBER_NAME_CHECKED(UDMMaterialValueText, Kerning);
+	static const inline FName LineHeightName = GET_MEMBER_NAME_CHECKED(UDMMaterialValueText, LineHeight);
+	static const inline FName PaddingLeftName = GET_MEMBER_NAME_CHECKED(UDMMaterialValueText, PaddingLeft);
+	static const inline FName PaddingRightName = GET_MEMBER_NAME_CHECKED(UDMMaterialValueText, PaddingRight);
+	static const inline FName PaddingTopName = GET_MEMBER_NAME_CHECKED(UDMMaterialValueText, PaddingTop);
+	static const inline FName PaddingBottomName = GET_MEMBER_NAME_CHECKED(UDMMaterialValueText, PaddingBottom);
+	static const inline FName TextScaleName = GET_MEMBER_NAME_CHECKED(UDMMaterialValueText, TextScale);
+	static const inline FName TextureSizeOverrideName = GET_MEMBER_NAME_CHECKED(UDMMaterialValueText, TextureSizeOverride);
+	static const inline FName bOutlineName = GET_MEMBER_NAME_CHECKED(UDMMaterialValueText, bOutline);
+	static const inline FName OutlineColorName = GET_MEMBER_NAME_CHECKED(UDMMaterialValueText, OutlineColor);
+	static const inline FName bShadowName = GET_MEMBER_NAME_CHECKED(UDMMaterialValueText, bShadow);
+	static const inline FName ShadowColorName = GET_MEMBER_NAME_CHECKED(UDMMaterialValueText, ShadowColor);
+	static const inline FName ShadowOffsetName = GET_MEMBER_NAME_CHECKED(UDMMaterialValueText, ShadowOffset);
+	static const inline FName bGlowName = GET_MEMBER_NAME_CHECKED(UDMMaterialValueText, bGlow);
+	static const inline FName GlowColorName = GET_MEMBER_NAME_CHECKED(UDMMaterialValueText, GlowColor);
+	static const inline FName GlowInnerRadiusName = GET_MEMBER_NAME_CHECKED(UDMMaterialValueText, GlowInnerRadius);
+	static const inline FName GlowOuterRadiusName = GET_MEMBER_NAME_CHECKED(UDMMaterialValueText, GlowOuterRadius);
+
+	static const inline TSet<FName> PropertyNames = {
+		FontName, TextName, ColorName, BackgroundColorName, JustifyName, KerningName, LineHeightName, PaddingLeftName, PaddingRightName,
+		PaddingTopName, PaddingBottomName, TextScaleName, TextureSizeOverrideName, bOutlineName, OutlineColorName, bShadowName,
+		ShadowColorName, ShadowOffsetName, bGlowName, GlowColorName, GlowInnerRadiusName, GlowOuterRadiusName
+	};
+};
+#endif
+
+UDMMaterialValueText::UDMMaterialValueText()
+	: UDMMaterialValueTexture()
+{
+	Type = EDMValueType::VT_Text;
+
+#if WITH_EDITOR
+	EditableProperties.Add(FDMMaterialValueText::FontName);
+	EditableProperties.Add(FDMMaterialValueText::TextName);
+	EditableProperties.Add(FDMMaterialValueText::ColorName);
+	EditableProperties.Add(FDMMaterialValueText::BackgroundColorName);
+	EditableProperties.Add(FDMMaterialValueText::JustifyName);
+	EditableProperties.Add(FDMMaterialValueText::KerningName);
+	EditableProperties.Add(FDMMaterialValueText::LineHeightName);
+	EditableProperties.Add(FDMMaterialValueText::PaddingLeftName);
+	EditableProperties.Add(FDMMaterialValueText::PaddingRightName);
+	EditableProperties.Add(FDMMaterialValueText::PaddingTopName);
+	EditableProperties.Add(FDMMaterialValueText::PaddingBottomName);
+	EditableProperties.Add(FDMMaterialValueText::TextScaleName);
+	EditableProperties.Add(FDMMaterialValueText::TextureSizeOverrideName);
+	EditableProperties.Add(FDMMaterialValueText::bOutlineName);
+	EditableProperties.Add(FDMMaterialValueText::OutlineColorName);
+	EditableProperties.Add(FDMMaterialValueText::bShadowName);
+	EditableProperties.Add(FDMMaterialValueText::ShadowColorName);
+	EditableProperties.Add(FDMMaterialValueText::ShadowOffsetName);
+	EditableProperties.Add(FDMMaterialValueText::bGlowName);
+	EditableProperties.Add(FDMMaterialValueText::GlowColorName);
+	EditableProperties.Add(FDMMaterialValueText::GlowInnerRadiusName);
+	EditableProperties.Add(FDMMaterialValueText::GlowOuterRadiusName);
+#endif
+}
+
+void UDMMaterialValueText::PostLoad()
+{
+	Super::PostLoad();
+
+	RequestUpdateTextTexture(/* Needs Resize */ true);
+}
+
+UFont* UDMMaterialValueText::GetFont() const
+{
+	return Font;
+}
+
+void UDMMaterialValueText::SetFont(UFont* InFont)
+{
+	if (Font == InFont)
+	{
+		return;
+	}
+
+	Font = InFont;
+
+	RequestUpdateTextTexture(/* Needs Resize */ true);
+}
+
+const FText& UDMMaterialValueText::GetText() const
+{
+	return Text;
+}
+
+void UDMMaterialValueText::SetText(const FText& InText)
+{
+	if (Text.EqualTo(InText, ETextComparisonLevel::Default))
+	{
+		return;
+	}
+
+	Text = InText;
+
+	UpdateTextLines();
+	RequestUpdateTextTexture(/* Needs Resize */ true);
+}
+
+const FLinearColor& UDMMaterialValueText::GetColor() const
+{
+	return Color;
+}
+
+void UDMMaterialValueText::SetColor(const FLinearColor& InColor)
+{
+	if (Color == InColor)
+	{
+		return;
+	}
+
+	Color = InColor;
+
+	RequestUpdateTextTexture(/* Needs Resize */ false);
+}
+
+const FLinearColor& UDMMaterialValueText::GetBackgroundColor() const
+{
+	return BackgroundColor;
+}
+
+void UDMMaterialValueText::SetBackgroundColor(const FLinearColor& InBackgroundColor)
+{
+	if (BackgroundColor == InBackgroundColor)
+	{
+		return;
+	}
+
+	BackgroundColor = InBackgroundColor;
+
+	RequestUpdateTextTexture(/* Needs Resize */ false);
+}
+
+ETextJustify::Type UDMMaterialValueText::GetJustify() const
+{
+	return Justify;
+}
+
+void UDMMaterialValueText::SetJustify(ETextJustify::Type InJustify)
+{
+	if (Justify == InJustify)
+	{
+		return;
+	}
+
+	Justify = InJustify;
+
+	RequestUpdateTextTexture(/* Needs Resize */ false);
+}
+
+float UDMMaterialValueText::GetKerning() const
+{
+	return Kerning;
+}
+
+void UDMMaterialValueText::SetKerning(float InKerning)
+{
+	if (Kerning == InKerning)
+	{
+		return;
+	}
+
+	Kerning = InKerning;
+
+	RequestUpdateTextTexture(/* Needs Resize */ true);
+}
+
+const FLinearColor& UDMMaterialValueText::GetOutlineColor() const
+{
+	return OutlineColor;
+}
+
+void UDMMaterialValueText::SetOutlineColor(const FLinearColor& InOutlineColor)
+{
+	if (OutlineColor == InOutlineColor)
+	{
+		return;
+	}
+
+	OutlineColor = InOutlineColor;
+
+	RequestUpdateTextTexture(/* Needs Resize */ false);
+}
+
+float UDMMaterialValueText::GetLineHeight() const
+{
+	return LineHeight;
+}
+
+void UDMMaterialValueText::SetLineHeight(float InLineHeight)
+{
+	if (LineHeight == InLineHeight)
+	{
+		return;
+	}
+
+	LineHeight = InLineHeight;
+
+	RequestUpdateTextTexture(/* Needs Resize */ true);
+}
+
+float UDMMaterialValueText::GetPaddingLeft() const
+{
+	return PaddingLeft;
+}
+
+void UDMMaterialValueText::SetPaddingLeft(float InPaddingLeft)
+{
+	if (PaddingLeft == InPaddingLeft)
+	{
+		return;
+	}
+
+	PaddingLeft = InPaddingLeft;
+
+	RequestUpdateTextTexture(/* Needs Resize */ true);
+}
+
+float UDMMaterialValueText::GetPaddingRight() const
+{
+	return PaddingRight;
+}
+
+void UDMMaterialValueText::SetPaddingRight(float InPaddingRight)
+{
+	if (PaddingRight == InPaddingRight)
+	{
+		return;
+	}
+
+	PaddingRight = InPaddingRight;
+
+	RequestUpdateTextTexture(/* Needs Resize */ true);
+}
+
+float UDMMaterialValueText::GetPaddingTop() const
+{
+	return PaddingTop;
+}
+
+void UDMMaterialValueText::SetPaddingTop(float InPaddingTop)
+{
+	if (PaddingTop == InPaddingTop)
+	{
+		return;
+	}
+
+	PaddingTop = InPaddingTop;
+
+	RequestUpdateTextTexture(/* Needs Resize */ true);
+}
+
+float UDMMaterialValueText::GetPaddingBottom() const
+{
+	return PaddingBottom;
+}
+
+void UDMMaterialValueText::SetPaddingBottom(float InPaddingBottom)
+{
+	if (PaddingBottom == InPaddingBottom)
+	{
+		return;
+	}
+
+	PaddingBottom = InPaddingBottom;
+
+	RequestUpdateTextTexture(/* Needs Resize */ true);
+}
+
+const FIntPoint& UDMMaterialValueText::GetTextScale() const
+{
+	return TextScale;
+}
+
+void UDMMaterialValueText::SetTextScale(const FIntPoint& InTextScale)
+{
+	FIntPoint NewTextScale = InTextScale;
+	NewTextScale.X = FMath::Clamp(InTextScale.X, 1, UE::DynamicMaterial::Private::MaxTexScale.X);
+	NewTextScale.Y = FMath::Clamp(InTextScale.Y, 1, UE::DynamicMaterial::Private::MaxTexScale.Y);
+
+	if (TextScale == NewTextScale)
+	{
+		return;
+	}
+
+	TextScale = NewTextScale;
+
+	RequestUpdateTextTexture(/* Needs Resize */ true);
+}
+
+const FIntPoint& UDMMaterialValueText::GetTextureSizeOverride() const
+{
+	return TextureSizeOverride;
+}
+
+void UDMMaterialValueText::SetTextureSizeOverride(const FIntPoint& InTextureSizeOverride)
+{
+	if (TextureSizeOverride == InTextureSizeOverride)
+	{
+		return;
+	}
+
+	TextureSizeOverride = InTextureSizeOverride;
+
+	RequestUpdateTextTexture(/* Needs Resize */ true);
+}
+
+bool UDMMaterialValueText::GetHasOutline() const
+{
+	return bOutline;
+}
+
+void UDMMaterialValueText::SetHasOutline(bool bInHasOutline)
+{
+	if (bOutline == bInHasOutline)
+	{
+		return;
+	}
+
+	bOutline = bInHasOutline;
+
+	RequestUpdateTextTexture(/* Needs Resize */ false);
+}
+
+bool UDMMaterialValueText::GetHasShadow() const
+{
+	return bShadow;
+}
+
+void UDMMaterialValueText::SetHasShadow(bool bInHasShadow)
+{
+	if (bShadow == bInHasShadow)
+	{
+		return;
+	}
+
+	bShadow = bInHasShadow;
+
+	RequestUpdateTextTexture(/* Needs Resize */ false);
+}
+
+const FLinearColor& UDMMaterialValueText::GetShadowColor() const
+{
+	return ShadowColor;
+}
+
+void UDMMaterialValueText::SetShadowColor(const FLinearColor& InShadowColor)
+{
+	if (ShadowColor == InShadowColor)
+	{
+		return;
+	}
+
+	ShadowColor = InShadowColor;
+
+	RequestUpdateTextTexture(/* Needs Resize */ false);
+}
+
+const FVector2D& UDMMaterialValueText::GetShadowOffset() const
+{
+	return ShadowOffset;
+}
+
+void UDMMaterialValueText::SetShadowOffset(const FVector2D& InShadowOffset)
+{
+	if (ShadowOffset == InShadowOffset)
+	{
+		return;
+	}
+
+	ShadowOffset = InShadowOffset;
+
+	RequestUpdateTextTexture(/* Needs Resize */ false);
+}
+
+bool UDMMaterialValueText::GetHasGlow() const
+{
+	return bGlow;
+}
+
+void UDMMaterialValueText::SetHasGlow(bool bInHasGlow)
+{
+	if (bGlow == bInHasGlow)
+	{
+		return;
+	}
+
+	bGlow = bInHasGlow;
+
+	RequestUpdateTextTexture(/* Needs Resize */ false);
+}
+
+const FLinearColor& UDMMaterialValueText::GetGlowColor() const
+{
+	return GlowColor;
+}
+
+void UDMMaterialValueText::SetGlowColor(const FLinearColor& InGlowColor)
+{
+	if (GlowColor == InGlowColor)
+	{
+		return;
+	}
+
+	GlowColor = InGlowColor;
+
+	RequestUpdateTextTexture(/* Needs Resize */ false);
+}
+
+const FVector2D& UDMMaterialValueText::GetGlowInnerRadius() const
+{
+	return GlowInnerRadius;
+}
+
+void UDMMaterialValueText::SetGlowInnerRadius(const FVector2D& InGlowInnerRadius)
+{
+	if (GlowInnerRadius == InGlowInnerRadius)
+	{
+		return;
+	}
+
+	GlowInnerRadius = InGlowInnerRadius;
+
+	RequestUpdateTextTexture(/* Needs Resize */ false);
+}
+
+const FVector2D& UDMMaterialValueText::GetGlowOuterRadius() const
+{
+	return GlowOuterRadius;
+}
+
+void UDMMaterialValueText::SetGlowOuterRadius(const FVector2D& InGlowOuterRadius)
+{
+	if (GlowOuterRadius == InGlowOuterRadius)
+	{
+		return;
+	}
+
+	GlowOuterRadius = InGlowOuterRadius;
+
+	RequestUpdateTextTexture(/* Needs Resize */ false);
+}
+
+#if WITH_EDITOR
+void UDMMaterialValueText::PostEditChangeProperty(FPropertyChangedEvent& InPropertyChangedEvent)
+{
+	Super::PostEditChangeProperty(InPropertyChangedEvent);
+
+	const FName PropertyName = InPropertyChangedEvent.GetMemberPropertyName();
+
+	if (PropertyName == FDMMaterialValueText::TextName)
+	{
+		UpdateTextLines();
+		RequestUpdateTextTexture(/* Needs Resize */ true);
+	}
+	else if (PropertyName == FDMMaterialValueText::FontName
+		|| PropertyName == FDMMaterialValueText::KerningName
+		|| PropertyName == FDMMaterialValueText::LineHeightName
+		|| PropertyName == FDMMaterialValueText::PaddingLeftName
+		|| PropertyName == FDMMaterialValueText::PaddingRightName
+		|| PropertyName == FDMMaterialValueText::PaddingTopName
+		|| PropertyName == FDMMaterialValueText::PaddingBottomName
+		|| PropertyName == FDMMaterialValueText::TextScaleName
+		|| PropertyName == FDMMaterialValueText::TextureSizeOverrideName)
+	{
+		TextScale.X = FMath::Clamp(TextScale.X, 1, UE::DynamicMaterial::Private::MaxTexScale.X);
+		TextScale.Y = FMath::Clamp(TextScale.Y, 1, UE::DynamicMaterial::Private::MaxTexScale.Y);
+		RequestUpdateTextTexture(/* Needs Resize */ true);
+	}
+	else if (PropertyName == FDMMaterialValueText::bGlowName
+		|| PropertyName == FDMMaterialValueText::bShadowName
+		|| PropertyName == FDMMaterialValueText::bOutlineName)
+	{
+		// Cause details panel refresh.
+		Update(EDMUpdateType::Structure);
+		RequestUpdateTextTexture(/* Needs Resize */ false);
+	}
+	else if (FDMMaterialValueText::PropertyNames.Contains(PropertyName))
+	{
+		RequestUpdateTextTexture(/* Needs Resize */ false);
+	}
+}
+#endif
+
+UTextureRenderTarget2D* UDMMaterialValueText::CreateRenderTarget(UObject* InOuter, const FIntPoint& InSize, 
+	const FLinearColor& InClearColor)
+{
+	UTextureRenderTarget2D* RenderTarget = NewObject<UTextureRenderTarget2D>(InOuter, NAME_None, 
+		EObjectFlags::RF_Transactional | EObjectFlags::RF_DuplicateTransient | EObjectFlags::RF_TextExportTransient
+	);
+
+	check(RenderTarget);
+	RenderTarget->RenderTargetFormat = RTF_RGBA16f;
+	RenderTarget->ClearColor = InClearColor;
+	RenderTarget->bAutoGenerateMips = false;
+	RenderTarget->bCanCreateUAV = false;
+	RenderTarget->InitAutoFormat(InSize.X, InSize.Y);
+	RenderTarget->UpdateResourceImmediate(true);
+
+	return RenderTarget;
+}
+
+void UDMMaterialValueText::CalculateLineSizes()
+{
+	FTextSizingParameters Params(0, 0, 0, 0, Font);
+	Params.Scaling = FVector2D(1, 1);
+
+	for (FDMTextLine& Line : Lines)
+	{
+		Params.DrawXL = 0;
+		Params.DrawYL = 0;
+
+		UCanvas::CanvasStringSize(Params, *Line.Line);
+
+		int32 RequiredWidth = Params.DrawXL;
+		RequiredWidth += Line.Line.Len() > 0 ? ((Line.Line.Len() - 1) * Kerning) : 0;
+
+		Line.Width = RequiredWidth;
+	}
+}
+
+FIntPoint UDMMaterialValueText::GetRequiredTextureSize() const
+{
+	if (TextureSizeOverride.X > 0 && TextureSizeOverride.Y > 0)
+	{
+		return TextureSizeOverride;
+	}
+
+	if (!IsValid(Font) || Lines.IsEmpty())
+	{
+		return UE::DynamicMaterial::Private::MinimumTextTextureSize;
+	}
+
+	const float MaxCharHeight = Font->GetMaxCharHeight();
+
+	if (MaxCharHeight <= 0)
+	{
+		return UE::DynamicMaterial::Private::MinimumTextTextureSize;
+	}
+
+	FVector2f Size = FVector2f::ZeroVector;
+
+	for (const FDMTextLine& Line : Lines)
+	{
+		Size.X = FMath::Max(Size.X, Line.Width);
+		Size.Y += MaxCharHeight * LineHeight;
+	}
+
+	Size.X += PaddingLeft + PaddingRight;
+	Size.Y += PaddingTop + PaddingBottom;
+
+	Size.X = FMath::Max(Size.X, UE::DynamicMaterial::Private::MinimumTextTextureSize.X);
+	Size.Y = FMath::Max(Size.Y, UE::DynamicMaterial::Private::MinimumTextTextureSize.Y);
+
+	FIntPoint RequiredSize = FIntPoint::ZeroValue;
+	RequiredSize.X = FMath::CeilToInt(Size.X * TextScale.X);
+	RequiredSize.Y = FMath::CeilToInt(Size.Y * TextScale.Y);
+
+	return RequiredSize;
+}
+
+void UDMMaterialValueText::UpdateTextLines()
+{
+	const FString TextStr = Text.ToString();
+	TArray<FString> NewLines;
+	TextStr.ParseIntoArray(NewLines, TEXT("\n"), /* Cull empty */ false);
+
+	Lines.Empty();
+	Lines.Reserve(NewLines.Num());
+
+	for (FString& NewLine : NewLines)
+	{
+		if (NewLine.EndsWith(TEXT("\r")))
+		{
+			NewLine.LeftChopInline(1);
+		}
+
+		Lines.Add({NewLine, -1.f});
+	}
+}
+
+void UDMMaterialValueText::RequestUpdateTextTexture(bool bInNeedsResize)
+{
+	bNeedsResize = bInNeedsResize;
+
+	if (EndOfFrameDelegateHandle.IsValid())
+	{
+		return;
+	}
+
+	EndOfFrameDelegateHandle = FCoreDelegates::OnEndFrame.AddUObject(this, &UDMMaterialValueText::UpdateTextTexture);
+}
+
+void UDMMaterialValueText::UpdateTextTexture()
+{
+	if (EndOfFrameDelegateHandle.IsValid())
+	{
+		FCoreDelegates::OnEndFrame.Remove(EndOfFrameDelegateHandle);
+		EndOfFrameDelegateHandle.Reset();
+	}
+
+	if (!IsValid(Font) || Lines.IsEmpty())
+	{
+		return;
+	}
+
+	if (bNeedsResize)
+	{
+		CalculateLineSizes();
+	}
+		
+	const bool bValidRenderTarget = IsValid(RenderTarget.Get());
+
+	if (!bValidRenderTarget)
+	{
+		bNeedsResize = true;
+	}
+	else if (bNeedsResize && TextureSizeOverride.X > 0 && TextureSizeOverride.X == RenderTarget->SizeX
+		&& TextureSizeOverride.Y > 0 && TextureSizeOverride.Y == RenderTarget->SizeY)
+	{
+		bNeedsResize = false;
+	}
+
+	if (bNeedsResize)
+	{
+		FIntPoint RequiredSize = GetRequiredTextureSize();
+		RequiredSize.X = FMath::Max(RequiredSize.X, UE::DynamicMaterial::Private::MinimumTextTextureSize.X);
+		RequiredSize.Y = FMath::Max(RequiredSize.Y, UE::DynamicMaterial::Private::MinimumTextTextureSize.Y);
+
+		if (bValidRenderTarget && RequiredSize.X == RenderTarget->SizeX && RequiredSize.X == RenderTarget->SizeY)
+		{
+			bNeedsResize = false;
+		}
+
+		if (bNeedsResize)
+		{
+			RenderTarget = CreateRenderTarget(this, RequiredSize, FLinearColor::Black);
+			SetValue(RenderTarget);
+			bNeedsResize = false;
+		}
+	}
+
+	RenderTarget->UpdateResourceImmediate(true);
+
+	FTextureRenderTargetResource* RenderTargetResource = RenderTarget->GameThread_GetRenderTargetResource();
+
+	if (!RenderTargetResource)
+	{
+		return;
+	}
+
+	static const FGameTime GameTime = FGameTime::CreateUndilated(0, 0);
+
+	FCanvas RenderCanvas(
+		RenderTargetResource,
+		nullptr,
+		GameTime,
+		GEngine->GetDefaultWorldFeatureLevel(),
+		FCanvas::ECanvasDrawMode::CDM_ImmediateDrawing
+	);
+
+	const FVector2D RenderTargetSize = FVector2D(RenderTarget->SizeX, RenderTarget->SizeY);
+
+	FCanvasTileItem TileItem = FCanvasTileItem(FVector2D::ZeroVector, RenderTargetSize, BackgroundColor);
+	RenderCanvas.DrawItem(TileItem);
+
+	const float MaxCharHeight = Font->GetMaxCharHeight();
+
+	FCanvasTextItem TextItem = CreateTextItem(FVector2D::ZeroVector, FText::GetEmpty());
+
+	const float LeftOffset = PaddingLeft * static_cast<float>(TextScale.X);
+	const float RightOffset = PaddingRight * static_cast<float>(TextScale.X);
+
+	TextItem.Position.X = LeftOffset;
+	TextItem.Position.Y = PaddingTop * static_cast<float>(TextScale.Y);
+	TextItem.Scale.X = static_cast<float>(TextScale.X);
+	TextItem.Scale.Y = static_cast<float>(TextScale.Y);
+
+	if (!FMath::IsNearlyEqual(LineHeight, 1.f))
+	{
+		TextItem.Position.Y += MaxCharHeight * (LineHeight - 1.f) * 0.5f * static_cast<float>(TextScale.Y);
+	}
+
+	const float ActualLineHeight = MaxCharHeight * LineHeight * static_cast<float>(TextScale.Y);
+
+	for (const FDMTextLine& Line : Lines)
+	{
+		switch (Justify)
+		{
+			default:
+			case ETextJustify::Left:
+			case ETextJustify::InvariantLeft:
+				TextItem.Position.X = LeftOffset;
+				break;
+
+			case ETextJustify::Center:
+				TextItem.Position.X = (RenderTarget->SizeX - (Line.Width * TextItem.Scale.X)) * 0.5f;
+				break;
+
+			case ETextJustify::Right:
+			case ETextJustify::InvariantRight:
+				TextItem.Position.X = RenderTarget->SizeX - (Line.Width * TextItem.Scale.X) - RightOffset;
+				break;
+		}
+
+		TextItem.Text = FText::FromString(Line.Line);
+		RenderCanvas.DrawItem(TextItem);
+		TextItem.Position.Y += ActualLineHeight;
+	}
+
+	RenderCanvas.Flush_GameThread();
+
+	RenderTarget->UpdateResourceImmediate(false);
+}
+
+FCanvasTextItem UDMMaterialValueText::CreateTextItem(const FVector2D& InPosition, const FText& InText) const
+{
+	FCanvasTextItem TextItem(InPosition, InText, Font, Color);
+
+	TextItem.BlendMode = SE_BLEND_Translucent; // Must be translucent for fonts!
+	TextItem.HorizSpacingAdjust = Kerning;
+
+	if (bOutline)
+	{
+		TextItem.bOutlined = true;
+		TextItem.OutlineColor = OutlineColor;
+	}
+	else
+	{
+		TextItem.bOutlined = false;
+	}
+
+	if (bShadow)
+	{
+		TextItem.EnableShadow(ShadowColor, ShadowOffset);
+	}
+	else
+	{
+		TextItem.DisableShadow();
+	}
+
+	if (bGlow)
+	{
+		TextItem.FontRenderInfo.GlowInfo.bEnableGlow = true;
+		TextItem.FontRenderInfo.GlowInfo.GlowColor = GlowColor;
+		TextItem.FontRenderInfo.GlowInfo.GlowInnerRadius = GlowInnerRadius;
+		TextItem.FontRenderInfo.GlowInfo.GlowOuterRadius = GlowOuterRadius;
+	}
+	else
+	{
+		TextItem.FontRenderInfo.GlowInfo.bEnableGlow = false;
+	}
+
+	return TextItem;
+}
+
+#undef LOCTEXT_NAMESPACE
diff --git a/Engine/Plugins/Experimental/DynamicMaterial/Source/DynamicMaterial/Private/DMValueDefinition.cpp b/Engine/Plugins/Experimental/DynamicMaterial/Source/DynamicMaterial/Private/DMValueDefinition.cpp
--- a/Engine/Plugins/Experimental/DynamicMaterial/Source/DynamicMaterial/Private/DMValueDefinition.cpp
+++ b/Engine/Plugins/Experimental/DynamicMaterial/Source/DynamicMaterial/Private/DMValueDefinition.cpp
@@ -16,7 +16,8 @@ namespace UE::MaterialDesigner::Private
 			EDMValueType::VT_Float3_XYZ,
 			EDMValueType::VT_Float4_RGBA,
 			EDMValueType::VT_Texture,
-			EDMValueType::VT_ColorAtlas
+			EDMValueType::VT_ColorAtlas,
+			EDMValueType::VT_Text
 	};
 
 	const TMap<EDMValueType, FDMValueDefinition> TypeDefinitions = {
@@ -94,7 +95,16 @@ namespace UE::MaterialDesigner::Private
 			{EDMValueType::VT_ColorAtlas,
 			4,
 			LOCTEXT("ColorAtlas", "Color Atlas"),
-			{LOCTEXT("Alpha", "Alpha")}}
+			{LOCTEXT("Red", "Red"),
+				LOCTEXT("Green", "Green"),
+				LOCTEXT("Blue", "Blue"),
+				LOCTEXT("Alpha", "Alpha")}}
+		},
+		{EDMValueType::VT_Text,
+			{EDMValueType::VT_Text,
+			0,
+			LOCTEXT("Text", "Text"),
+			{}}
 		}
 	};
 }
diff --git a/Engine/Plugins/Experimental/DynamicMaterial/Source/DynamicMaterial/Public/Components/DMMaterialValue.h b/Engine/Plugins/Experimental/DynamicMaterial/Source/DynamicMaterial/Public/Components/DMMaterialValue.h
--- a/Engine/Plugins/Experimental/DynamicMaterial/Source/DynamicMaterial/Public/Components/DMMaterialValue.h
+++ b/Engine/Plugins/Experimental/DynamicMaterial/Source/DynamicMaterial/Public/Components/DMMaterialValue.h
@@ -122,6 +122,9 @@ public:
 
 	/** Return true if, when setting the base stage, the same value should be applied to the mask stage. */
 	virtual bool IsWholeLayerValue() const { return false; }
+
+	/** Returns true if the property row generator should expose the value property. */
+	virtual bool AllowEditValue() const { return true; }
 #endif
  
 	virtual void SetMIDParameter(UMaterialInstanceDynamic* InMID) const;
diff --git a/Engine/Plugins/Experimental/DynamicMaterial/Source/DynamicMaterial/Public/Components/MaterialValues/DMMaterialValueText.h b/Engine/Plugins/Experimental/DynamicMaterial/Source/DynamicMaterial/Public/Components/MaterialValues/DMMaterialValueText.h
new file mode 100644
--- /dev/null
+++ b/Engine/Plugins/Experimental/DynamicMaterial/Source/DynamicMaterial/Public/Components/MaterialValues/DMMaterialValueText.h
@@ -0,0 +1,240 @@
+// Copyright Epic Games, Inc. All Rights Reserved.
+
+#pragma once
+
+#include "Components/MaterialValues/DMMaterialValueTexture.h"
+#include "Delegates/IDelegateInstance.h"
+#include "Framework/Text/TextLayout.h"
+#include "Internationalization/Text.h"
+#include "Math/Color.h"
+#include "Math/Vector2D.h"
+#include "DMMaterialValueText.generated.h"
+
+class FCanvasTextItem;
+class UCanvas;
+class UDMMaterialValueTexture;
+class UFont;
+class UTextureRenderTarget2D;
+
+USTRUCT()
+struct FDMTextLine
+{
+	GENERATED_BODY()
+
+	UPROPERTY()
+	FString Line;
+
+	UPROPERTY()
+	float Width;
+};
+
+UCLASS(BlueprintType, ClassGroup = "Material Designer")
+class DYNAMICMATERIAL_API UDMMaterialValueText : public UDMMaterialValueTexture
+{
+	GENERATED_BODY()
+
+	friend class SDMComponentEdit;
+	friend struct FDMMaterialValueText;
+
+public:
+	UDMMaterialValueText();
+
+#if WITH_EDITOR
+	//~ Begin UDMMaterialValue
+	virtual bool AllowEditValue() const override { return false; }
+	//~ End UDMMaterialValue
+#endif
+
+	//~ Begin UObject
+	virtual void PostLoad() override;
+#if WITH_EDITOR
+	virtual void PostEditChangeProperty(FPropertyChangedEvent& InPropertyChangedEvent) override;
+#endif
+	//~ End UObject
+
+	UFont* GetFont() const;
+	void SetFont(UFont* InFont);
+
+	const FText& GetText() const;
+	void SetText(const FText& InText);
+
+	const FLinearColor& GetColor() const;
+	void SetColor(const FLinearColor& InColor);
+
+	const FLinearColor& GetBackgroundColor() const;
+	void SetBackgroundColor(const FLinearColor& InBackgroundColor);
+
+	ETextJustify::Type GetJustify() const;
+	void SetJustify(ETextJustify::Type InJustify);
+
+	float GetKerning() const;
+	void SetKerning(float InKerning);
+
+	float GetLineHeight() const;
+	void SetLineHeight(float InLineHeight);
+
+	float GetPaddingLeft() const;
+	void SetPaddingLeft(float InPadding);
+
+	float GetPaddingRight() const;
+	void SetPaddingRight(float InPadding);
+
+	float GetPaddingTop() const;
+	void SetPaddingTop(float InPadding);
+
+	float GetPaddingBottom() const;
+	void SetPaddingBottom(float InPadding);
+
+	const FIntPoint& GetTextScale() const;
+	void SetTextScale(const FIntPoint& InTextScale);
+
+	const FIntPoint& GetTextureSizeOverride() const;
+	void SetTextureSizeOverride(const FIntPoint& InTextureSizeOverride);
+
+	bool GetHasOutline() const;
+	void SetHasOutline(bool bInHasOutline);
+
+	const FLinearColor& GetOutlineColor() const;
+	void SetOutlineColor(const FLinearColor& InOutlineColor);
+
+	bool GetHasShadow() const;
+	void SetHasShadow(bool bInHasShadow);
+
+	const FLinearColor& GetShadowColor() const;
+	void SetShadowColor(const FLinearColor& InShadowColor);
+
+	const FVector2D& GetShadowOffset() const;
+	void SetShadowOffset(const FVector2D& InShadowOffset);
+
+	bool GetHasGlow() const;
+	void SetHasGlow(bool bInHasGlow);
+
+	const FLinearColor& GetGlowColor() const;
+	void SetGlowColor(const FLinearColor& InGlowColor);
+
+	const FVector2D& GetGlowInnerRadius() const;
+	void SetGlowInnerRadius(const FVector2D& InGlowInnerRadius);
+
+	const FVector2D& GetGlowOuterRadius() const;
+	void SetGlowOuterRadius(const FVector2D& InGlowOuterRadius);
+	
+protected:
+	UPROPERTY(EditInstanceOnly, BlueprintReadOnly, Setter = SetFont, Category = "Material Designer",
+		meta = (NotKeyframeable, AllowPrivateAccess = "true"))
+	UFont* Font = nullptr;
+
+	UPROPERTY(EditInstanceOnly, BlueprintReadOnly, Setter = SetText, Category = "Material Designer",
+		meta = (MultiLine, NotKeyframeable, AllowPrivateAccess = "true"))
+	FText Text = FText::GetEmpty();
+
+	UPROPERTY()
+	TArray<FDMTextLine> Lines = {};
+
+	UPROPERTY(EditInstanceOnly, BlueprintReadOnly, Setter = SetColor, Category = "Material Designer",
+		meta = (AllowPrivateAccess = "true"))
+	FLinearColor Color = FLinearColor::White;
+
+	/** Clear color of the render target. */
+	UPROPERTY(EditInstanceOnly, BlueprintReadOnly, Setter = SetBackgroundColor, Category = "Material Designer",
+		meta = (AllowPrivateAccess = "true"))
+	FLinearColor BackgroundColor = FLinearColor::Black;
+
+	UPROPERTY(EditInstanceOnly, BlueprintReadOnly, Setter = SetJustify, Category = "Material Designer",
+		meta = (NotKeyframeable, AllowPrivateAccess = "true"))
+	TEnumAsByte<ETextJustify::Type> Justify = ETextJustify::Left;
+
+	UPROPERTY(EditInstanceOnly, BlueprintReadOnly, Setter = SetKerning, Category = "Material Designer",
+		meta = (AllowPrivateAccess = "true"))
+	float Kerning = 0.f;
+
+	/** Multiplier on the base font height. */
+	UPROPERTY(EditInstanceOnly, BlueprintReadOnly, Setter = SetLineHeight, Category = "Material Designer",
+		meta = (NotKeyframeable, AllowPrivateAccess = "true", UIMin = 0.1, UIMax = 2.0))
+	float LineHeight = 1.f;
+
+	/** Extra space adding beyond the edge of the glyphs. Useful for shadows, glows, etc. */
+	UPROPERTY(EditInstanceOnly, BlueprintReadOnly, Setter = SetPaddingLeft, Category = "Material Designer",
+		meta = (NotKeyframeable, AllowPrivateAccess = "true"))
+	float PaddingLeft = 0.f;
+
+	/** Extra space adding beyond the edge of the glyphs. Useful for shadows, glows, etc. */
+	UPROPERTY(EditInstanceOnly, BlueprintReadOnly, Setter = SetPaddingRight, Category = "Material Designer",
+		meta = (NotKeyframeable, AllowPrivateAccess = "true"))
+	float PaddingRight = 0.f;
+
+	/** Extra space adding beyond the edge of the glyphs. Useful for shadows, glows, etc. */
+	UPROPERTY(EditInstanceOnly, BlueprintReadOnly, Setter = SetPaddingTop, Category = "Material Designer",
+		meta = (NotKeyframeable, AllowPrivateAccess = "true"))
+	float PaddingTop = 0.f;
+
+	/** Extra space adding beyond the edge of the glyphs. Useful for shadows, glows, etc. */
+	UPROPERTY(EditInstanceOnly, BlueprintReadOnly, Setter = SetPaddingBottom, Category = "Material Designer",
+		meta = (NotKeyframeable, AllowPrivateAccess = "true"))
+	float PaddingBottom = 0.f;
+
+	UPROPERTY(EditInstanceOnly, BlueprintReadOnly, Setter = SetTextScale, Category = "Material Designer",
+		meta = (NotKeyframeable, AllowPrivateAccess = "true", ClampMin = 1, UIMin = 1, ClampMax = 64, UIMax = 64,
+			ToolTip = "Increasing this will increase the resolution of the text."))
+	FIntPoint TextScale = FIntPoint(10, 10);
+
+	UPROPERTY(EditInstanceOnly, BlueprintReadOnly, Setter = SetTextureSizeOverride, Category = "Material Designer",
+		meta = (NotKeyframeable, AllowPrivateAccess = "true", ToolTip = "Override the size of the render target. Set to 0 to disable."))
+	FIntPoint TextureSizeOverride = FIntPoint(0, 0);
+
+	UPROPERTY(EditInstanceOnly, BlueprintReadOnly, Setter = SetHasOutline, Category = "Material Designer",
+		meta = (AllowPrivateAccess = "true"))
+	bool bOutline = false;
+
+	UPROPERTY(EditInstanceOnly, BlueprintReadOnly, Setter = SetOutlineColor, Category = "Material Designer",
+		meta = (AllowPrivateAccess = "true", EditCondition = "bOutline", EditConditionHides = true))
+	FLinearColor OutlineColor = FLinearColor::Black;
+
+	UPROPERTY(EditInstanceOnly, BlueprintReadOnly, Setter = SetHasShadow, Category = "Material Designer",
+		meta = (AllowPrivateAccess = "true"))
+	bool bShadow = false;
+
+	UPROPERTY(EditInstanceOnly, BlueprintReadOnly, Setter = SetShadowColor, Category = "Material Designer",
+		meta = (AllowPrivateAccess = "true", EditCondition = "bShadow", EditConditionHides = true))
+	FLinearColor ShadowColor = FLinearColor::Black;
+
+	UPROPERTY(EditInstanceOnly, BlueprintReadOnly, Setter = SetShadowOffset, Category = "Material Designer",
+		meta = (AllowPrivateAccess = "true", EditCondition="bShadow", EditConditionHides = true))
+	FVector2D ShadowOffset = FVector2D(1.0, 1.0);
+
+	UPROPERTY(EditInstanceOnly, BlueprintReadOnly, Setter = SetHasGlow, Category = "Material Designer",
+		meta = (AllowPrivateAccess = "true"))
+	bool bGlow = false;
+
+	UPROPERTY(EditInstanceOnly, BlueprintReadOnly, Setter = SetGlowColor, Category = "Material Designer",
+		meta = (AllowPrivateAccess = "true", EditCondition = "bGlow", EditConditionHides = true))
+	FLinearColor GlowColor = FLinearColor::White;
+
+	UPROPERTY(EditInstanceOnly, BlueprintReadOnly, Setter = SetGlowInnerRadius, Category = "Material Designer",
+		meta = (AllowPrivateAccess = "true", EditCondition = "bGlow", EditConditionHides = true, ClampMin = 0, UIMin = 0, ClampMax = 1, UIMax = 1))
+	FVector2D GlowInnerRadius = FVector2D(0.5, 0.5);
+
+	UPROPERTY(EditInstanceOnly, BlueprintReadOnly, Setter = SetGlowOuterRadius, Category = "Material Designer",
+		meta = (AllowPrivateAccess = "true", EditCondition = "bGlow", EditConditionHides = true, ClampMin = 0, UIMin = 0, ClampMax = 1, UIMax = 1))
+	FVector2D GlowOuterRadius = FVector2D(0.5, 0.5);
+
+	UPROPERTY(Instanced)
+	TObjectPtr<UTextureRenderTarget2D> RenderTarget = nullptr;
+
+	FDelegateHandle EndOfFrameDelegateHandle;
+
+	bool bNeedsResize = false;
+
+	static UTextureRenderTarget2D* CreateRenderTarget(UObject* InOuter, const FIntPoint& InSize, const FLinearColor& InClearColor);
+
+	void CalculateLineSizes();
+
+	FIntPoint GetRequiredTextureSize() const;
+
+	void UpdateTextLines();
+
+	void RequestUpdateTextTexture(bool bInNeedsResize);
+
+	void UpdateTextTexture();
+
+	FCanvasTextItem CreateTextItem(const FVector2D& InPosition, const FText& InText) const;
+};
diff --git a/Engine/Plugins/Experimental/DynamicMaterial/Source/DynamicMaterial/Public/Components/MaterialValues/DMMaterialValueTexture.h b/Engine/Plugins/Experimental/DynamicMaterial/Source/DynamicMaterial/Public/Components/MaterialValues/DMMaterialValueTexture.h
--- a/Engine/Plugins/Experimental/DynamicMaterial/Source/DynamicMaterial/Public/Components/MaterialValues/DMMaterialValueTexture.h
+++ b/Engine/Plugins/Experimental/DynamicMaterial/Source/DynamicMaterial/Public/Components/MaterialValues/DMMaterialValueTexture.h
@@ -76,9 +76,7 @@ protected:
 	UPROPERTY(EditAnywhere, BlueprintReadWrite, BlueprintSetter = SetDefaultValue, Category = "Material Designer",
 		meta = (AllowPrivateAccess = "true"))
 	TObjectPtr<UTexture> DefaultValue;
-#endif
 
-#if WITH_EDITORONLY_DATA
 	UPROPERTY()
 	TObjectPtr<UTexture> OldValue;
 #endif
diff --git a/Engine/Plugins/Experimental/DynamicMaterial/Source/DynamicMaterial/Public/DMDefs.h b/Engine/Plugins/Experimental/DynamicMaterial/Source/DynamicMaterial/Public/DMDefs.h
--- a/Engine/Plugins/Experimental/DynamicMaterial/Source/DynamicMaterial/Public/DMDefs.h
+++ b/Engine/Plugins/Experimental/DynamicMaterial/Source/DynamicMaterial/Public/DMDefs.h
@@ -48,7 +48,9 @@ enum class EDMValueType : uint8
 	VT_Float4_RGBA,
 	VT_Float_Any,
 	VT_Texture,
-	VT_ColorAtlas
+	VT_ColorAtlas,
+	VT_Text,
+	VT_MAX
 };
 
 UENUM(BlueprintType)
diff --git a/Engine/Plugins/Experimental/DynamicMaterial/Source/DynamicMaterialEditor/Private/Components/DMMaterialStage.cpp b/Engine/Plugins/Experimental/DynamicMaterial/Source/DynamicMaterialEditor/Private/Components/DMMaterialStage.cpp
--- a/Engine/Plugins/Experimental/DynamicMaterial/Source/DynamicMaterialEditor/Private/Components/DMMaterialStage.cpp
+++ b/Engine/Plugins/Experimental/DynamicMaterial/Source/DynamicMaterialEditor/Private/Components/DMMaterialStage.cpp
@@ -555,7 +555,11 @@ void UDMMaterialStage::GenerateExpressions(const TSharedRef<FDMMaterialBuildStat
 		return;
 	}
 
-	check(Source);
+	if (!IsValid(Source))
+	{
+		UE_LOG(LogDynamicMaterialEditor, Warning, TEXT("Stage with no expressions attempted to generate material expressions."));
+		return;
+	}
 
 	if (InBuildState->HasStage(this))
 	{
diff --git a/Engine/Plugins/Experimental/DynamicMaterial/Source/DynamicMaterialEditor/Private/Components/DMMaterialStageThroughput.cpp b/Engine/Plugins/Experimental/DynamicMaterial/Source/DynamicMaterialEditor/Private/Components/DMMaterialStageThroughput.cpp
--- a/Engine/Plugins/Experimental/DynamicMaterial/Source/DynamicMaterialEditor/Private/Components/DMMaterialStageThroughput.cpp
+++ b/Engine/Plugins/Experimental/DynamicMaterial/Source/DynamicMaterialEditor/Private/Components/DMMaterialStageThroughput.cpp
@@ -461,6 +461,7 @@ void UDMMaterialStageThroughput::AddDefaultInput(int32 InInputIndex) const
 		case EDMValueType::VT_Float3_RPY:
 		case EDMValueType::VT_Float3_XYZ:
 		case EDMValueType::VT_Texture:
+		case EDMValueType::VT_Text:
 			UDMMaterialStageInputValue::ChangeStageInput_NewLocalValue(
 				Stage, 
 				InInputIndex,
diff --git a/Engine/Plugins/Experimental/DynamicMaterial/Source/DynamicMaterialEditor/Private/Components/MaterialStageExpressions/DMMSEText.cpp b/Engine/Plugins/Experimental/DynamicMaterial/Source/DynamicMaterialEditor/Private/Components/MaterialStageExpressions/DMMSEText.cpp
deleted file mode 100644
diff --git a/Engine/Plugins/Experimental/DynamicMaterial/Source/DynamicMaterialEditor/Private/Components/MaterialStageExpressions/DMMSETextRenderer.cpp b/Engine/Plugins/Experimental/DynamicMaterial/Source/DynamicMaterialEditor/Private/Components/MaterialStageExpressions/DMMSETextRenderer.cpp
new file mode 100644
--- /dev/null
+++ b/Engine/Plugins/Experimental/DynamicMaterial/Source/DynamicMaterialEditor/Private/Components/MaterialStageExpressions/DMMSETextRenderer.cpp
@@ -0,0 +1,19 @@
+// Copyright Epic Games, Inc. All Rights Reserved.
+
+#include "Components/MaterialStageExpressions/DMMSETextRenderer.h"
+#include "Materials/MaterialExpressionTextureSample.h"
+
+#define LOCTEXT_NAMESPACE "DMMaterialStageExpressionTextRenderer"
+
+UDMMaterialStageExpressionTextRenderer::UDMMaterialStageExpressionTextRenderer()
+	: UDMMaterialStageExpressionTextureSampleBase(
+		LOCTEXT("Text", "Text"),
+		UMaterialExpressionTextureSample::StaticClass()
+	)
+{
+	Menus.Add(EDMExpressionMenu::Texture);
+
+	InputConnectors[0] = {1, LOCTEXT("Text", "Text"), EDMValueType::VT_Text};
+}
+
+#undef LOCTEXT_NAMESPACE
diff --git a/Engine/Plugins/Experimental/DynamicMaterial/Source/DynamicMaterialEditor/Private/Menus/DMMaterialSlotLayerMenus.cpp b/Engine/Plugins/Experimental/DynamicMaterial/Source/DynamicMaterialEditor/Private/Menus/DMMaterialSlotLayerMenus.cpp
--- a/Engine/Plugins/Experimental/DynamicMaterial/Source/DynamicMaterialEditor/Private/Menus/DMMaterialSlotLayerMenus.cpp
+++ b/Engine/Plugins/Experimental/DynamicMaterial/Source/DynamicMaterialEditor/Private/Menus/DMMaterialSlotLayerMenus.cpp
@@ -7,6 +7,7 @@
 #include "Components/DMMaterialStageExpression.h"
 #include "Components/DMMaterialStageGradient.h"
 #include "Components/DMMaterialValue.h"
+#include "Components/MaterialStageExpressions/DMMSETextRenderer.h"
 #include "Components/MaterialStageExpressions/DMMSETextureSample.h"
 #include "Components/MaterialStageExpressions/DMMSETextureSampleEdgeColor.h"
 #include "DMDefs.h"
@@ -537,6 +538,19 @@ void FDMMaterialSlotLayerMenus::AddAddLayerSection(UToolMenu* InMenu)
 		))
 	);
 
+	NewSection.AddMenuEntry(
+		NAME_None,
+		LOCTEXT("AddText", "Text"),
+		LOCTEXT("AddTextTooltip", "Add a Material Stage based on a Text Renderer."),
+		FSlateIcon(),
+		FUIAction(FExecuteAction::CreateSP(
+			SlotWidget.Get(),
+			&SDMSlot::AddNewLayer_Expression,
+			TSubclassOf<UDMMaterialStageExpression>(UDMMaterialStageExpressionTextRenderer::StaticClass()),
+			EDMMaterialLayerStage::Base
+		))
+	);
+
 	if constexpr (UE::DynamicMaterialEditor::bAdvancedSlotsEnabled)
 	{
 		if (bHasValidSlot)
diff --git a/Engine/Plugins/Experimental/DynamicMaterial/Source/DynamicMaterialEditor/Private/Menus/DMMaterialStageSourceMenus.cpp b/Engine/Plugins/Experimental/DynamicMaterial/Source/DynamicMaterialEditor/Private/Menus/DMMaterialStageSourceMenus.cpp
--- a/Engine/Plugins/Experimental/DynamicMaterial/Source/DynamicMaterialEditor/Private/Menus/DMMaterialStageSourceMenus.cpp
+++ b/Engine/Plugins/Experimental/DynamicMaterial/Source/DynamicMaterialEditor/Private/Menus/DMMaterialStageSourceMenus.cpp
@@ -14,7 +14,7 @@
 #include "Components/DMMaterialStageThroughputLayerBlend.h"
 #include "Components/DMMaterialValue.h"
 #include "Components/MaterialStageExpressions/DMMSESceneTexture.h"
-#include "Components/MaterialStageExpressions/DMMSEText.h"
+#include "Components/MaterialStageExpressions/DMMSETextRenderer.h"
 #include "Components/MaterialStageExpressions/DMMSETextureSample.h"
 #include "Components/MaterialStageExpressions/DMMSETextureSampleEdgeColor.h"
 #include "Components/MaterialStageInputs/DMMSIExpression.h"
@@ -1201,7 +1201,7 @@ namespace UE::DynamicMaterialEditor::Private
 
 			UDMMaterialStageInputExpression::ChangeStageInput_Expression(
 				Stage,
-				UDMMaterialStageExpressionText::StaticClass(),
+				UDMMaterialStageExpressionTextRenderer::StaticClass(),
 				UDMMaterialStageBlend::InputB,
 				FDMMaterialStageConnectorChannel::WHOLE_CHANNEL,
 				0,
@@ -1215,21 +1215,21 @@ namespace UE::DynamicMaterialEditor::Private
 
 			UDMMaterialStageInputExpression::ChangeStageInput_Expression(
 				Stage,
-				UDMMaterialStageExpressionText::StaticClass(),
+				UDMMaterialStageExpressionTextRenderer::StaticClass(),
 				UDMMaterialStageThroughputLayerBlend::InputMaskSource,
 				FDMMaterialStageConnectorChannel::WHOLE_CHANNEL,
 				0,
-				FDMMaterialStageConnectorChannel::FIRST_CHANNEL
+				FDMMaterialStageConnectorChannel::WHOLE_CHANNEL
 			);
 		}
 		else
 		{
-			FScopedTransaction Transaction(LOCTEXT("SetStageInput", "Set Material Designer Source"));
+			FScopedTransaction Transaction(LOCTEXT("SetStageSource", "Set Material Designer Stage Source"));
 			Stage->Modify();
 
 			UDMMaterialStageInputExpression::ChangeStageSource_Expression(
 				Stage,
-				UDMMaterialStageExpressionText::StaticClass()
+				UDMMaterialStageExpressionTextRenderer::StaticClass()
 			);
 		}
 
@@ -1353,6 +1353,18 @@ namespace UE::DynamicMaterialEditor::Private
 			)
 		);
 
+		InSection.AddMenuEntry("Text",
+			LOCTEXT("ChangeSourceText", "Text"),
+			LOCTEXT("ChangeSourceTextTooltip", "Change the source of this stage to a Text Renderer."),
+			FSlateIcon(),
+			FUIAction(
+				FExecuteAction::CreateStatic(
+					&ChangeSourceToTextFromContext,
+					MenuContext
+				)
+			)
+		);
+
 		if constexpr (UE::DynamicMaterialEditor::bGlobalValuesEnabled)
 		{
 			const TArray<UDMMaterialValue*>& Values = MaterialModel->GetValues();
@@ -1400,18 +1412,6 @@ namespace UE::DynamicMaterialEditor::Private
 			)
 		);
 
-		InSection.AddMenuEntry("Text",
-			LOCTEXT("ChangeSourceText", "Text"),
-			LOCTEXT("ChangeSourceTextTooltip", "Change the source of this stage to a Text Renderer."),
-			FSlateIcon(),
-			FUIAction(
-				FExecuteAction::CreateStatic(
-					&ChangeSourceToTextFromContext,
-					MenuContext
-				)
-			)
-		);
-
 		if constexpr (UE::DynamicMaterialEditor::bAdvancedSlotsEnabled)
 		{
 			if (bHasValidSlot)
diff --git a/Engine/Plugins/Experimental/DynamicMaterial/Source/DynamicMaterialEditor/Private/Slate/Properties/Generators/DMMaterialValuePropertyRowGenerator.cpp b/Engine/Plugins/Experimental/DynamicMaterial/Source/DynamicMaterialEditor/Private/Slate/Properties/Generators/DMMaterialValuePropertyRowGenerator.cpp
--- a/Engine/Plugins/Experimental/DynamicMaterial/Source/DynamicMaterialEditor/Private/Slate/Properties/Generators/DMMaterialValuePropertyRowGenerator.cpp
+++ b/Engine/Plugins/Experimental/DynamicMaterial/Source/DynamicMaterialEditor/Private/Slate/Properties/Generators/DMMaterialValuePropertyRowGenerator.cpp
@@ -43,30 +43,32 @@ void FDMMaterialValuePropertyRowGenerator::AddComponentProperties(const TSharedR
 
 	InOutProcessedObjects.Add(InComponent);
 
-	FDMPropertyHandle Handle = SDMEditor::GetPropertyHandle(&*InComponentEditWidget, Value, UDMMaterialValue::ValueName);
-
-	Handle.ResetToDefaultOverride = FResetToDefaultOverride::Create(
-		FIsResetToDefaultVisible::CreateUObject(Value, &UDMMaterialValue::CanResetToDefault),
-		FResetToDefaultHandler::CreateUObject(Value, &UDMMaterialValue::ResetToDefault)
-	);
-
-	if (UDMMaterialValueFloat* FloatValue = Cast<UDMMaterialValueFloat>(Value))
+	if (Value->AllowEditValue())
 	{
-		if (FloatValue->HasValueRange())
+		FDMPropertyHandle Handle = SDMEditor::GetPropertyHandle(&*InComponentEditWidget, Value, UDMMaterialValue::ValueName);
+
+		Handle.ResetToDefaultOverride = FResetToDefaultOverride::Create(
+			FIsResetToDefaultVisible::CreateUObject(Value, &UDMMaterialValue::CanResetToDefault),
+			FResetToDefaultHandler::CreateUObject(Value, &UDMMaterialValue::ResetToDefault)
+		);
+		if (UDMMaterialValueFloat* FloatValue = Cast<UDMMaterialValueFloat>(Value))
 		{
-			static const FName UIMin = FName("UIMin");
-			static const FName UIMax = FName("UIMax");
-			static const FName ClampMin = FName("ClampMin");
-			static const FName ClampMax = FName("ClampMax");
+			if (FloatValue->HasValueRange())
+			{
+				static const FName UIMin = FName("UIMin");
+				static const FName UIMax = FName("UIMax");
+				static const FName ClampMin = FName("ClampMin");
+				static const FName ClampMax = FName("ClampMax");
 
-			Handle.PropertyHandle->SetInstanceMetaData(UIMin, FString::SanitizeFloat(FloatValue->GetValueRange().Min));
-			Handle.PropertyHandle->SetInstanceMetaData(ClampMin, FString::SanitizeFloat(FloatValue->GetValueRange().Min));
-			Handle.PropertyHandle->SetInstanceMetaData(UIMax, FString::SanitizeFloat(FloatValue->GetValueRange().Max));
-			Handle.PropertyHandle->SetInstanceMetaData(ClampMax, FString::SanitizeFloat(FloatValue->GetValueRange().Max));
+				Handle.PropertyHandle->SetInstanceMetaData(UIMin, FString::SanitizeFloat(FloatValue->GetValueRange().Min));
+				Handle.PropertyHandle->SetInstanceMetaData(ClampMin, FString::SanitizeFloat(FloatValue->GetValueRange().Min));
+				Handle.PropertyHandle->SetInstanceMetaData(UIMax, FString::SanitizeFloat(FloatValue->GetValueRange().Max));
+				Handle.PropertyHandle->SetInstanceMetaData(ClampMax, FString::SanitizeFloat(FloatValue->GetValueRange().Max));
+			}
 		}
-	}
 
-	InOutPropertyRows.Add(Handle);
+		InOutPropertyRows.Add(Handle);
+	}
 
 	const TArray<FName>& Properties = Value->GetEditableProperties();
 
diff --git a/Engine/Plugins/Experimental/DynamicMaterial/Source/DynamicMaterialEditor/Public/Components/MaterialStageExpressions/DMMSEText.h b/Engine/Plugins/Experimental/DynamicMaterial/Source/DynamicMaterialEditor/Public/Components/MaterialStageExpressions/DMMSEText.h
deleted file mode 100644
diff --git a/Engine/Plugins/Experimental/DynamicMaterial/Source/DynamicMaterialEditor/Public/Components/MaterialStageExpressions/DMMSETextRenderer.h b/Engine/Plugins/Experimental/DynamicMaterial/Source/DynamicMaterialEditor/Public/Components/MaterialStageExpressions/DMMSETextRenderer.h
new file mode 100644
--- /dev/null
+++ b/Engine/Plugins/Experimental/DynamicMaterial/Source/DynamicMaterialEditor/Public/Components/MaterialStageExpressions/DMMSETextRenderer.h
@@ -0,0 +1,17 @@
+// Copyright Epic Games, Inc. All Rights Reserved.
+
+#pragma once
+
+#include "Components/MaterialStageExpressions/DMMSETextureSampleBase.h"
+#include "DMMSETextRenderer.generated.h"
+
+UCLASS(BlueprintType, ClassGroup = "Material Designer")
+class DYNAMICMATERIALEDITOR_API UDMMaterialStageExpressionTextRenderer : public UDMMaterialStageExpressionTextureSampleBase
+{
+	GENERATED_BODY()
+
+	friend class SDMComponentEdit;
+
+public:
+	UDMMaterialStageExpressionTextRenderer();
+};
