Material Designer
- Save Level As now unloads the current asset and loads the new one when the map changes.

diff --git a/Engine/Plugins/Experimental/DynamicMaterial/Source/DynamicMaterialEditor/Private/DynamicMaterialEditorModule.cpp b/Engine/Plugins/Experimental/DynamicMaterial/Source/DynamicMaterialEditor/Private/DynamicMaterialEditorModule.cpp
--- a/Engine/Plugins/Experimental/DynamicMaterial/Source/DynamicMaterialEditor/Private/DynamicMaterialEditorModule.cpp
+++ b/Engine/Plugins/Experimental/DynamicMaterial/Source/DynamicMaterialEditor/Private/DynamicMaterialEditorModule.cpp
@@ -438,6 +438,43 @@ void FDynamicMaterialEditorModule::AddBuildRequest(UObject* InToBuild, bool bInD
 	UDMMaterialComponent::PreventClean(VeryShortTime);
 }
 
+void FDynamicMaterialEditorModule::RemoveBuildRequest(UObject* InToNotBuild)
+{
+	if (!InToNotBuild)
+	{
+		return;
+	}
+
+	FDynamicMaterialEditorModule::Get().BuildRequestList.Remove({InToNotBuild->GetPathName(), false});
+}
+
+void FDynamicMaterialEditorModule::RemoveBuildRequestForOuter(UObject* InOuter)
+{
+	if (!InOuter)
+	{
+		return;
+	}
+
+	const FString ObjectPath = InOuter->GetPathName();
+	const int32 ObjectPathLength = ObjectPath.Len();
+
+	for (TSet<FDMBuildRequestEntry>::TIterator Iter(FDynamicMaterialEditorModule::Get().BuildRequestList); Iter; ++Iter)
+	{
+		if (Iter->AssetPath.Len() > ObjectPathLength && Iter->AssetPath.StartsWith(ObjectPath))
+		{
+			// Make sure it's a path separator after the parent path.
+			switch (Iter->AssetPath[ObjectPathLength])
+			{
+				case '.':
+				case '/':
+				case ':':
+					Iter.RemoveCurrent();
+					break;
+			}
+		}
+	}
+}
+
 void FDynamicMaterialEditorModule::OpenEditor(UWorld* InWorld)
 {
 	if (!IsValid(InWorld))
diff --git a/Engine/Plugins/Experimental/DynamicMaterial/Source/DynamicMaterialEditor/Private/DynamicMaterialEditorModule.h b/Engine/Plugins/Experimental/DynamicMaterial/Source/DynamicMaterialEditor/Private/DynamicMaterialEditorModule.h
--- a/Engine/Plugins/Experimental/DynamicMaterial/Source/DynamicMaterialEditor/Private/DynamicMaterialEditorModule.h
+++ b/Engine/Plugins/Experimental/DynamicMaterial/Source/DynamicMaterialEditor/Private/DynamicMaterialEditorModule.h
@@ -111,6 +111,10 @@ public:
 
 	void AddBuildRequest(UObject* InToBuild, bool bInDirtyAssets);
 
+	void RemoveBuildRequest(UObject* InToNotBuild);
+
+	void RemoveBuildRequestForOuter(UObject* InOuter);
+
 	const TSharedRef<FUICommandList>& GetCommandList() const { return CommandList; }
 
 protected:
diff --git a/Engine/Plugins/Experimental/DynamicMaterial/Source/DynamicMaterialEditor/Private/LevelEditor/DMLevelEditorIntegration.cpp b/Engine/Plugins/Experimental/DynamicMaterial/Source/DynamicMaterialEditor/Private/LevelEditor/DMLevelEditorIntegration.cpp
--- a/Engine/Plugins/Experimental/DynamicMaterial/Source/DynamicMaterialEditor/Private/LevelEditor/DMLevelEditorIntegration.cpp
+++ b/Engine/Plugins/Experimental/DynamicMaterial/Source/DynamicMaterialEditor/Private/LevelEditor/DMLevelEditorIntegration.cpp
@@ -2,12 +2,17 @@
 
 #include "DMLevelEditorIntegration.h"
 #include "DMLevelEditorIntegrationInstance.h"
+#include "DynamicMaterialEditorModule.h"
 #include "LevelEditor.h"
+#include "Model/DynamicMaterialModel.h"
+#include "Model/DynamicMaterialModelEditorOnlyData.h"
 #include "Modules/ModuleManager.h"
+#include "Slate/SDMEditor.h"
 
 namespace UE::DynamicMaterialEditor::Private
 {
 	FDelegateHandle LevelEditorCreatedHandle;
+	FDelegateHandle LevelEditorMapChangeHandle;
 
 	FLevelEditorModule& GetLevelEditorModule()
 	{
@@ -23,13 +28,101 @@ namespace UE::DynamicMaterialEditor::Private
 	{
 		return FModuleManager::LoadModuleChecked<FLevelEditorModule>("LevelEditor");
 	}
+
+	void OnMapTearDown(UWorld* InWorld)
+	{
+		FDynamicMaterialEditorModule::Get().RemoveBuildRequestForOuter(InWorld);
+
+		FDMLevelEditorIntegrationInstance* Instance = FDMLevelEditorIntegrationInstance::GetMutableIntegrationForWorld(InWorld);
+
+		if (!Instance)
+		{
+			return;
+		}
+
+		Instance->SetLastAssetOpenPartialPath("");
+
+		TSharedPtr<SDMEditor> Editor = Instance->GetEditor();
+
+		if (!Editor.IsValid())
+		{
+			return;
+		}
+
+		UDynamicMaterialModel* MaterialModel = Editor->GetMaterialModel();
+
+		if (!MaterialModel)
+		{
+			return;
+		}
+
+		Editor->SetMaterialModel(nullptr);
+
+		const FString WorldPath = InWorld->GetPathName();
+		const int WorldPathLength = WorldPath.Len();
+		const FString ModelPath = MaterialModel->GetPathName();
+
+		if (ModelPath.Len() > WorldPathLength && ModelPath.StartsWith(WorldPath))
+		{
+			switch (ModelPath[WorldPathLength])
+			{
+				case '.':
+				case '/':
+				case ':':
+					Instance->SetLastAssetOpenPartialPath(ModelPath.RightChop(WorldPath.Len()));
+					break;
+			}
+		}
+	}
+
+	void OnMapLoad(UWorld* InWorld)
+	{
+		FDMLevelEditorIntegrationInstance* Instance = FDMLevelEditorIntegrationInstance::GetMutableIntegrationForWorld(InWorld);
+
+		if (!Instance)
+		{
+			return;
+		}
+
+		const FString PartialAssetPath = Instance->GetLastOpenAssetPartialPath();
+		Instance->SetLastAssetOpenPartialPath("");
+
+		if (PartialAssetPath.IsEmpty())
+		{
+			return;
+		}
+
+		TSharedPtr<SDMEditor> Editor = Instance->GetEditor();
+
+		if (!Editor.IsValid())
+		{
+			return;
+		}
+
+		const FString ModelPath = InWorld->GetPathName() + PartialAssetPath;
+
+		if (UObject* Object = FindObject<UObject>(nullptr, *ModelPath, false))
+		{
+			if (UDynamicMaterialModel* MaterialModel = Cast<UDynamicMaterialModel>(Object))
+			{
+				Editor->SetMaterialModel(MaterialModel);
+
+				if (UDynamicMaterialModelEditorOnlyData* EditorOnlyData = UDynamicMaterialModelEditorOnlyData::Get(MaterialModel))
+				{
+					FDynamicMaterialEditorModule::Get().AddBuildRequest(EditorOnlyData, false);
+				}
+			}
+		}
+	}
 }
 
 void FDMLevelEditorIntegration::Initialize()
 {
 	using namespace UE::DynamicMaterialEditor::Private;
 
-	LevelEditorCreatedHandle = LoadLevelEditorModuleChecked().OnLevelEditorCreated().AddLambda(
+	FLevelEditorModule& LevelEditorModule = LoadLevelEditorModuleChecked();
+
+	LevelEditorCreatedHandle = LevelEditorModule.OnLevelEditorCreated().AddLambda(
 		[](TSharedPtr<ILevelEditor> InLevelEditor)
 		{
 			if (InLevelEditor.IsValid())
@@ -38,6 +131,22 @@ void FDMLevelEditorIntegration::Initialize()
 			}
 		}
 	);
+
+	LevelEditorMapChangeHandle = LevelEditorModule.OnMapChanged().AddLambda(
+		[](UWorld* InWorld, EMapChangeType InMapChangeType)
+		{
+			switch (InMapChangeType)
+			{
+				case EMapChangeType::TearDownWorld:
+					OnMapTearDown(InWorld);
+					break;
+
+				case EMapChangeType::LoadMap:
+					OnMapLoad(InWorld);
+					break;
+			}
+		}
+	);
 }
 
 void FDMLevelEditorIntegration::Shutdown()
@@ -50,6 +159,9 @@ void FDMLevelEditorIntegration::Shutdown()
 		{
 			ModulePtr->OnLevelEditorCreated().Remove(LevelEditorCreatedHandle);
 			LevelEditorCreatedHandle.Reset();
+
+			ModulePtr->OnMapChanged().Remove(LevelEditorMapChangeHandle);
+			LevelEditorMapChangeHandle.Reset();
 		}
 	}
 
diff --git a/Engine/Plugins/Experimental/DynamicMaterial/Source/DynamicMaterialEditor/Private/LevelEditor/DMLevelEditorIntegrationInstance.cpp b/Engine/Plugins/Experimental/DynamicMaterial/Source/DynamicMaterialEditor/Private/LevelEditor/DMLevelEditorIntegrationInstance.cpp
--- a/Engine/Plugins/Experimental/DynamicMaterial/Source/DynamicMaterialEditor/Private/LevelEditor/DMLevelEditorIntegrationInstance.cpp
+++ b/Engine/Plugins/Experimental/DynamicMaterial/Source/DynamicMaterialEditor/Private/LevelEditor/DMLevelEditorIntegrationInstance.cpp
@@ -48,31 +48,7 @@ void FDMLevelEditorIntegrationInstance::RemoveIntegrations()
 
 const FDMLevelEditorIntegrationInstance* FDMLevelEditorIntegrationInstance::GetIntegrationForWorld(UWorld* InWorld)
 {
-	if (!IsValid(InWorld))
-	{
-		return nullptr;
-	}
-
-	ValidateInstances();
-
-	for (const FDMLevelEditorIntegrationInstance& Instance : Instances)
-	{
-		// Always return the first level editor integration for null words - they are assets.
-		if (!InWorld)
-		{
-			return &Instance;
-		}
-
-		if (TSharedPtr<ILevelEditor> LevelEditor = Instance.LevelEditorWeak.Pin())
-		{
-			if (LevelEditor->GetWorld() == InWorld)
-			{
-				return &Instance;
-			}
-		}
-	}
-
-	return nullptr;
+	return GetMutableIntegrationForWorld(InWorld);
 }
 
 FDMLevelEditorIntegrationInstance::~FDMLevelEditorIntegrationInstance()
@@ -118,7 +94,17 @@ TSharedPtr<SDockTab> FDMLevelEditorIntegrationInstance::InvokeTab() const
 	return Tab;
 }
 
-void FDMLevelEditorIntegrationInstance::FDMLevelEditorIntegrationInstance::ValidateInstances()
+const FString& FDMLevelEditorIntegrationInstance::GetLastOpenAssetPartialPath() const
+{
+	return LastOpenAssetPartialPath;
+}
+
+void FDMLevelEditorIntegrationInstance::SetLastAssetOpenPartialPath(const FString& InPath)
+{
+	LastOpenAssetPartialPath = InPath;
+}
+
+void FDMLevelEditorIntegrationInstance::ValidateInstances()
 {
 	for (int32 Index = 0; Index < Instances.Num(); ++Index)
 	{
@@ -130,6 +116,35 @@ void FDMLevelEditorIntegrationInstance::FDMLevelEditorIntegrationInstance::Valid
 	}
 }
 
+FDMLevelEditorIntegrationInstance* FDMLevelEditorIntegrationInstance::GetMutableIntegrationForWorld(UWorld* InWorld)
+{
+	if (!IsValid(InWorld))
+	{
+		return nullptr;
+	}
+
+	ValidateInstances();
+
+	for (FDMLevelEditorIntegrationInstance& Instance : Instances)
+	{
+		// Always return the first level editor integration for null words - they are assets.
+		if (!InWorld)
+		{
+			return &Instance;
+		}
+
+		if (TSharedPtr<ILevelEditor> LevelEditor = Instance.LevelEditorWeak.Pin())
+		{
+			if (LevelEditor->GetWorld() == InWorld)
+			{
+				return &Instance;
+			}
+		}
+	}
+
+	return nullptr;
+}
+
 FDMLevelEditorIntegrationInstance::FDMLevelEditorIntegrationInstance(const TSharedRef<ILevelEditor>& InLevelEditor)
 {
 	LevelEditorWeak = InLevelEditor;
diff --git a/Engine/Plugins/Experimental/DynamicMaterial/Source/DynamicMaterialEditor/Private/LevelEditor/DMLevelEditorIntegrationInstance.h b/Engine/Plugins/Experimental/DynamicMaterial/Source/DynamicMaterialEditor/Private/LevelEditor/DMLevelEditorIntegrationInstance.h
--- a/Engine/Plugins/Experimental/DynamicMaterial/Source/DynamicMaterialEditor/Private/LevelEditor/DMLevelEditorIntegrationInstance.h
+++ b/Engine/Plugins/Experimental/DynamicMaterial/Source/DynamicMaterialEditor/Private/LevelEditor/DMLevelEditorIntegrationInstance.h
@@ -2,6 +2,9 @@
 
 #pragma once
 
+#include "Containers/Array.h"
+#include "Containers/ArrayView.h"
+#include "Containers/UnrealString.h"
 #include "Templates/SharedPointer.h"
 #include "UObject/WeakObjectPtr.h"
 
@@ -22,12 +25,18 @@ public:
 
 	static const FDMLevelEditorIntegrationInstance* GetIntegrationForWorld(UWorld* InWorld);
 
+	static FDMLevelEditorIntegrationInstance* GetMutableIntegrationForWorld(UWorld* InWorld);
+
 	~FDMLevelEditorIntegrationInstance();
 
 	const TSharedPtr<SDMEditor>& GetEditor() const;
 
 	TSharedPtr<SDockTab> InvokeTab() const;
 
+	const FString& GetLastOpenAssetPartialPath() const;
+
+	void SetLastAssetOpenPartialPath(const FString& InPath);
+
 private:
 	static TArray<FDMLevelEditorIntegrationInstance, TInlineAllocator<1>> Instances;
 
@@ -37,6 +46,7 @@ private:
 	TWeakObjectPtr<UTypedElementSelectionSet> ActorSelectionSetWeak;
 	TWeakObjectPtr<UTypedElementSelectionSet> ObjectSelectionSetWeak;
 	TSharedPtr<SDMEditor> Editor;
+	FString LastOpenAssetPartialPath;
 
 	FDMLevelEditorIntegrationInstance(const TSharedRef<ILevelEditor>& InLevelEditor);
 
