Material Designer
- Unintialised USTRUCT property

diff --git a/Engine/Plugins/Experimental/DynamicMaterial/Source/DynamicMaterial/Public/Components/MaterialValues/DMMaterialValueText.h b/Engine/Plugins/Experimental/DynamicMaterial/Source/DynamicMaterial/Public/Components/MaterialValues/DMMaterialValueText.h
--- a/Engine/Plugins/Experimental/DynamicMaterial/Source/DynamicMaterial/Public/Components/MaterialValues/DMMaterialValueText.h
+++ b/Engine/Plugins/Experimental/DynamicMaterial/Source/DynamicMaterial/Public/Components/MaterialValues/DMMaterialValueText.h
@@ -25,7 +25,7 @@ struct FDMTextLine
 	FString Line;
 
 	UPROPERTY()
-	float Width;
+	float Width = 0.f;
 };
 
 UCLASS(BlueprintType, ClassGroup = "Material Designer")
