Material Designer
- Custom property rows can now be added by row generators.
- In-details panel visual UV manipulator.

diff --git a/Engine/Plugins/Experimental/DynamicMaterial/Source/DynamicMaterialEditor/Private/DynamicMaterialEditorSettings.cpp b/Engine/Plugins/Experimental/DynamicMaterial/Source/DynamicMaterialEditor/Private/DynamicMaterialEditorSettings.cpp
--- a/Engine/Plugins/Experimental/DynamicMaterial/Source/DynamicMaterialEditor/Private/DynamicMaterialEditorSettings.cpp
+++ b/Engine/Plugins/Experimental/DynamicMaterial/Source/DynamicMaterialEditor/Private/DynamicMaterialEditorSettings.cpp
@@ -99,6 +99,7 @@ void UDynamicMaterialEditorSettings::ResetAllLayoutSettings()
 	TooltipTextureSize = 512;
 	MaxFloatSliderWidth = 200.f;
 	bPreviewImagesUseTextureUVs = true;
+	bUVVisualizerVisible = true;
 }
 
 float UDynamicMaterialEditorSettings::GetSpinboxValueChangeMultiplier(const FModifierKeysState& InModifierKeys) const
diff --git a/Engine/Plugins/Experimental/DynamicMaterial/Source/DynamicMaterialEditor/Private/DynamicMaterialEditorSettings.h b/Engine/Plugins/Experimental/DynamicMaterial/Source/DynamicMaterialEditor/Private/DynamicMaterialEditorSettings.h
--- a/Engine/Plugins/Experimental/DynamicMaterial/Source/DynamicMaterialEditor/Private/DynamicMaterialEditorSettings.h
+++ b/Engine/Plugins/Experimental/DynamicMaterial/Source/DynamicMaterialEditor/Private/DynamicMaterialEditorSettings.h
@@ -79,6 +79,9 @@ public:
 		ClampMin = "100", UIMin = "100", ClampMax = "1000", UIMax = "1000"))
 	float MaxFloatSliderWidth;
 
+	UPROPERTY(Config, EditAnywhere, BlueprintReadWrite, Category = "Layout")
+	bool bUVVisualizerVisible;
+
 	/** The size of the material layer preview images. */
 	UPROPERTY(Config, EditAnywhere, BlueprintReadWrite, Category = "Preview", meta = (ClampMin = "32", UIMin = "32", ClampMax = "128", UIMax = "128"))
 	int32 LayerPreviewSize;
diff --git a/Engine/Plugins/Experimental/DynamicMaterial/Source/DynamicMaterialEditor/Private/Slate/Properties/Generators/DMTextureUVPropertyRowGenerator.cpp b/Engine/Plugins/Experimental/DynamicMaterial/Source/DynamicMaterialEditor/Private/Slate/Properties/Generators/DMTextureUVPropertyRowGenerator.cpp
--- a/Engine/Plugins/Experimental/DynamicMaterial/Source/DynamicMaterialEditor/Private/Slate/Properties/Generators/DMTextureUVPropertyRowGenerator.cpp
+++ b/Engine/Plugins/Experimental/DynamicMaterial/Source/DynamicMaterialEditor/Private/Slate/Properties/Generators/DMTextureUVPropertyRowGenerator.cpp
@@ -12,6 +12,7 @@
 #include "Slate/Properties/Editors/SDMPropertyEditFloat.h"
 #include "Slate/Properties/Editors/SDMPropertyEditVector.h"
 #include "Slate/SDMComponentEdit.h"
+#include "Slate/Properties/SDMTextureUVVisualizerProperty.h"
 #include "Styling/SlateIconFinder.h"
 #include "Widgets/Layout/SBox.h"
 
@@ -37,6 +38,8 @@ namespace UE::DynamicMaterialEditor::Private
 
 	void AddMirrorRow(const TSharedRef<SDMComponentEdit>& InComponentEditWidget, UDMTextureUV* InTextureUV, TArray<FDMPropertyHandle>& InOutPropertyRows);
 	TSharedRef<SWidget> CreateMirrorExtensionButtons(const TSharedRef<SDMComponentEdit>& InComponentEditWidget, UDMTextureUV* InTextureUV);
+
+	void AddVisualizerRow(const TSharedRef<SDMComponentEdit>& InComponentEditWidget, UDMTextureUV* InTextureUV, TArray<FDMPropertyHandle>& InOutPropertyRows);
 }
 
 void FDMTextureUVPropertyRowGenerator::AddComponentProperties(const TSharedRef<SDMComponentEdit>& InComponentEditWidget, UDMMaterialComponent* InComponent,
@@ -70,6 +73,7 @@ void FDMTextureUVPropertyRowGenerator::AddComponentProperties(const TSharedRef<S
 	AddScaleRow(InComponentEditWidget, TextureUV, InOutPropertyRows);
 	AddPropertyEditRows(InComponentEditWidget, InComponent, UDMTextureUV::NAME_Pivot, InOutPropertyRows, InOutProcessedObjects);
 	AddMirrorRow(InComponentEditWidget, TextureUV, InOutPropertyRows);
+	AddVisualizerRow(InComponentEditWidget, TextureUV, InOutPropertyRows);
 }
 
 bool FDMTextureUVPropertyRowGenerator::AllowKeyframeButton(UDMMaterialComponent* InComponent, FProperty* InProperty)
@@ -411,4 +415,22 @@ TSharedRef<SWidget> UE::DynamicMaterialEditor::Private::CreateMirrorExtensionBut
 	return ButtonsBox;
 }
 
+void UE::DynamicMaterialEditor::Private::AddVisualizerRow(const TSharedRef<SDMComponentEdit>& InComponentEditWidget, UDMTextureUV* InTextureUV, TArray<FDMPropertyHandle>& InOutPropertyRows)
+{
+	// Make sure we don't get a substage
+	UDMMaterialStage* Stage = InTextureUV->GetTypedParent<UDMMaterialStage>(/* Allow Subclasses */ false);
+
+	if (!Stage)
+	{
+		return;
+	}
+
+	FDMPropertyHandle VisualizerHandle;
+	VisualizerHandle.NameOverride = LOCTEXT("Visualizer", "UV Visualizer");
+	VisualizerHandle.NameToolTipOverride = LOCTEXT("VisualizerToolTip", "A graphical Texture UV editor.\n\n- Offset Mode: Change the Texture UV offset.\n- Pivot Mode: Change the Texture UV pivot, rotation and scale.\n\nControl+click to reset values to default.");
+	VisualizerHandle.ValueName = FName(*InTextureUV->GetComponentPath());
+	VisualizerHandle.ValueWidget = SNew(SDMTextureUVVisualizerProperty, Stage, InTextureUV);
+	InOutPropertyRows.Add(VisualizerHandle);
+}
+
 #undef LOCTEXT_NAMESPACE
diff --git a/Engine/Plugins/Experimental/DynamicMaterial/Source/DynamicMaterialEditor/Private/Slate/Properties/SDMPropertyEdit.h b/Engine/Plugins/Experimental/DynamicMaterial/Source/DynamicMaterialEditor/Private/Slate/Properties/SDMPropertyEdit.h
--- a/Engine/Plugins/Experimental/DynamicMaterial/Source/DynamicMaterialEditor/Private/Slate/Properties/SDMPropertyEdit.h
+++ b/Engine/Plugins/Experimental/DynamicMaterial/Source/DynamicMaterialEditor/Private/Slate/Properties/SDMPropertyEdit.h
@@ -2,12 +2,13 @@
 
 #pragma once
 
+#include "Widgets/SCompoundWidget.h"
+#include "Components/DMMaterialValue.h"
 #include "DynamicMaterialEditorModule.h"
 #include "PropertyCustomizationHelpers.h"
 #include "PropertyHandle.h"
 #include "Widgets/Input/SCheckBox.h"
 #include "Widgets/Input/SSpinBox.h"
-#include "Widgets/SCompoundWidget.h"
 
 class SDMComponentEdit;
 class SHorizontalBox;
diff --git a/Engine/Plugins/Experimental/DynamicMaterial/Source/DynamicMaterialEditor/Private/Slate/Properties/SDMTextureUVVisualizer.cpp b/Engine/Plugins/Experimental/DynamicMaterial/Source/DynamicMaterialEditor/Private/Slate/Properties/SDMTextureUVVisualizer.cpp
new file mode 100644
--- /dev/null
+++ b/Engine/Plugins/Experimental/DynamicMaterial/Source/DynamicMaterialEditor/Private/Slate/Properties/SDMTextureUVVisualizer.cpp
@@ -0,0 +1,1076 @@
+// Copyright Epic Games, Inc. All Rights Reserved.
+
+#include "Slate/Properties/SDMTextureUVVisualizer.h"
+
+#include "Brushes/SlateColorBrush.h"
+#include "Components/DMMaterialStage.h"
+#include "Components/DMTextureUV.h"
+#include "DynamicMaterialEditorModule.h"
+#include "Editor.h"
+#include "Framework/Application/SlateApplication.h"
+#include "Materials/MaterialInterface.h"
+#include "ScopedTransaction.h"
+#include "Styling/StyleColors.h"
+#include "Widgets/Images/SImage.h"
+#include "Widgets/Input/SButton.h"
+#include "Widgets/SBoxPanel.h"
+
+#define LOCTEXT_NAMESPACE "SDMTextureUVVisualizer"
+
+namespace UE::DynamicMaterialEditor::Private
+{
+	/** In details panel visualizer is small and square. */
+	const FVector2D TextureUVVisualizerImageSize = FVector2D(128.f, 128.f);
+
+	/** Popout visualizer is large and aspected differently. */
+	const FVector2D TextureUVVisualizerPopoutImageSize = FVector2D(1024.f, 768.f);
+
+	/** Outer size of the position handle. */
+	constexpr float TextureUVVisualizerLargeRadius = 10.f;
+
+	/** Inner size of the position handle. */
+	constexpr float TextureUVVisualizerSmallRadius = 5.f;
+
+	/** "Radius" circle handle is for mouse interaction. */
+	constexpr float TextureUVVisualizerCircleHandleRadius = 5.f;
+
+	/** Based distance the circle handle is from the center compared to the width of the image. */
+	constexpr float TextureUVVisualizerCircleHandleBaseRadiusMultiplier = 0.25f;
+}
+
+SDMTextureUVVisualizer::SDMTextureUVVisualizer()
+	: Brush(FSlateMaterialBrush(UE::DynamicMaterialEditor::Private::TextureUVVisualizerImageSize))
+	, bIsPopout(false)
+	, bPivotEditMode(false)
+	, CurrentAbsoluteSize(FVector2f::ZeroVector)
+	, CurrentAbsoluteCenter(FVector2f::ZeroVector)
+	, ScrubbingMode(EScrubbingMode::None)
+	, ScrubbingStartAbsoluteCenter(FVector2f::ZeroVector)
+	, ScrubbingStartAbsoluteMouse(FVector2f::ZeroVector)
+	, HandleAxis(EHandleAxis::None)
+	, ValueStart(FVector2D::ZeroVector)
+	, bInvertScale(false)
+{
+}
+
+void SDMTextureUVVisualizer::Construct(const FArguments& InArgs, UDMMaterialStage* InMaterialStage, UDMTextureUV* InTextureUV)
+{
+	check(InMaterialStage);
+	check(InTextureUV);
+
+	StageWeak = InMaterialStage;
+	TextureUVWeak = InTextureUV;
+	bIsPopout = InArgs._IsPopout;
+
+	InMaterialStage->GetOnUpdate().AddSP(this, &SDMTextureUVVisualizer::OnStageUpdated);
+	OnStageUpdated(InMaterialStage, EDMUpdateType::Structure);
+
+	SetCanTick(true);
+
+	using namespace UE::DynamicMaterialEditor::Private;
+
+	ChildSlot
+	[
+		SAssignNew(StageImage, SImage)
+		.Image(&Brush)
+		.DesiredSizeOverride(bIsPopout ? TextureUVVisualizerPopoutImageSize : TextureUVVisualizerImageSize)
+	];
+}
+
+SDMTextureUVVisualizer::EScrubbingMode SDMTextureUVVisualizer::GetScrubbingMode() const
+{
+	return ScrubbingMode;
+}
+
+bool SDMTextureUVVisualizer::IsInPivotEditMode() const
+{
+	return bPivotEditMode;
+}
+
+void SDMTextureUVVisualizer::SetInPivotEditMode(bool bInEditingPivot)
+{
+	bPivotEditMode = bInEditingPivot;
+}
+
+void SDMTextureUVVisualizer::TogglePivotEditMode()
+{
+	SetInPivotEditMode(!IsInPivotEditMode());
+}
+
+void SDMTextureUVVisualizer::Tick(const FGeometry& InAllottedGeometry, const double InCurrentTime, const float InDeltaTime)
+{
+	SCompoundWidget::Tick(InAllottedGeometry, InCurrentTime, InDeltaTime);
+
+	if (!PreviewMaterialWeak.IsValid())
+	{
+		Brush.SetMaterial(nullptr);
+	}
+
+	if (UDMMaterialComponent::CanClean())
+	{
+		if (UDMMaterialStage* Stage = StageWeak.Get())
+		{
+			if (Stage->IsComponentValid() && Stage->NeedsClean())
+			{
+				Stage->DoClean();
+				OnStageUpdated(StageWeak.Get(), EDMUpdateType::Structure);
+			}
+		}
+	}
+
+	CurrentAbsoluteSize = StageImage->GetTickSpaceGeometry().GetAbsoluteSize();
+	CurrentAbsoluteCenter = StageImage->GetTickSpaceGeometry().GetAbsolutePosition() + (CurrentAbsoluteSize * 0.5f);
+
+	if (ScrubbingMode != EScrubbingMode::None)
+	{
+		if (!FSlateApplication::Get().GetPressedMouseButtons().Contains(EKeys::LeftMouseButton))
+		{
+			SetScrubbingMode(EScrubbingMode::None, EHandleAxis::None);
+		}
+		else
+		{
+			UpdateScrub();
+		}
+	}
+	else
+	{
+		ScrubbingTransaction.Reset();
+	}
+}
+
+FReply SDMTextureUVVisualizer::OnMouseButtonDown(const FGeometry& InGeometry, const FPointerEvent& InMouseEvent)
+{
+	const bool bIsLeftMouse = InMouseEvent.GetEffectingButton() != EKeys::RightMouseButton;
+
+	if (!bIsLeftMouse)
+	{
+		return SCompoundWidget::OnMouseButtonDown(InGeometry, InMouseEvent);
+	}
+
+	// Could use the mouse event, but I want a consistent value source.
+	const FVector2f MousePosition = FSlateApplication::Get().GetCursorPos();
+	const bool bResetToDefault = InMouseEvent.GetModifierKeys().IsControlDown();
+
+	if (TryClickCenterHandle(MousePosition, bResetToDefault))
+	{
+		return FReply::Handled();
+	}
+
+	if (bPivotEditMode && TryClickCircleHandle(MousePosition, bResetToDefault))
+	{
+		return FReply::Handled();
+	}
+
+	return FReply::Unhandled();
+}
+
+FCursorReply SDMTextureUVVisualizer::OnCursorQuery(const FGeometry& InGeometry, const FPointerEvent& InCursorEvent) const
+{
+	if (HandleAxis != EHandleAxis::None)
+	{
+		return FCursorReply::Cursor(EMouseCursor::Crosshairs);
+	}
+
+	const FVector2f MousePosition = FSlateApplication::Get().GetCursorPos();
+
+	// If the UV image isn't under the mouse, we don't need to do anything.
+	if (!StageImage->GetTickSpaceGeometry().IsUnderLocation(MousePosition))
+	{
+		return SCompoundWidget::OnCursorQuery(InGeometry, InCursorEvent);
+	}
+
+	// Work out if the mouse is over a part of the center handle
+	EHandleAxis Axis = GetCenterHandleAxis(MousePosition);
+
+	switch (Axis)
+	{
+		case EHandleAxis::X:
+			return FCursorReply::Cursor(EMouseCursor::ResizeLeftRight);
+
+		case EHandleAxis::Y:
+			return FCursorReply::Cursor(EMouseCursor::ResizeUpDown);
+
+		case EHandleAxis::XY:
+			return FCursorReply::Cursor(EMouseCursor::CardinalCross);
+	}
+
+	if (bPivotEditMode)
+	{
+		// Work out if the mouse is over a part of the circle handle
+		Axis = GetCircleHandleAxis(MousePosition);
+
+		switch (Axis)
+		{
+			case EHandleAxis::X:
+				return FCursorReply::Cursor(EMouseCursor::ResizeLeftRight);
+
+			case EHandleAxis::Y:
+				return FCursorReply::Cursor(EMouseCursor::ResizeUpDown);
+
+			case EHandleAxis::XY:
+				return FCursorReply::Cursor(EMouseCursor::GrabHand);
+		}
+	}
+
+	return SCompoundWidget::OnCursorQuery(InGeometry, InCursorEvent);
+}
+
+int32 SDMTextureUVVisualizer::OnPaint(const FPaintArgs& InArgs, const FGeometry& InAllottedGeometry, const FSlateRect& InMyCullingRect, FSlateWindowElementList& OutDrawElements, int32 InLayerId, const FWidgetStyle& InWidgetStyle, bool bInParentEnabled) const
+{
+	InLayerId = SCompoundWidget::OnPaint(InArgs, InAllottedGeometry, InMyCullingRect, OutDrawElements, InLayerId, InWidgetStyle, bInParentEnabled);
+
+	UDMTextureUV* TextureUV = TextureUVWeak.Get();
+
+	if (!IsValid(TextureUV))
+	{
+		return InLayerId;
+	}
+
+	++InLayerId;
+
+	using namespace UE::DynamicMaterialEditor::Private;
+
+	static const FSlateColorBrush WhiteBrush = FSlateColorBrush(FColor(255, 255, 255, 255));
+	static const FLinearColor BorderColor = FLinearColor::Black;
+	static const FLinearColor NormalColor = FLinearColor::White;
+	static const FLinearColor HighlightColor = FStyleColors::Primary.GetSpecifiedColor();
+
+	const float Rotation = TextureUV->GetRotation();
+	const float RotationRadians = FMath::DegreesToRadians(Rotation);
+	const FVector2f LocalCenterOffset = GetCenterHandleLocation(InAllottedGeometry.GetLocalSize());
+	const FVector2f LocalPivotOffset = GetPivotLocation(InAllottedGeometry.GetLocalSize());
+
+	auto DrawRotatedBox = [this, &OutDrawElements, &InLayerId, &InAllottedGeometry, &Rotation, &LocalCenterOffset, &LocalPivotOffset]
+		(const FVector2f& InLocation, const FVector2f& InSize, const FLinearColor& InColor, float InRotationRadians)
+		{
+			const FVector2f& BaseLocation = bPivotEditMode ? LocalPivotOffset : LocalCenterOffset;
+			const FVector2f Location = BaseLocation + InLocation - InSize * 0.5f;
+			const FVector2f RotationOffset = BaseLocation - Location;
+
+			FSlateDrawElement::MakeRotatedBox(
+				OutDrawElements,
+				InLayerId,
+				InAllottedGeometry.ToPaintGeometry(
+					InSize,
+					FSlateLayoutTransform(Location)
+				),
+				&WhiteBrush,
+				ESlateDrawEffect::NoPixelSnapping,
+				InRotationRadians,
+				RotationOffset,
+				FSlateDrawElement::ERotationSpace::RelativeToElement,
+				InColor
+			);
+
+			return Location;
+		};
+
+	auto DrawRotatedBorderBox = [&DrawRotatedBox](const FVector2f& InLocation, const FVector2f& InSize, const FLinearColor& InInnerColor, float InRotationRadians)
+		{
+			DrawRotatedBox(InLocation, InSize, BorderColor, InRotationRadians);
+			return DrawRotatedBox(InLocation, InSize - FVector2f(4.f, 4.f), InInnerColor, InRotationRadians);
+		};
+
+	/** Center Handle */
+	DrawRotatedBorderBox(
+		FVector2f::ZeroVector,
+		FVector2f(TextureUVVisualizerLargeRadius, TextureUVVisualizerLargeRadius) * 2.f,
+		NormalColor,
+		bPivotEditMode ? 0 : RotationRadians
+	);
+
+	if (ScrubbingMode == EScrubbingMode::Offset || ScrubbingMode == EScrubbingMode::Pivot)
+	{
+		if (HandleAxis != EHandleAxis::Y)
+		{
+			DrawRotatedBox(
+				FVector2f(TextureUVVisualizerLargeRadius - TextureUVVisualizerSmallRadius * 0.5f - 2.f, 0.f),
+				FVector2f(TextureUVVisualizerSmallRadius, TextureUVVisualizerLargeRadius * 2.f - 4.f),
+				HighlightColor,
+				bPivotEditMode ? 0 : RotationRadians
+			);
+
+			DrawRotatedBox(
+				FVector2f(-TextureUVVisualizerLargeRadius + TextureUVVisualizerSmallRadius * 0.5f + 2.f, 0.f),
+				FVector2f(TextureUVVisualizerSmallRadius, TextureUVVisualizerLargeRadius * 2.f - 4.f),
+				HighlightColor,
+				bPivotEditMode ? 0 : RotationRadians
+			);
+		}
+
+		if (HandleAxis != EHandleAxis::X)
+		{
+			DrawRotatedBox(
+				FVector2f(0.f, TextureUVVisualizerLargeRadius - TextureUVVisualizerSmallRadius * 0.5f - 2.f),
+				FVector2f(TextureUVVisualizerLargeRadius * 2.f - 4.f, TextureUVVisualizerSmallRadius),
+				HighlightColor,
+				bPivotEditMode ? 0 : RotationRadians
+			);
+
+			DrawRotatedBox(
+				FVector2f(0.f, -TextureUVVisualizerLargeRadius + TextureUVVisualizerSmallRadius * 0.5f + 2.f),
+				FVector2f(TextureUVVisualizerLargeRadius * 2.f - 4.f, TextureUVVisualizerSmallRadius),
+				HighlightColor,
+				bPivotEditMode ? 0 : RotationRadians
+			);
+		}
+	}
+
+	if (bPivotEditMode)
+	{
+		/** Scale Handles */
+		float RadiusAtAngle[16];
+
+		for (int32 Direction = 0; Direction < 16; ++Direction)
+		{
+			RadiusAtAngle[Direction] = GetCircleHandleRadiusAtAngle(static_cast<float>(Direction) * 22.5f);
+		}
+
+		DrawRotatedBorderBox(
+			FVector2f(0.f, RadiusAtAngle[0]),
+			FVector2f(TextureUVVisualizerLargeRadius, TextureUVVisualizerLargeRadius),
+			(ScrubbingMode == EScrubbingMode::Scale && HandleAxis == EHandleAxis::Y) ? HighlightColor : NormalColor,
+			RotationRadians
+		);
+
+		DrawRotatedBorderBox(
+			FVector2f(RadiusAtAngle[4], 0.f),
+			FVector2f(TextureUVVisualizerLargeRadius, TextureUVVisualizerLargeRadius),
+			(ScrubbingMode == EScrubbingMode::Scale && HandleAxis == EHandleAxis::X) ? HighlightColor : NormalColor,
+			RotationRadians
+		);
+
+		DrawRotatedBorderBox(
+			FVector2f(0.f, -RadiusAtAngle[8]),
+			FVector2f(TextureUVVisualizerLargeRadius, TextureUVVisualizerLargeRadius),
+			(ScrubbingMode == EScrubbingMode::Scale && HandleAxis == EHandleAxis::Y) ? HighlightColor : NormalColor,
+			RotationRadians
+		);
+
+		DrawRotatedBorderBox(
+			FVector2f(-RadiusAtAngle[12], 0.f),
+			FVector2f(TextureUVVisualizerLargeRadius, TextureUVVisualizerLargeRadius),
+			(ScrubbingMode == EScrubbingMode::Scale && HandleAxis == EHandleAxis::X) ? HighlightColor : NormalColor,
+			RotationRadians
+		);
+
+		/** Rotation Handles */
+		auto AngleRadiusToLocation = [&LocalPivotOffset, &Rotation, &RadiusAtAngle](int32 InAngleIndex)
+			{
+				const float Angle = FMath::DegreesToRadians(static_cast<float>(InAngleIndex) * 22.5f - Rotation);
+				return LocalPivotOffset + (FVector2f(FMath::Sin(Angle), FMath::Cos(Angle)) * RadiusAtAngle[InAngleIndex]);
+			};
+
+		auto DrawPartialCicle = [this, &OutDrawElements, &InLayerId, &InAllottedGeometry, &AngleRadiusToLocation](int32 InStartIndex)
+			{
+				TArray<FVector2f> Points = {
+					AngleRadiusToLocation(InStartIndex),
+					AngleRadiusToLocation(InStartIndex + 1),
+					AngleRadiusToLocation(InStartIndex + 2)
+				};
+
+				FSlateDrawElement::MakeLines(
+					OutDrawElements,
+					InLayerId,
+					InAllottedGeometry.ToPaintGeometry(),
+					Points,
+					ESlateDrawEffect::NoPixelSnapping,
+					BorderColor,
+					true,
+					TextureUVVisualizerSmallRadius
+				);
+
+				FSlateDrawElement::MakeLines(
+					OutDrawElements,
+					InLayerId,
+					InAllottedGeometry.ToPaintGeometry(),
+					Points,
+					ESlateDrawEffect::NoPixelSnapping,
+					ScrubbingMode == EScrubbingMode::Rotation ? HighlightColor : NormalColor,
+					true,
+					TextureUVVisualizerSmallRadius - 2.f
+				);
+			};
+
+		DrawPartialCicle(1);
+		DrawPartialCicle(5);
+		DrawPartialCicle(9);
+		DrawPartialCicle(13);
+	}
+
+	return InLayerId;
+}
+
+void SDMTextureUVVisualizer::OnStageUpdated(UDMMaterialComponent* InComponent, EDMUpdateType InUpdateType)
+{
+	if (UDMMaterialStage* Stage = Cast<UDMMaterialStage>(InComponent))
+	{
+		if (Stage == StageWeak.Get() && IsValid(Stage) && Stage->IsComponentValid())
+		{
+			UMaterialInterface* PreviewMaterial = Stage->GetPreviewMaterial();
+
+			if (Brush.GetResourceObject() != PreviewMaterial)
+			{
+				Brush.SetMaterial(PreviewMaterial);
+				PreviewMaterialWeak = PreviewMaterial;
+			}
+		}
+	}
+}
+
+bool SDMTextureUVVisualizer::HasValidGeometry() const
+{
+	// The chances of the center being at anywhere near 0,0 is remote...
+	return !FMath::IsNearlyZero(CurrentAbsoluteSize.X) && !FMath::IsNearlyZero(CurrentAbsoluteSize.Y)
+		&& !FMath::IsNearlyZero(CurrentAbsoluteCenter.X) && !FMath::IsNearlyZero(CurrentAbsoluteCenter.Y);
+}
+
+float SDMTextureUVVisualizer::GetCircleHandleBaseRadius() const
+{
+	using namespace UE::DynamicMaterialEditor::Private;
+
+	const FVector2f ImageSize = Brush.GetImageSize();
+	const float CircleHandleRadius = FMath::Min(ImageSize.X, ImageSize.Y) * TextureUVVisualizerCircleHandleBaseRadiusMultiplier;
+
+	return FMath::Clamp(CircleHandleRadius, 10, 50.f);
+}
+
+FVector2f SDMTextureUVVisualizer::ApplyTextureUVTransform(const FVector2f& InUV) const
+{
+	UDMTextureUV* TextureUV = TextureUVWeak.Get();
+
+	if (!IsValid(TextureUV))
+	{
+		return InUV;
+	}
+
+	FVector2f Offset = static_cast<FVector2f>(TextureUV->GetOffset());
+	Offset.Y *= -1.f;
+
+	const float Rotation = TextureUV->GetRotation();
+	const FVector2f& Scale = static_cast<FVector2f>(TextureUV->GetScale());
+	const FVector2f& Pivot = static_cast<FVector2f>(TextureUV->GetPivot());
+
+	FVector2f TransformedUV = InUV;
+	TransformedUV -= Pivot;
+
+	if (!FMath::IsNearlyZero(Rotation))
+	{
+		TransformedUV = TransformedUV.GetRotated(Rotation);
+	}
+
+	TransformedUV *= Scale;
+	TransformedUV += Pivot;
+
+	if (!FMath::IsNearlyZero(Rotation))
+	{
+		TransformedUV += (Offset * Scale).GetRotated(Rotation);
+	}
+	else
+	{
+		TransformedUV += Offset * Scale;
+	}
+
+	return TransformedUV;
+}
+
+FVector2f SDMTextureUVVisualizer::GetCenterHandleLocation(const FVector2f& InSize) const
+{
+	if (bPivotEditMode)
+	{
+		return GetPivotLocation(InSize);
+	}
+
+	static const FVector2f Center = {0.5f, 0.5f};
+
+	return ApplyTextureUVTransform(Center) * InSize;
+}
+
+FVector2f SDMTextureUVVisualizer::GetPivotLocation(const FVector2f& InSize) const
+{
+	static const FVector2f Center = {0.5f, 0.5f};
+
+	if (UDMTextureUV* TextureUV = TextureUVWeak.Get())
+	{
+		return static_cast<FVector2f>(TextureUV->GetPivot()) * InSize;
+	}
+
+	return Center * InSize;
+}
+
+FVector2f SDMTextureUVVisualizer::GetAbsoluteCenterHandleLocation() const
+{
+	return CurrentAbsoluteCenter - (CurrentAbsoluteSize * 0.5f) + GetCenterHandleLocation(CurrentAbsoluteSize);
+}
+
+FVector2f SDMTextureUVVisualizer::GetAbsolutePivotLocation() const
+{
+	return CurrentAbsoluteCenter - (CurrentAbsoluteSize * 0.5f) + GetPivotLocation(CurrentAbsoluteSize);
+}
+
+float SDMTextureUVVisualizer::GetCircleHandleRadiusAtAngle(float InAngle) const
+{
+	const float BaseDistance = GetCircleHandleBaseRadius();
+	UDMTextureUV* TextureUV = TextureUVWeak.Get();
+
+	if (!IsValid(TextureUV))
+	{
+		return BaseDistance;
+	}
+
+	const FVector2D& Scale2D = TextureUV->GetScale();
+
+	if (FMath::IsNearlyEqual(Scale2D.X, Scale2D.Y))
+	{
+		return BaseDistance * Scale2D.X;
+	}
+
+	const float RadianAngle = FMath::DegreesToRadians(InAngle);
+
+	const FVector2f Scale = static_cast<FVector2f>(Scale2D);
+
+	const FVector2f Offset = {
+		FMath::Sin(RadianAngle) * Scale.X,
+		FMath::Cos(RadianAngle) * Scale.Y
+	};
+
+	return BaseDistance * Offset.Size();
+}
+
+SDMTextureUVVisualizer::EHandleAxis SDMTextureUVVisualizer::GetCenterHandleAxis(const FVector2f& InAbsolutePosition) const
+{
+	if (!HasValidGeometry())
+	{
+		return EHandleAxis::None;
+	}
+
+	UDMTextureUV* TextureUV = TextureUVWeak.Get();
+
+	if (!IsValid(TextureUV))
+	{
+		return EHandleAxis::None;
+	}
+
+	using namespace UE::DynamicMaterialEditor::Private;
+
+	FVector2f CenterOffset = InAbsolutePosition - GetAbsoluteCenterHandleLocation();
+
+	if (!bPivotEditMode)
+	{
+		const float Rotation = TextureUV->GetRotation();
+
+		if (!FMath::IsNearlyZero(Rotation))
+		{
+			CenterOffset = CenterOffset.GetRotated(-Rotation);
+		}
+	}
+
+	CenterOffset.X = FMath::Abs(CenterOffset.X);
+	CenterOffset.Y = FMath::Abs(CenterOffset.Y);
+
+	if (CenterOffset.X <= TextureUVVisualizerSmallRadius && CenterOffset.Y <= TextureUVVisualizerSmallRadius)
+	{
+		return EHandleAxis::XY;
+	}
+
+	if (CenterOffset.X <= TextureUVVisualizerSmallRadius && CenterOffset.Y <= TextureUVVisualizerLargeRadius)
+	{
+		return EHandleAxis::Y;
+	}
+
+	if (CenterOffset.X <= TextureUVVisualizerLargeRadius && CenterOffset.Y <= TextureUVVisualizerSmallRadius)
+	{
+		return EHandleAxis::X;
+	}
+
+	if (CenterOffset.X <= TextureUVVisualizerLargeRadius && CenterOffset.Y <= TextureUVVisualizerLargeRadius)
+	{
+		return EHandleAxis::XY;
+	}
+
+	return EHandleAxis::None;
+}
+
+SDMTextureUVVisualizer::EHandleAxis SDMTextureUVVisualizer::GetCircleHandleAxis(const FVector2f& InAbsolutePosition) const
+{
+	UDMTextureUV* TextureUV = TextureUVWeak.Get();
+
+	if (!IsValid(TextureUV))
+	{
+		return EHandleAxis::None;
+	}
+
+	using namespace UE::DynamicMaterialEditor::Private;
+
+	const FVector2f CenterOffset = InAbsolutePosition - GetAbsolutePivotLocation();
+	const float DistanceFromCenter = CenterOffset.Size();
+	const FVector2D& Scale2D = TextureUV->GetScale();
+	// We manage angle clockwise from +Y axis. Atan2 handles it anti-clockwise from +X axis.
+	float Angle = 90.f - FMath::RadiansToDegrees(FMath::Atan2(CenterOffset.Y, CenterOffset.X)) + TextureUV->GetRotation();
+	float DistanceFromCircleHandle;
+
+	if (FMath::IsNearlyEqual(Scale2D.X, Scale2D.Y))
+	{
+		DistanceFromCircleHandle = FMath::Abs(DistanceFromCenter - GetCircleHandleBaseRadius());
+	}
+	else
+	{
+		DistanceFromCircleHandle = FMath::Abs(DistanceFromCenter - GetCircleHandleRadiusAtAngle(Angle));
+	}
+
+	if (DistanceFromCircleHandle > TextureUVVisualizerCircleHandleRadius)
+	{
+		return EHandleAxis::None;
+	}
+
+	Angle += 22.5f; // Rotate slightly so it's easier to follow the below code.
+	Angle = UE::Math::TRotator<float>::ClampAxis(Angle);
+
+	constexpr float AnglePerQuadrant = 360.f / 8.f; // 45
+	constexpr float TopEnd = AnglePerQuadrant; // 45
+	constexpr float TopRightEnd = TopEnd + AnglePerQuadrant; // 90
+	constexpr float RightEnd = TopRightEnd + AnglePerQuadrant; // 135
+	constexpr float BottomRightEnd = RightEnd + AnglePerQuadrant; // 180
+	constexpr float BottomEnd = BottomRightEnd + AnglePerQuadrant; // 225
+	constexpr float BottomLeftEnd = BottomEnd + AnglePerQuadrant; // 270
+	constexpr float LeftEnd = BottomLeftEnd + AnglePerQuadrant; // 313
+	constexpr float TopLeftEnd = LeftEnd + AnglePerQuadrant; // 360
+
+	if (Angle < TopEnd)
+	{
+		return EHandleAxis::Y;
+	}
+
+	if (Angle < TopRightEnd)
+	{
+		return EHandleAxis::XY;
+	}
+
+	if (Angle < RightEnd)
+	{
+		return EHandleAxis::X;
+	}
+
+	if (Angle < BottomRightEnd)
+	{
+		return EHandleAxis::XY;
+	}
+
+	if (Angle < BottomEnd)
+	{
+		return EHandleAxis::Y;
+	}
+
+	if (Angle < BottomLeftEnd)
+	{
+		return EHandleAxis::XY;
+	}
+
+	if (Angle < LeftEnd)
+	{
+		return EHandleAxis::X;
+	}
+
+	// Top left
+	return EHandleAxis::XY;
+}
+
+bool SDMTextureUVVisualizer::TryClickCenterHandle(const FVector2f& InMousePosition, bool bInResetToDefault)
+{
+	UDMTextureUV* TextureUV = TextureUVWeak.Get();
+
+	if (!IsValid(TextureUV))
+	{
+		return false;
+	}
+
+	// Work out if the mouse is over a part of the center handle
+	EHandleAxis Axis = GetCenterHandleAxis(InMousePosition);
+
+	if (Axis == EHandleAxis::None)
+	{
+		return false;
+	}
+
+	// Offset handle
+	if (!bPivotEditMode)
+	{
+		if (!bInResetToDefault)
+		{
+			SetScrubbingMode(EScrubbingMode::Offset, Axis);
+		}
+		else
+		{
+			FScopedTransaction Transaction(LOCTEXT("ResetOffset", "Reset Offset to Default."));
+			TextureUV->Modify();
+			FVector2D NewOffset = TextureUV->GetOffset();
+
+			if (Axis == EHandleAxis::X)
+			{
+				NewOffset.X = 0;
+			}
+			else if (Axis == EHandleAxis::Y)
+			{
+				NewOffset.Y = 0;
+			}
+			else
+			{
+				NewOffset = FVector2D::ZeroVector;
+			}
+
+			TextureUV->SetOffset(NewOffset);
+		}
+	}
+	// Pivot handle
+	else
+	{
+		if (!bInResetToDefault)
+		{
+			SetScrubbingMode(EScrubbingMode::Pivot, Axis);
+		}
+		else
+		{
+			FScopedTransaction Transaction(LOCTEXT("ResetPivot", "Reset Pivot to Default."));
+			TextureUV->Modify();
+			FVector2D NewPivot = TextureUV->GetOffset();
+
+			if (Axis == EHandleAxis::X)
+			{
+				NewPivot.X = 0.5;
+			}
+			else if (Axis == EHandleAxis::Y)
+			{
+				NewPivot.Y = 0.5;
+			}
+			else
+			{
+				NewPivot = FVector2D(0.5, 0.5);
+			}
+
+			TextureUV->SetPivot(NewPivot);
+		}
+	}
+
+	return true;
+}
+
+bool SDMTextureUVVisualizer::TryClickCircleHandle(const FVector2f& InMousePosition, bool bInResetToDefault)
+{
+	UDMTextureUV* TextureUV = TextureUVWeak.Get();
+
+	if (!IsValid(TextureUV))
+	{
+		return false;
+	}
+
+	const float Rotation = TextureUV->GetRotation();
+	const FVector2f AbsolutePivotLocation = GetAbsolutePivotLocation();
+
+	// When the uv is rotated, do the opposite action
+	const bool bRegularInvertScale = Rotation <= 90.f || Rotation > 270.f;
+
+	// Work out if the mouse is over part of the circle handle.
+	EHandleAxis Axis = GetCircleHandleAxis(InMousePosition);
+
+	switch (Axis)
+	{
+		default:
+			return false;
+
+		// XY on the circle handle represents rotation
+		case EHandleAxis::XY:
+		{
+			if (!bInResetToDefault)
+			{
+				SetScrubbingMode(EScrubbingMode::Rotation, Axis);
+			}
+			else
+			{
+				FScopedTransaction Transaction(LOCTEXT("ResetRotation", "Reset Rotation to Default."));
+				TextureUV->Modify();
+				TextureUV->SetRotation(0.f);
+			}
+
+			break;
+		}
+
+		case EHandleAxis::X:
+		{
+			if (!bInResetToDefault)
+			{
+				SetScrubbingMode(EScrubbingMode::Scale, Axis);
+				bInvertScale = (InMousePosition.X <= AbsolutePivotLocation.X) == bRegularInvertScale;
+			}
+			else
+			{
+				FScopedTransaction Transaction(LOCTEXT("ResetScaleX", "Reset Scale X to Default."));
+				TextureUV->Modify();
+				FVector2D NewScale = TextureUV->GetScale();
+				NewScale.X = 1.0;
+				TextureUV->SetScale(NewScale);
+			}
+
+			break;
+		}
+
+		case EHandleAxis::Y:
+		{
+			if (!bInResetToDefault)
+			{
+				SetScrubbingMode(EScrubbingMode::Scale, Axis);
+				bInvertScale = (InMousePosition.Y <= AbsolutePivotLocation.Y) == bRegularInvertScale;
+			}
+			else
+			{
+				FScopedTransaction Transaction(LOCTEXT("ResetScaleY", "Reset Scale Y to Default."));
+				TextureUV->Modify();
+				FVector2D NewScale = TextureUV->GetScale();
+				NewScale.Y = 1.0;
+				TextureUV->SetScale(NewScale);
+			}
+
+			break;
+		}
+	}
+
+	return true;
+}
+
+void SDMTextureUVVisualizer::SetScrubbingMode(EScrubbingMode InMode, EHandleAxis InAxis)
+{
+	UDMTextureUV* TextureUV = TextureUVWeak.Get();
+
+	if (InMode != EScrubbingMode::None && (InAxis == EHandleAxis::None || !IsValid(TextureUV) || !HasValidGeometry()))
+	{
+		InMode = EScrubbingMode::None;
+	}
+
+	ScrubbingMode = InMode;
+	HandleAxis = ScrubbingMode != EScrubbingMode::None ? InAxis : EHandleAxis::None;
+
+	if (InMode == EScrubbingMode::None)
+	{
+		ScrubbingTransaction.Reset();
+		return;
+	}
+
+	ScrubbingStartAbsoluteCenter = CurrentAbsoluteCenter;
+	ScrubbingStartAbsoluteMouse = FSlateApplication::Get().GetCursorPos();
+
+	switch (ScrubbingMode)
+	{
+		case EScrubbingMode::Offset:
+			ValueStart = TextureUV->GetOffset();
+			UE_LOG(LogDynamicMaterialEditor, Verbose, TEXT("Started Offset mode"));
+			break;
+
+		case EScrubbingMode::Rotation:
+		{
+			ValueStart.X = TextureUV->GetRotation();
+
+			/** Store the original mouse angle */
+			const FVector2f MouseOffset = ScrubbingStartAbsoluteMouse - GetAbsolutePivotLocation();
+			ValueStart.Y = FMath::RadiansToDegrees(FMath::Atan2(MouseOffset.Y, MouseOffset.X));
+
+			UE_LOG(LogDynamicMaterialEditor, Verbose, TEXT("Started Rotation mode"));
+			break;
+		}
+
+		case EScrubbingMode::Scale:
+			ValueStart = TextureUV->GetScale();
+			UE_LOG(LogDynamicMaterialEditor, Verbose, TEXT("Started Scale mode"));
+			break;
+
+		case EScrubbingMode::Pivot:
+			ValueStart = TextureUV->GetPivot();
+			UE_LOG(LogDynamicMaterialEditor, Verbose, TEXT("Started Pivot mode"));
+			break;
+
+		default:
+			return;
+	}
+
+	ScrubbingTransaction = MakeShared<FScopedTransaction>(LOCTEXT("VisualizerUVScrubbingTransaction", "UV Visualizer Scrub"));
+	TextureUV->Modify();
+}
+
+void SDMTextureUVVisualizer::UpdateScrub()
+{
+	switch (ScrubbingMode)
+	{
+		case EScrubbingMode::Offset:
+			UpdateScrub_Offset();
+			break;
+
+		case EScrubbingMode::Rotation:
+			UpdateScrub_Rotation();
+			break;
+
+		case EScrubbingMode::Scale:
+			UpdateScrub_Scale();
+			break;
+
+		case EScrubbingMode::Pivot:
+			UpdateScrub_Pivot();
+			break;
+
+		default:
+			return;
+	}
+
+	if (GEditor)
+	{
+		GEditor->RedrawLevelEditingViewports();
+	}
+}
+
+void SDMTextureUVVisualizer::UpdateScrub_Offset()
+{
+	UDMTextureUV* TextureUV = TextureUVWeak.Get();
+
+	if (!IsValid(TextureUV))
+	{
+		return;
+	}
+
+	const FVector2f MouseOffset = FSlateApplication::Get().GetCursorPos() - ScrubbingStartAbsoluteMouse;
+
+	FVector2D OffsetChange = static_cast<FVector2D>(MouseOffset / CurrentAbsoluteSize);
+
+	const float Rotation = TextureUV->GetRotation();
+
+	if (!FMath::IsNearlyZero(Rotation))
+	{
+		OffsetChange = OffsetChange.GetRotated(-Rotation);
+	}
+
+	OffsetChange /= TextureUV->GetScale();
+
+	if (HandleAxis == EHandleAxis::X)
+	{
+		OffsetChange.Y = 0;
+	}
+	else
+	{
+		OffsetChange.Y *= -1.0;
+	}
+
+	if (HandleAxis == EHandleAxis::Y)
+	{
+		OffsetChange.X = 0;
+	}
+
+	TextureUV->SetOffset(ValueStart + OffsetChange);
+}
+
+void SDMTextureUVVisualizer::UpdateScrub_Rotation()
+{
+	UDMTextureUV* TextureUV = TextureUVWeak.Get();
+
+	if (!IsValid(TextureUV))
+	{
+		return;
+	}
+
+	const FVector2f AbsoluteMousePosition = FSlateApplication::Get().GetCursorPos();
+	const FVector2f CurrentMouseOffset = AbsoluteMousePosition - GetAbsolutePivotLocation();
+	const float Angle = FMath::RadiansToDegrees(FMath::Atan2(CurrentMouseOffset.Y, CurrentMouseOffset.X));
+	const float NewValue = UE::Math::TRotator<float>::ClampAxis(ValueStart.X + Angle - ValueStart.Y);
+
+	TextureUV->SetRotation(NewValue);
+}
+
+void SDMTextureUVVisualizer::UpdateScrub_Scale()
+{
+	UDMTextureUV* TextureUV = TextureUVWeak.Get();
+
+	if (!IsValid(TextureUV))
+	{
+		return;
+	}
+
+	const FVector2f MouseOffset = FSlateApplication::Get().GetCursorPos() - ScrubbingStartAbsoluteMouse;
+
+	FVector2D ScaleChange = static_cast<FVector2D>(MouseOffset / CurrentAbsoluteSize);
+	ScaleChange /= TextureUV->GetScale();
+
+	if (bInvertScale)
+	{
+		ScaleChange *= -1.f;
+	}
+
+	const float Rotation = TextureUV->GetRotation();
+
+	if (!FMath::IsNearlyZero(Rotation))
+	{
+		ScaleChange = ScaleChange.GetRotated(-Rotation);
+	}
+
+	if (HandleAxis == EHandleAxis::X)
+	{
+		ScaleChange.Y = 0;
+	}
+	else
+	{
+		ScaleChange.Y *= -1.0;
+	}
+
+	if (HandleAxis == EHandleAxis::Y)
+	{
+		ScaleChange.X = 0;
+	}
+
+	FVector2D NewScale = ValueStart;
+
+	if (HandleAxis == EHandleAxis::X)
+	{
+		if (ScaleChange.X > 0)
+		{
+			NewScale.X = FMath::Max(0.001, NewScale.X * (1.f + ScaleChange.X));
+		}
+		else
+		{
+			NewScale.X = FMath::Max(0.001, NewScale.X / (1.f - ScaleChange.X));
+		}
+	}
+	else
+	{
+		if (ScaleChange.Y > 0)
+		{
+			NewScale.Y = FMath::Max(0.001, NewScale.Y / (1.f + ScaleChange.Y));
+		}
+		else
+		{
+			NewScale.Y = FMath::Max(0.001, NewScale.Y * (1.f - ScaleChange.Y));
+		}
+	}
+
+	TextureUV->SetScale(NewScale);
+}
+
+void SDMTextureUVVisualizer::UpdateScrub_Pivot()
+{
+	UDMTextureUV* TextureUV = TextureUVWeak.Get();
+
+	if (!IsValid(TextureUV))
+	{
+		return;
+	}
+
+	const FVector2f MouseOffset = FSlateApplication::Get().GetCursorPos() - ScrubbingStartAbsoluteMouse;
+
+	FVector2D PivotChange = static_cast<FVector2D>(MouseOffset / CurrentAbsoluteSize);
+
+	if (HandleAxis == EHandleAxis::X)
+	{
+		PivotChange.Y = 0;
+	}
+	else if (HandleAxis == EHandleAxis::Y)
+	{
+		PivotChange.X = 0;
+	}
+
+	TextureUV->SetPivot(ValueStart + PivotChange);
+}
+
+#undef LOCTEXT_NAMESPACE
diff --git a/Engine/Plugins/Experimental/DynamicMaterial/Source/DynamicMaterialEditor/Private/Slate/Properties/SDMTextureUVVisualizer.h b/Engine/Plugins/Experimental/DynamicMaterial/Source/DynamicMaterialEditor/Private/Slate/Properties/SDMTextureUVVisualizer.h
new file mode 100644
--- /dev/null
+++ b/Engine/Plugins/Experimental/DynamicMaterial/Source/DynamicMaterialEditor/Private/Slate/Properties/SDMTextureUVVisualizer.h
@@ -0,0 +1,123 @@
+// Copyright Epic Games, Inc. All Rights Reserved.
+
+#pragma once
+
+#include "Widgets/SCompoundWidget.h"
+#include "SlateMaterialBrush.h"
+#include "Templates/SharedPointer.h"
+#include "UObject/WeakObjectPtr.h"
+
+class FScopedTransaction;
+class SImage;
+class SWidget;
+class UDMMaterialComponent;
+class UDMMaterialStage;
+class UDMTextureUV;
+class UMaterialInterface;
+enum class EDMUpdateType : uint8;
+
+/**
+ * Material Designer Texture UV Visualizer 
+ *
+ * Ability to edit Texture UV settings in a visual manner.
+ */
+class SDMTextureUVVisualizer : public SCompoundWidget
+{
+public:
+	enum class EScrubbingMode : uint8
+	{
+		None,
+		Offset,
+		Rotation,
+		Scale,
+		Pivot
+	};
+
+	enum class EHandleAxis : uint8
+	{
+		None,
+		X,
+		Y,
+		XY
+	};
+
+	SLATE_BEGIN_ARGS(SDMTextureUVVisualizer)
+		: _IsPopout(false)
+		{}
+		SLATE_ARGUMENT(bool, IsPopout)
+	SLATE_END_ARGS()
+
+	SDMTextureUVVisualizer();
+
+	/** The TextureUV should be a sub-property of the stage */
+	void Construct(const FArguments& InArgs, UDMMaterialStage* InMaterialStage, UDMTextureUV* InTextureUV);
+
+	EScrubbingMode GetScrubbingMode() const;
+
+	bool IsInPivotEditMode() const;
+
+	void SetInPivotEditMode(bool bInEditingPivot);
+
+	void TogglePivotEditMode();
+
+	//~ Begin SWidget
+	virtual void Tick(const FGeometry& InAllottedGeometry, const double InCurrentTime, const float InDeltaTime) override;
+	virtual FReply OnMouseButtonDown(const FGeometry& InGeometry, const FPointerEvent& InMouseEvent) override;
+	virtual FCursorReply OnCursorQuery(const FGeometry& InGeometry, const FPointerEvent& InCursorEvent) const override;
+	virtual int32 OnPaint(const FPaintArgs& InArgs, const FGeometry& InAllottedGeometry, const FSlateRect& InMyCullingRect, 
+		FSlateWindowElementList& OutDrawElements, int32 InLayerId, const FWidgetStyle& InWidgetStyle, bool bInParentEnabled) const override;
+	//~ End SWidget
+
+protected:
+	TWeakObjectPtr<UDMMaterialStage> StageWeak;
+	TWeakObjectPtr<UDMTextureUV> TextureUVWeak;
+	TSharedPtr<SImage> StageImage;
+	TWeakObjectPtr<UMaterialInterface> PreviewMaterialWeak;
+	FSlateMaterialBrush Brush;
+	bool bIsPopout;
+	bool bPivotEditMode;
+	FVector2f CurrentAbsoluteSize;
+	FVector2f CurrentAbsoluteCenter;
+	EScrubbingMode ScrubbingMode;
+	FVector2f ScrubbingStartAbsoluteCenter;
+	FVector2f ScrubbingStartAbsoluteMouse;
+	EHandleAxis HandleAxis;
+	FVector2D ValueStart;
+	bool bInvertScale;
+	TSharedPtr<FScopedTransaction> ScrubbingTransaction;
+
+	void OnStageUpdated(UDMMaterialComponent* InComponent, EDMUpdateType InUpdateType);
+
+	bool HasValidGeometry() const;
+
+	float GetCircleHandleBaseRadius() const;
+
+	FVector2f ApplyTextureUVTransform(const FVector2f& InUV) const;
+
+	FVector2f GetCenterHandleLocation(const FVector2f& InSize) const;
+
+	FVector2f GetPivotLocation(const FVector2f& InSize) const;
+
+	FVector2f GetAbsoluteCenterHandleLocation() const;
+
+	FVector2f GetAbsolutePivotLocation() const;
+
+	/** Degrees clockwise from +Y axis */
+	float GetCircleHandleRadiusAtAngle(float InAngle) const;
+
+	EHandleAxis GetCenterHandleAxis(const FVector2f& InAbsolutePosition) const;
+
+	EHandleAxis GetCircleHandleAxis(const FVector2f& InAbsolutePosition) const;
+
+	bool TryClickCenterHandle(const FVector2f& InMousePosition, bool bInResetToDefault);
+
+	bool TryClickCircleHandle(const FVector2f& InMousePosition, bool bInResetToDefault);
+
+	void SetScrubbingMode(EScrubbingMode InMode, EHandleAxis InAxis);
+
+	void UpdateScrub();
+	void UpdateScrub_Offset();
+	void UpdateScrub_Rotation();
+	void UpdateScrub_Scale();
+	void UpdateScrub_Pivot();
+};
diff --git a/Engine/Plugins/Experimental/DynamicMaterial/Source/DynamicMaterialEditor/Private/Slate/Properties/SDMTextureUVVisualizerProperty.cpp b/Engine/Plugins/Experimental/DynamicMaterial/Source/DynamicMaterialEditor/Private/Slate/Properties/SDMTextureUVVisualizerProperty.cpp
new file mode 100644
--- /dev/null
+++ b/Engine/Plugins/Experimental/DynamicMaterial/Source/DynamicMaterialEditor/Private/Slate/Properties/SDMTextureUVVisualizerProperty.cpp
@@ -0,0 +1,138 @@
+// Copyright Epic Games, Inc. All Rights Reserved.
+
+#include "Slate/Properties/SDMTextureUVVisualizerProperty.h"
+#include "DetailLayoutBuilder.h"
+#include "DynamicMaterialEditorSettings.h"
+#include "Slate/Properties/SDMTextureUVVisualizer.h"
+#include "Widgets/Input/SButton.h"
+#include "Widgets/Layout/SBox.h"
+#include "Widgets/Text/STextBlock.h"
+
+#define LOCTEXT_NAMESPACE "SDMTextureUVVisualizerProperty"
+
+void SDMTextureUVVisualizerProperty::Construct(const FArguments& InArgs, UDMMaterialStage* InMaterialStage, UDMTextureUV* InTextureUV)
+{
+	check(InMaterialStage);
+	check(InTextureUV);
+
+	ChildSlot
+	[
+		SNew(SVerticalBox)
+		+ SVerticalBox::Slot()
+		.AutoHeight()
+		.HAlign(EHorizontalAlignment::HAlign_Fill)
+		[
+			SNew(SHorizontalBox)
+			.Visibility(this, &SDMTextureUVVisualizerProperty::GetVisualizerVisibility)
+			+ SHorizontalBox::Slot()
+			.FillContentWidth(1.f)
+			[
+				SNew(SButton)
+				.OnClicked(this, &SDMTextureUVVisualizerProperty::OnOpenPopoutClicked)
+				.IsEnabled(false)
+				.Content()
+				[
+					SNew(STextBlock)
+					.Text(LOCTEXT("PopoutVisualizer", "Popout"))
+					.Font(IDetailLayoutBuilder::GetDetailFontBold())
+				]
+			]
+			+ SHorizontalBox::Slot()
+			.FillContentWidth(1.f)
+			[
+				SNew(SButton)
+				.OnClicked(this, &SDMTextureUVVisualizerProperty::OnToggleModeClicked)
+				.Content()
+				[
+					SNew(STextBlock)
+					.Text(this, &SDMTextureUVVisualizerProperty::GetModeButtonText)
+					.Font(IDetailLayoutBuilder::GetDetailFontBold())
+				]
+			]
+		]
+		+ SVerticalBox::Slot()
+		.AutoHeight()
+		.HAlign(EHorizontalAlignment::HAlign_Left)
+		.VAlign(EVerticalAlignment::VAlign_Top)
+		.Padding(0.f, 3.f, 0.f, 0.f)
+		[
+			SNew(SBox)
+			.Visibility(this, &SDMTextureUVVisualizerProperty::GetVisualizerVisibility)
+			.MinAspectRatio(1)
+			.MaxAspectRatio(1)
+			[
+				SAssignNew(Visualizer, SDMTextureUVVisualizer, InMaterialStage, InTextureUV)
+					.IsPopout(false)
+			]			
+		]
+		+ SVerticalBox::Slot()
+		.AutoHeight()
+		.HAlign(EHorizontalAlignment::HAlign_Left)
+		.Padding(0.f, 3.f, 0.f, 0.f)
+		[
+			SNew(SHorizontalBox)
+			+ SHorizontalBox::Slot()
+			[
+				SNew(SButton)
+				.OnClicked(this, &SDMTextureUVVisualizerProperty::OnToggleVisualizerClicked)
+				.Content()
+				[
+					SNew(STextBlock)
+					.Text(LOCTEXT("ToggleVisualizer", "Toggle"))
+					.Font(IDetailLayoutBuilder::GetDetailFontBold())
+				]
+			]
+		]
+	];
+}
+
+
+FReply SDMTextureUVVisualizerProperty::OnToggleVisualizerClicked()
+{
+	if (UDynamicMaterialEditorSettings* Settings = GetMutableDefault<UDynamicMaterialEditorSettings>())
+	{
+		Settings->bUVVisualizerVisible = !Settings->bUVVisualizerVisible;
+		Settings->SaveConfig();
+	}
+
+	return FReply::Handled();
+}
+
+FReply SDMTextureUVVisualizerProperty::OnToggleModeClicked()
+{
+	if (Visualizer.IsValid())
+	{
+		Visualizer->TogglePivotEditMode();
+	}
+
+	return FReply::Handled();
+}
+
+FText SDMTextureUVVisualizerProperty::GetModeButtonText() const
+{
+	if (Visualizer.IsValid() && Visualizer->IsInPivotEditMode())
+	{
+		return LOCTEXT("VisualizerPivot", "Pivot");
+	}
+
+	return LOCTEXT("VisualizerOffset", "Offset");
+}
+
+FReply SDMTextureUVVisualizerProperty::OnOpenPopoutClicked()
+{
+	return FReply::Handled();
+}
+
+EVisibility SDMTextureUVVisualizerProperty::GetVisualizerVisibility() const
+{
+	if (const UDynamicMaterialEditorSettings* Settings = GetDefault<UDynamicMaterialEditorSettings>())
+	{
+		return Settings->bUVVisualizerVisible
+			? EVisibility::Visible
+			: EVisibility::Collapsed;
+	}
+
+	return EVisibility::Collapsed;
+}
+
+#undef LOCTEXT_NAMESPACE
diff --git a/Engine/Plugins/Experimental/DynamicMaterial/Source/DynamicMaterialEditor/Private/Slate/Properties/SDMTextureUVVisualizerProperty.h b/Engine/Plugins/Experimental/DynamicMaterial/Source/DynamicMaterialEditor/Private/Slate/Properties/SDMTextureUVVisualizerProperty.h
new file mode 100644
--- /dev/null
+++ b/Engine/Plugins/Experimental/DynamicMaterial/Source/DynamicMaterialEditor/Private/Slate/Properties/SDMTextureUVVisualizerProperty.h
@@ -0,0 +1,38 @@
+// Copyright Epic Games, Inc. All Rights Reserved.
+
+#pragma once
+
+#include "Widgets/SCompoundWidget.h"
+
+class SDMTextureUVVisualizer;
+class UDMMaterialStage;
+class UDMTextureUV;
+
+/**
+ * Material Designer Texture UV Visualizer Property
+ *
+ * Houses a Texture UV editor and a few buttons to control it.
+ */
+class SDMTextureUVVisualizerProperty : public SCompoundWidget
+{
+public:
+
+	SLATE_BEGIN_ARGS(SDMTextureUVVisualizerProperty) {}
+	SLATE_END_ARGS()
+
+	/** The TextureUV should be a sub-property of the stage */
+	void Construct(const FArguments& InArgs, UDMMaterialStage* InMaterialStage, UDMTextureUV* InTextureUV);
+
+protected:
+	TSharedPtr<SDMTextureUVVisualizer> Visualizer;
+
+	FReply OnToggleVisualizerClicked();
+
+	FReply OnToggleModeClicked();
+
+	FText GetModeButtonText() const;
+
+	FReply OnOpenPopoutClicked();
+
+	EVisibility GetVisualizerVisibility() const;
+};
diff --git a/Engine/Plugins/Experimental/DynamicMaterial/Source/DynamicMaterialEditor/Private/Slate/SDMComponentEdit.cpp b/Engine/Plugins/Experimental/DynamicMaterial/Source/DynamicMaterialEditor/Private/Slate/SDMComponentEdit.cpp
--- a/Engine/Plugins/Experimental/DynamicMaterial/Source/DynamicMaterialEditor/Private/Slate/SDMComponentEdit.cpp
+++ b/Engine/Plugins/Experimental/DynamicMaterial/Source/DynamicMaterialEditor/Private/Slate/SDMComponentEdit.cpp
@@ -308,20 +308,44 @@ TSharedRef<SWidget> SDMComponentEdit::CreateEditWidget()
 
 	for (const FDMPropertyHandle& EditRow : EditRows)
 	{
-		if (!EditRow.DetailTreeNode)
+		const bool bHasValidCustomWidget = EditRow.ValueWidget.IsValid() && !EditRow.ValueName.IsNone() && EditRow.NameOverride.IsSet();
+
+		if (!EditRow.DetailTreeNode && !bHasValidCustomWidget)
 		{
 			continue;
 		}
 
 		ECustomDetailsTreeInsertPosition Position = ECustomDetailsTreeInsertPosition::Child;
 
-		if (EditRow.DetailTreeNode->CreatePropertyHandle()->HasMetaData("HighPriority"))
+		if (EditRow.DetailTreeNode)
 		{
-			Position = ECustomDetailsTreeInsertPosition::FirstChild;
+			if (EditRow.DetailTreeNode->CreatePropertyHandle()->HasMetaData("HighPriority"))
+			{
+				Position = ECustomDetailsTreeInsertPosition::FirstChild;
+			}
+			else if (EditRow.DetailTreeNode->CreatePropertyHandle()->HasMetaData("LowPriority"))
+			{
+				Position = ECustomDetailsTreeInsertPosition::LastChild;
+			}
 		}
-		else if (EditRow.DetailTreeNode->CreatePropertyHandle()->HasMetaData("LowPriority"))
+
+		if (bHasValidCustomWidget)
 		{
-			Position = ECustomDetailsTreeInsertPosition::LastChild;
+			TSharedPtr<ICustomDetailsViewCustomItem> Item = DetailsView->CreateCustomItem(EditRow.ValueName, EditRow.NameOverride.GetValue(), EditRow.NameToolTipOverride.GetValue());
+
+			if (!Item.IsValid())
+			{
+				continue;
+			}
+
+			Item->SetValueWidget(EditRow.ValueWidget.ToSharedRef());
+			DetailsView->ExtendTree(RootId, Position, Item->AsItem());
+			continue;
+		}
+
+		if (!EditRow.DetailTreeNode)
+		{
+			continue;
 		}
 
 		TSharedRef<ICustomDetailsViewItem> Item = DetailsView->CreateDetailTreeItem(EditRow.DetailTreeNode.ToSharedRef());
diff --git a/Engine/Plugins/Experimental/DynamicMaterial/Source/DynamicMaterialEditor/Public/DMEDefs.h b/Engine/Plugins/Experimental/DynamicMaterial/Source/DynamicMaterialEditor/Public/DMEDefs.h
--- a/Engine/Plugins/Experimental/DynamicMaterial/Source/DynamicMaterialEditor/Public/DMEDefs.h
+++ b/Engine/Plugins/Experimental/DynamicMaterial/Source/DynamicMaterialEditor/Public/DMEDefs.h
@@ -120,6 +120,8 @@ struct FDMPropertyHandle
 	TOptional<FText> NameOverride;
 	TOptional<FText> NameToolTipOverride;
 	TOptional<FResetToDefaultOverride> ResetToDefaultOverride;
+	TSharedPtr<SWidget> ValueWidget;
+	FName ValueName;
 };
 
 UENUM(BlueprintType, meta = (Bitflags, UseEnumValuesAsMaskValuesInEditor = "true"))
