Material Designer
- The popup UV editor now works.

diff --git a/Engine/Plugins/Experimental/DynamicMaterial/Source/DynamicMaterialEditor/Private/Slate/Properties/SDMTextureUVVisualizer.cpp b/Engine/Plugins/Experimental/DynamicMaterial/Source/DynamicMaterialEditor/Private/Slate/Properties/SDMTextureUVVisualizer.cpp
--- a/Engine/Plugins/Experimental/DynamicMaterial/Source/DynamicMaterialEditor/Private/Slate/Properties/SDMTextureUVVisualizer.cpp
+++ b/Engine/Plugins/Experimental/DynamicMaterial/Source/DynamicMaterialEditor/Private/Slate/Properties/SDMTextureUVVisualizer.cpp
@@ -51,6 +51,7 @@ SDMTextureUVVisualizer::SDMTextureUVVisualizer()
 	, ValueStart(FVector2D::ZeroVector)
 	, bInvertScale(false)
 {
+	Brush.SetUVRegion(FBox2f(FVector2f(0, 0), FVector2f(1, 1)));
 }
 
 void SDMTextureUVVisualizer::Construct(const FArguments& InArgs, UDMMaterialStage* InMaterialStage, UDMTextureUV* InTextureUV)
@@ -97,6 +98,16 @@ void SDMTextureUVVisualizer::TogglePivotEditMode()
 	SetInPivotEditMode(!IsInPivotEditMode());
 }
 
+UDMMaterialStage* SDMTextureUVVisualizer::GetStage() const
+{
+	return StageWeak.Get();
+}
+
+UDMTextureUV* SDMTextureUVVisualizer::GetTextureUV() const
+{
+	return TextureUVWeak.Get();
+}
+
 void SDMTextureUVVisualizer::Tick(const FGeometry& InAllottedGeometry, const double InCurrentTime, const float InDeltaTime)
 {
 	SCompoundWidget::Tick(InAllottedGeometry, InCurrentTime, InDeltaTime);
@@ -136,6 +147,11 @@ void SDMTextureUVVisualizer::Tick(const FGeometry& InAllottedGeometry, const dou
 	{
 		ScrubbingTransaction.Reset();
 	}
+
+	if (bIsPopout)
+	{
+		UpdatePopoutUVs();
+	}
 }
 
 FReply SDMTextureUVVisualizer::OnMouseButtonDown(const FGeometry& InGeometry, const FPointerEvent& InMouseEvent)
@@ -237,7 +253,7 @@ int32 SDMTextureUVVisualizer::OnPaint(const FPaintArgs& InArgs, const FGeometry&
 
 	const float Rotation = TextureUV->GetRotation();
 	const float RotationRadians = FMath::DegreesToRadians(Rotation);
-	const FVector2f LocalCenterOffset = GetCenterHandleLocation(InAllottedGeometry.GetLocalSize());
+	const FVector2f LocalCenterOffset = GetOffsetLocation(InAllottedGeometry.GetLocalSize());
 	const FVector2f LocalPivotOffset = GetPivotLocation(InAllottedGeometry.GetLocalSize());
 
 	auto DrawRotatedBox = [this, &OutDrawElements, &InLayerId, &InAllottedGeometry, &Rotation, &LocalCenterOffset, &LocalPivotOffset]
@@ -245,6 +261,7 @@ int32 SDMTextureUVVisualizer::OnPaint(const FPaintArgs& InArgs, const FGeometry&
 		{
 			const FVector2f& BaseLocation = bPivotEditMode ? LocalPivotOffset : LocalCenterOffset;
 			const FVector2f Location = BaseLocation + InLocation - InSize * 0.5f;
+
 			const FVector2f RotationOffset = BaseLocation - Location;
 
 			FSlateDrawElement::MakeRotatedBox(
@@ -363,7 +380,7 @@ int32 SDMTextureUVVisualizer::OnPaint(const FPaintArgs& InArgs, const FGeometry&
 
 		auto DrawPartialCicle = [this, &OutDrawElements, &InLayerId, &InAllottedGeometry, &AngleRadiusToLocation](int32 InStartIndex)
 			{
-				TArray<FVector2f> Points = {
+				const TArray<FVector2f> Points = {
 					AngleRadiusToLocation(InStartIndex),
 					AngleRadiusToLocation(InStartIndex + 1),
 					AngleRadiusToLocation(InStartIndex + 2)
@@ -474,33 +491,38 @@ FVector2f SDMTextureUVVisualizer::ApplyTextureUVTransform(const FVector2f& InUV)
 	return TransformedUV;
 }
 
-FVector2f SDMTextureUVVisualizer::GetCenterHandleLocation(const FVector2f& InSize) const
+FVector2f SDMTextureUVVisualizer::GetOffsetLocation(const FVector2f& InSize) const
 {
-	if (bPivotEditMode)
+	static const FVector2f Center = {0.5f, 0.5f};
+
+	if (bIsPopout)
 	{
-		return GetPivotLocation(InSize);
+		return ToPopoutLocation(FVector2f::UnitVector, ApplyTextureUVTransform(Center)) * InSize;
 	}
 
-	static const FVector2f Center = {0.5f, 0.5f};
-
 	return ApplyTextureUVTransform(Center) * InSize;
 }
 
 FVector2f SDMTextureUVVisualizer::GetPivotLocation(const FVector2f& InSize) const
 {
-	static const FVector2f Center = {0.5f, 0.5f};
-
 	if (UDMTextureUV* TextureUV = TextureUVWeak.Get())
 	{
+		if (bIsPopout)
+		{
+			return ToPopoutLocation(FVector2f::UnitVector, static_cast<FVector2f>(TextureUV->GetPivot())) * InSize;
+		}
+
 		return static_cast<FVector2f>(TextureUV->GetPivot()) * InSize;
 	}
 
+	static const FVector2f Center = {0.5f, 0.5f};
+
 	return Center * InSize;
 }
 
-FVector2f SDMTextureUVVisualizer::GetAbsoluteCenterHandleLocation() const
+FVector2f SDMTextureUVVisualizer::GetAbsoluteOffsetLocation() const
 {
-	return CurrentAbsoluteCenter - (CurrentAbsoluteSize * 0.5f) + GetCenterHandleLocation(CurrentAbsoluteSize);
+	return CurrentAbsoluteCenter - (CurrentAbsoluteSize * 0.5f) + GetOffsetLocation(CurrentAbsoluteSize);
 }
 
 FVector2f SDMTextureUVVisualizer::GetAbsolutePivotLocation() const
@@ -553,37 +575,38 @@ SDMTextureUVVisualizer::EHandleAxis SDMTextureUVVisualizer::GetCenterHandleAxis(
 
 	using namespace UE::DynamicMaterialEditor::Private;
 
-	FVector2f CenterOffset = InAbsolutePosition - GetAbsoluteCenterHandleLocation();
-
+	const FVector2f HandleLocation = bPivotEditMode ? GetAbsolutePivotLocation() : GetAbsoluteOffsetLocation();
+	FVector2f HandleOffset = InAbsolutePosition - HandleLocation;
+		
 	if (!bPivotEditMode)
 	{
 		const float Rotation = TextureUV->GetRotation();
 
 		if (!FMath::IsNearlyZero(Rotation))
 		{
-			CenterOffset = CenterOffset.GetRotated(-Rotation);
+			HandleOffset = HandleOffset.GetRotated(-Rotation);
 		}
 	}
 
-	CenterOffset.X = FMath::Abs(CenterOffset.X);
-	CenterOffset.Y = FMath::Abs(CenterOffset.Y);
+	HandleOffset.X = FMath::Abs(HandleOffset.X);
+	HandleOffset.Y = FMath::Abs(HandleOffset.Y);
 
-	if (CenterOffset.X <= TextureUVVisualizerSmallRadius && CenterOffset.Y <= TextureUVVisualizerSmallRadius)
+	if (HandleOffset.X <= TextureUVVisualizerSmallRadius && HandleOffset.Y <= TextureUVVisualizerSmallRadius)
 	{
 		return EHandleAxis::XY;
 	}
 
-	if (CenterOffset.X <= TextureUVVisualizerSmallRadius && CenterOffset.Y <= TextureUVVisualizerLargeRadius)
+	if (HandleOffset.X <= TextureUVVisualizerSmallRadius && HandleOffset.Y <= TextureUVVisualizerLargeRadius)
 	{
 		return EHandleAxis::Y;
 	}
 
-	if (CenterOffset.X <= TextureUVVisualizerLargeRadius && CenterOffset.Y <= TextureUVVisualizerSmallRadius)
+	if (HandleOffset.X <= TextureUVVisualizerLargeRadius && HandleOffset.Y <= TextureUVVisualizerSmallRadius)
 	{
 		return EHandleAxis::X;
 	}
 
-	if (CenterOffset.X <= TextureUVVisualizerLargeRadius && CenterOffset.Y <= TextureUVVisualizerLargeRadius)
+	if (HandleOffset.X <= TextureUVVisualizerLargeRadius && HandleOffset.Y <= TextureUVVisualizerLargeRadius)
 	{
 		return EHandleAxis::XY;
 	}
@@ -602,20 +625,20 @@ SDMTextureUVVisualizer::EHandleAxis SDMTextureUVVisualizer::GetCircleHandleAxis(
 
 	using namespace UE::DynamicMaterialEditor::Private;
 
-	const FVector2f CenterOffset = InAbsolutePosition - GetAbsolutePivotLocation();
-	const float DistanceFromCenter = CenterOffset.Size();
+	const FVector2f HandleOffset = InAbsolutePosition - GetAbsolutePivotLocation();
+	const float DistanceFromHandle = HandleOffset.Size();
 	const FVector2D& Scale2D = TextureUV->GetScale();
 	// We manage angle clockwise from +Y axis. Atan2 handles it anti-clockwise from +X axis.
-	float Angle = 90.f - FMath::RadiansToDegrees(FMath::Atan2(CenterOffset.Y, CenterOffset.X)) + TextureUV->GetRotation();
+	float Angle = 90.f - FMath::RadiansToDegrees(FMath::Atan2(HandleOffset.Y, HandleOffset.X)) + TextureUV->GetRotation();
 	float DistanceFromCircleHandle;
 
 	if (FMath::IsNearlyEqual(Scale2D.X, Scale2D.Y))
 	{
-		DistanceFromCircleHandle = FMath::Abs(DistanceFromCenter - GetCircleHandleBaseRadius());
+		DistanceFromCircleHandle = FMath::Abs(DistanceFromHandle - GetCircleHandleBaseRadius());
 	}
 	else
 	{
-		DistanceFromCircleHandle = FMath::Abs(DistanceFromCenter - GetCircleHandleRadiusAtAngle(Angle));
+		DistanceFromCircleHandle = FMath::Abs(DistanceFromHandle - GetCircleHandleRadiusAtAngle(Angle));
 	}
 
 	if (DistanceFromCircleHandle > TextureUVVisualizerCircleHandleRadius)
@@ -836,6 +859,29 @@ bool SDMTextureUVVisualizer::TryClickCircleHandle(const FVector2f& InMousePositi
 	return true;
 }
 
+void SDMTextureUVVisualizer::UpdatePopoutUVs()
+{
+	FBox2d UVRegion = FBox2d(FVector2D(-1, -1), FVector2D(2, 2));
+
+	if (!FMath::IsNearlyZero(CurrentAbsoluteSize.X) && !FMath::IsNearlyZero(CurrentAbsoluteSize.Y))
+	{
+		if (CurrentAbsoluteSize.Y <= CurrentAbsoluteSize.X)
+		{
+			const float Multiplier = CurrentAbsoluteSize.X / CurrentAbsoluteSize.Y;
+			UVRegion.Min.X = 0.5f - (Multiplier * 1.5f);
+			UVRegion.Max.X = 0.5f + (Multiplier * 1.5f);
+		}
+		else
+		{
+			const float Multiplier = CurrentAbsoluteSize.Y / CurrentAbsoluteSize.X;
+			UVRegion.Min.Y = 0.5f - (Multiplier * 1.5f);
+			UVRegion.Max.Y = 0.5f + (Multiplier * 1.5f);
+		}
+	}
+
+	Brush.SetUVRegion(UVRegion);
+}
+
 void SDMTextureUVVisualizer::SetScrubbingMode(EScrubbingMode InMode, EHandleAxis InAxis)
 {
 	UDMTextureUV* TextureUV = TextureUVWeak.Get();
@@ -894,6 +940,44 @@ void SDMTextureUVVisualizer::SetScrubbingMode(EScrubbingMode InMode, EHandleAxis
 	TextureUV->Modify();
 }
 
+FVector2f SDMTextureUVVisualizer::ToPopoutLocation(const FVector2f& InSize, FVector2f&& InLocation) const
+{
+	const FBox2d BrushUV = Brush.GetUVRegion();
+
+	InLocation.X = FMath::GetMappedRangeValueUnclamped<float, float>(
+		FVector2f(BrushUV.Min.X, BrushUV.Max.X),
+		{0, 1},
+		InLocation.X / InSize.X
+	) * InSize.X;
+
+	InLocation.Y = FMath::GetMappedRangeValueUnclamped<float, float>(
+		FVector2f(BrushUV.Min.Y, BrushUV.Max.Y),
+		{0, 1},
+		InLocation.Y / InSize.Y
+	) * InSize.Y;
+
+	return InLocation;
+}
+
+FVector2f SDMTextureUVVisualizer::FromPopoutLocation(const FVector2f& InSize, FVector2f&& InLocation) const
+{
+	const FBox2d BrushUV = Brush.GetUVRegion();
+
+	InLocation.X = FMath::GetMappedRangeValueUnclamped<float, float>(
+		{0, 1},
+		FVector2f(BrushUV.Min.X, BrushUV.Max.X),
+		InLocation.X / InSize.X
+	) * InSize.X;
+
+	InLocation.Y = FMath::GetMappedRangeValueUnclamped<float, float>(
+		{0, 1},
+		FVector2f(BrushUV.Min.Y, BrushUV.Max.Y),
+		InLocation.Y / InSize.Y
+	) * InSize.Y;
+
+	return InLocation;
+}
+
 void SDMTextureUVVisualizer::UpdateScrub()
 {
 	switch (ScrubbingMode)
@@ -934,8 +1018,11 @@ void SDMTextureUVVisualizer::UpdateScrub_Offset()
 	}
 
 	const FVector2f MouseOffset = FSlateApplication::Get().GetCursorPos() - ScrubbingStartAbsoluteMouse;
+	const FBox2d BrushUV = Brush.GetUVRegion();
 
 	FVector2D OffsetChange = static_cast<FVector2D>(MouseOffset / CurrentAbsoluteSize);
+	OffsetChange.X *= (BrushUV.Max.X - BrushUV.Min.X);
+	OffsetChange.Y *= (BrushUV.Max.Y - BrushUV.Min.Y);
 
 	const float Rotation = TextureUV->GetRotation();
 
@@ -990,8 +1077,11 @@ void SDMTextureUVVisualizer::UpdateScrub_Scale()
 	}
 
 	const FVector2f MouseOffset = FSlateApplication::Get().GetCursorPos() - ScrubbingStartAbsoluteMouse;
+	const FBox2d BrushUV = Brush.GetUVRegion();
 
 	FVector2D ScaleChange = static_cast<FVector2D>(MouseOffset / CurrentAbsoluteSize);
+	ScaleChange.X *= (BrushUV.Max.X - BrushUV.Min.X);
+	ScaleChange.Y *= (BrushUV.Max.Y - BrushUV.Min.Y);
 	ScaleChange /= TextureUV->GetScale();
 
 	if (bInvertScale)
@@ -1058,8 +1148,11 @@ void SDMTextureUVVisualizer::UpdateScrub_Pivot()
 	}
 
 	const FVector2f MouseOffset = FSlateApplication::Get().GetCursorPos() - ScrubbingStartAbsoluteMouse;
+	const FBox2d BrushUV = Brush.GetUVRegion();
 
 	FVector2D PivotChange = static_cast<FVector2D>(MouseOffset / CurrentAbsoluteSize);
+	PivotChange.X *= (BrushUV.Max.X - BrushUV.Min.X);
+	PivotChange.Y *= (BrushUV.Max.Y - BrushUV.Min.Y);
 
 	if (HandleAxis == EHandleAxis::X)
 	{
diff --git a/Engine/Plugins/Experimental/DynamicMaterial/Source/DynamicMaterialEditor/Private/Slate/Properties/SDMTextureUVVisualizer.h b/Engine/Plugins/Experimental/DynamicMaterial/Source/DynamicMaterialEditor/Private/Slate/Properties/SDMTextureUVVisualizer.h
--- a/Engine/Plugins/Experimental/DynamicMaterial/Source/DynamicMaterialEditor/Private/Slate/Properties/SDMTextureUVVisualizer.h
+++ b/Engine/Plugins/Experimental/DynamicMaterial/Source/DynamicMaterialEditor/Private/Slate/Properties/SDMTextureUVVisualizer.h
@@ -60,6 +60,10 @@ public:
 
 	void TogglePivotEditMode();
 
+	UDMMaterialStage* GetStage() const;
+
+	UDMTextureUV* GetTextureUV() const;
+
 	//~ Begin SWidget
 	virtual void Tick(const FGeometry& InAllottedGeometry, const double InCurrentTime, const float InDeltaTime) override;
 	virtual FReply OnMouseButtonDown(const FGeometry& InGeometry, const FPointerEvent& InMouseEvent) override;
@@ -94,11 +98,11 @@ protected:
 
 	FVector2f ApplyTextureUVTransform(const FVector2f& InUV) const;
 
-	FVector2f GetCenterHandleLocation(const FVector2f& InSize) const;
+	FVector2f GetOffsetLocation(const FVector2f& InSize) const;
 
 	FVector2f GetPivotLocation(const FVector2f& InSize) const;
 
-	FVector2f GetAbsoluteCenterHandleLocation() const;
+	FVector2f GetAbsoluteOffsetLocation() const;
 
 	FVector2f GetAbsolutePivotLocation() const;
 
@@ -113,8 +117,14 @@ protected:
 
 	bool TryClickCircleHandle(const FVector2f& InMousePosition, bool bInResetToDefault);
 
+	void UpdatePopoutUVs();
+
 	void SetScrubbingMode(EScrubbingMode InMode, EHandleAxis InAxis);
 
+	FVector2f ToPopoutLocation(const FVector2f& InSize, FVector2f&& InLocation) const;
+
+	FVector2f FromPopoutLocation(const FVector2f& InSize, FVector2f&& InLocation) const;
+
 	void UpdateScrub();
 	void UpdateScrub_Offset();
 	void UpdateScrub_Rotation();
diff --git a/Engine/Plugins/Experimental/DynamicMaterial/Source/DynamicMaterialEditor/Private/Slate/Properties/SDMTextureUVVisualizerPopout.cpp b/Engine/Plugins/Experimental/DynamicMaterial/Source/DynamicMaterialEditor/Private/Slate/Properties/SDMTextureUVVisualizerPopout.cpp
new file mode 100644
--- /dev/null
+++ b/Engine/Plugins/Experimental/DynamicMaterial/Source/DynamicMaterialEditor/Private/Slate/Properties/SDMTextureUVVisualizerPopout.cpp
@@ -0,0 +1,161 @@
+// Copyright Epic Games, Inc. All Rights Reserved.
+
+#include "Slate/Properties/SDMTextureUVVisualizerPopout.h"
+#include "DetailLayoutBuilder.h"
+#include "Slate/Properties/SDMTextureUVVisualizer.h"
+#include "Widgets/Colors/SColorBlock.h"
+#include "Widgets/Input/SButton.h"
+#include "Widgets/Layout/SBox.h"
+#include "Widgets/Text/STextBlock.h"
+
+#define LOCTEXT_NAMESPACE "SDMTextureUVVisualizerPopout"
+
+void SDMTextureUVVisualizerPopout::Construct(const FArguments& InArgs, UDMMaterialStage* InMaterialStage, UDMTextureUV* InTextureUV)
+{
+	check(InMaterialStage);
+	check(InTextureUV);
+
+	SWindow::Construct(
+		SWindow::FArguments()
+		.Title(LOCTEXT("WindowTitle", "Material Designer Texture UV Visualizer"))
+		.ClientSize(FVector2f(1024, 768))
+		.MinWidth(128)
+		.MinHeight(160)
+	);
+
+	SetContent(
+		SNew(SVerticalBox)
+		+ SVerticalBox::Slot()
+		.AutoHeight()
+		.HAlign(EHorizontalAlignment::HAlign_Fill)
+		[
+			SNew(SButton)
+			.OnClicked(this, &SDMTextureUVVisualizerPopout::OnToggleModeClicked)
+			.Content()
+			[
+				SNew(STextBlock)
+				.Text(this, &SDMTextureUVVisualizerPopout::GetModeButtonText)
+				.Font(IDetailLayoutBuilder::GetDetailFontBold())
+			]
+		]
+		+ SVerticalBox::Slot()
+		.VAlign(EVerticalAlignment::VAlign_Fill)
+		.HAlign(EHorizontalAlignment::HAlign_Fill)
+		.VAlign(EVerticalAlignment::VAlign_Fill)
+		.Padding(0.f, 3.f, 0.f, 0.f)
+		[
+			SNew(SOverlay)
+			+ SOverlay::Slot()
+			[
+				SAssignNew(Visualizer, SDMTextureUVVisualizer, InMaterialStage, InTextureUV)
+				.IsPopout(true)
+			]
+			+ SOverlay::Slot()
+			[
+				SNew(SBox)
+				.HAlign(EHorizontalAlignment::HAlign_Fill)
+				.VAlign(EVerticalAlignment::VAlign_Top)
+				[
+					SNew(SColorBlock)
+					.Color(FLinearColor(0, 0, 0, 0.5))
+					.Size(this, &SDMTextureUVVisualizerPopout::GetHorizontalBarSize)
+				]
+			]
+			+ SOverlay::Slot()
+			[
+				SNew(SBox)
+				.HAlign(EHorizontalAlignment::HAlign_Fill)
+				.VAlign(EVerticalAlignment::VAlign_Bottom)
+				[
+					SNew(SColorBlock)
+					.Color(FLinearColor(0, 0, 0, 0.5))
+					.Size(this, &SDMTextureUVVisualizerPopout::GetHorizontalBarSize)
+				]
+			]
+			+ SOverlay::Slot()
+			[
+				SNew(SBox)
+				.HAlign(EHorizontalAlignment::HAlign_Left)
+				.VAlign(EVerticalAlignment::VAlign_Center)
+				[
+					SNew(SColorBlock)
+					.Color(FLinearColor(0, 0, 0, 0.5))
+					.Size(this, &SDMTextureUVVisualizerPopout::GetSideBlockSize)
+				]
+			]
+			+ SOverlay::Slot()
+			[
+				SNew(SBox)
+				.HAlign(EHorizontalAlignment::HAlign_Right)
+				.VAlign(EVerticalAlignment::VAlign_Center)
+				[
+					SNew(SColorBlock)
+					.Color(FLinearColor(0, 0, 0, 0.5))
+					.Size(this, &SDMTextureUVVisualizerPopout::GetSideBlockSize)
+				]
+			]
+		]
+	);
+}
+
+FReply SDMTextureUVVisualizerPopout::OnToggleModeClicked()
+{
+	if (Visualizer.IsValid())
+	{
+		Visualizer->TogglePivotEditMode();
+	}
+
+	return FReply::Handled();
+}
+
+FText SDMTextureUVVisualizerPopout::GetModeButtonText() const
+{
+	if (Visualizer.IsValid() && Visualizer->IsInPivotEditMode())
+	{
+		return LOCTEXT("VisualizerPivot", "Pivot");
+	}
+
+	return LOCTEXT("VisualizerOffset", "Offset");
+}
+
+FVector2D SDMTextureUVVisualizerPopout::GetHorizontalBarSize() const
+{
+	if (Visualizer.IsValid())
+	{
+		const FVector2f AbsoluteSize = Visualizer->GetTickSpaceGeometry().GetLocalSize();
+
+		if (!FMath::IsNearlyZero(AbsoluteSize.X) && !FMath::IsNearlyZero(AbsoluteSize.Y))
+		{
+			if (AbsoluteSize.Y <= AbsoluteSize.X)
+			{
+				return FVector2D(1, AbsoluteSize.Y / 3.f);
+			}
+
+			return FVector2D(1, AbsoluteSize.X / 3.f + (AbsoluteSize.Y - AbsoluteSize.X) * 0.5f);
+		}
+	}
+
+	return FVector2D::UnitVector;
+}
+
+FVector2D SDMTextureUVVisualizerPopout::GetSideBlockSize() const
+{
+	if (Visualizer.IsValid())
+	{
+		const FVector2f AbsoluteSize = Visualizer->GetTickSpaceGeometry().GetLocalSize();
+
+		if (!FMath::IsNearlyZero(AbsoluteSize.X) && !FMath::IsNearlyZero(AbsoluteSize.Y))
+		{
+			if (AbsoluteSize.X <= AbsoluteSize.Y)
+			{
+				return FVector2D(AbsoluteSize.X / 3.f, AbsoluteSize.Y / 3.f);
+			}
+
+			return FVector2D(AbsoluteSize.Y / 3.f + (AbsoluteSize.X - AbsoluteSize.Y) * 0.5f, AbsoluteSize.Y / 3.f);
+		}
+	}
+
+	return FVector2D::UnitVector;
+}
+
+#undef LOCTEXT_NAMESPACE
diff --git a/Engine/Plugins/Experimental/DynamicMaterial/Source/DynamicMaterialEditor/Private/Slate/Properties/SDMTextureUVVisualizerPopout.h b/Engine/Plugins/Experimental/DynamicMaterial/Source/DynamicMaterialEditor/Private/Slate/Properties/SDMTextureUVVisualizerPopout.h
new file mode 100644
--- /dev/null
+++ b/Engine/Plugins/Experimental/DynamicMaterial/Source/DynamicMaterialEditor/Private/Slate/Properties/SDMTextureUVVisualizerPopout.h
@@ -0,0 +1,39 @@
+// Copyright Epic Games, Inc. All Rights Reserved.
+
+#pragma once
+
+#include "Widgets/SWindow.h"
+
+class SDMTextureUVVisualizer;
+class UDMMaterialStage;
+class UDMTextureUV;
+
+/**
+ * Material Designer Texture UV Visualizer Popout
+ *
+ * Houses a Texture UV editor and a few buttons to control it.
+ *
+ * The popout specifically expands the visible area of the preview to 3x the normal size
+ * on the smallest axis. The other axis is expanded to match the aspect ratio.
+ */
+class SDMTextureUVVisualizerPopout : public SWindow
+{
+public:
+
+	SLATE_BEGIN_ARGS(SDMTextureUVVisualizerPopout) {}
+	SLATE_END_ARGS()
+
+	/** The TextureUV should be a sub-property of the stage */
+	void Construct(const FArguments& InArgs, UDMMaterialStage* InMaterialStage, UDMTextureUV* InTextureUV);
+
+protected:
+	TSharedPtr<SDMTextureUVVisualizer> Visualizer;
+
+	FReply OnToggleModeClicked();
+
+	FText GetModeButtonText() const;
+
+	FVector2D GetHorizontalBarSize() const;
+
+	FVector2D GetSideBlockSize() const;
+};
diff --git a/Engine/Plugins/Experimental/DynamicMaterial/Source/DynamicMaterialEditor/Private/Slate/Properties/SDMTextureUVVisualizerProperty.cpp b/Engine/Plugins/Experimental/DynamicMaterial/Source/DynamicMaterialEditor/Private/Slate/Properties/SDMTextureUVVisualizerProperty.cpp
--- a/Engine/Plugins/Experimental/DynamicMaterial/Source/DynamicMaterialEditor/Private/Slate/Properties/SDMTextureUVVisualizerProperty.cpp
+++ b/Engine/Plugins/Experimental/DynamicMaterial/Source/DynamicMaterialEditor/Private/Slate/Properties/SDMTextureUVVisualizerProperty.cpp
@@ -1,9 +1,12 @@
 // Copyright Epic Games, Inc. All Rights Reserved.
 
 #include "Slate/Properties/SDMTextureUVVisualizerProperty.h"
+#include "Components/DMMaterialStage.h"
+#include "Components/DMTextureUV.h"
 #include "DetailLayoutBuilder.h"
 #include "DynamicMaterialEditorSettings.h"
 #include "Slate/Properties/SDMTextureUVVisualizer.h"
+#include "Slate/Properties/SDMTextureUVVisualizerPopout.h"
 #include "Widgets/Input/SButton.h"
 #include "Widgets/Layout/SBox.h"
 #include "Widgets/Text/STextBlock.h"
@@ -29,7 +32,7 @@ void SDMTextureUVVisualizerProperty::Construct(const FArguments& InArgs, UDMMate
 			[
 				SNew(SButton)
 				.OnClicked(this, &SDMTextureUVVisualizerProperty::OnOpenPopoutClicked)
-				.IsEnabled(false)
+				.IsEnabled(this, &SDMTextureUVVisualizerProperty::IsPopoutEnabled)
 				.Content()
 				[
 					SNew(STextBlock)
@@ -62,7 +65,7 @@ void SDMTextureUVVisualizerProperty::Construct(const FArguments& InArgs, UDMMate
 			.MaxAspectRatio(1)
 			[
 				SAssignNew(Visualizer, SDMTextureUVVisualizer, InMaterialStage, InTextureUV)
-					.IsPopout(false)
+				.IsPopout(false)
 			]			
 		]
 		+ SVerticalBox::Slot()
@@ -118,8 +121,36 @@ FText SDMTextureUVVisualizerProperty::GetModeButtonText() const
 	return LOCTEXT("VisualizerOffset", "Offset");
 }
 
+bool SDMTextureUVVisualizerProperty::IsPopoutEnabled() const
+{
+	return !PopoutWindowWeak.IsValid();
+}
+
 FReply SDMTextureUVVisualizerProperty::OnOpenPopoutClicked()
 {
+	if (!Visualizer.IsValid())
+	{
+		return FReply::Handled();
+	}
+
+	if (PopoutWindowWeak.IsValid())
+	{
+		return FReply::Handled();
+	}
+
+	UDMMaterialStage* Stage = Visualizer->GetStage();
+	UDMTextureUV* TextureUV = Visualizer->GetTextureUV();
+
+	if (!IsValid(Stage) || !IsValid(TextureUV))
+	{
+		return FReply::Handled();
+	}
+
+	TSharedRef<SDMTextureUVVisualizerPopout> PopoutWindow = SNew(SDMTextureUVVisualizerPopout, Stage, TextureUV);
+	FSlateApplication::Get().AddWindow(PopoutWindow, /* Show Immediately */ true);
+
+	PopoutWindowWeak = PopoutWindow;
+
 	return FReply::Handled();
 }
 
diff --git a/Engine/Plugins/Experimental/DynamicMaterial/Source/DynamicMaterialEditor/Private/Slate/Properties/SDMTextureUVVisualizerProperty.h b/Engine/Plugins/Experimental/DynamicMaterial/Source/DynamicMaterialEditor/Private/Slate/Properties/SDMTextureUVVisualizerProperty.h
--- a/Engine/Plugins/Experimental/DynamicMaterial/Source/DynamicMaterialEditor/Private/Slate/Properties/SDMTextureUVVisualizerProperty.h
+++ b/Engine/Plugins/Experimental/DynamicMaterial/Source/DynamicMaterialEditor/Private/Slate/Properties/SDMTextureUVVisualizerProperty.h
@@ -5,6 +5,7 @@
 #include "Widgets/SCompoundWidget.h"
 
 class SDMTextureUVVisualizer;
+class SDMTextureUVVisualizerPopout;
 class UDMMaterialStage;
 class UDMTextureUV;
 
@@ -25,6 +26,7 @@ public:
 
 protected:
 	TSharedPtr<SDMTextureUVVisualizer> Visualizer;
+	TWeakPtr<SDMTextureUVVisualizerPopout> PopoutWindowWeak;
 
 	FReply OnToggleVisualizerClicked();
 
@@ -32,6 +34,7 @@ protected:
 
 	FText GetModeButtonText() const;
 
+	bool IsPopoutEnabled() const;
 	FReply OnOpenPopoutClicked();
 
 	EVisibility GetVisualizerVisibility() const;
