Material Designer
- Small fix to the uv popout window.

diff --git a/Engine/Plugins/Experimental/DynamicMaterial/Source/DynamicMaterialEditor/Private/Slate/Properties/SDMTextureUVVisualizerPopout.cpp b/Engine/Plugins/Experimental/DynamicMaterial/Source/DynamicMaterialEditor/Private/Slate/Properties/SDMTextureUVVisualizerPopout.cpp
--- a/Engine/Plugins/Experimental/DynamicMaterial/Source/DynamicMaterialEditor/Private/Slate/Properties/SDMTextureUVVisualizerPopout.cpp
+++ b/Engine/Plugins/Experimental/DynamicMaterial/Source/DynamicMaterialEditor/Private/Slate/Properties/SDMTextureUVVisualizerPopout.cpp
@@ -122,16 +122,16 @@ FVector2D SDMTextureUVVisualizerPopout::GetHorizontalBarSize() const
 {
 	if (Visualizer.IsValid())
 	{
-		const FVector2f AbsoluteSize = Visualizer->GetTickSpaceGeometry().GetLocalSize();
+		const FVector2f LocalSize = Visualizer->GetTickSpaceGeometry().GetLocalSize();
 
-		if (!FMath::IsNearlyZero(AbsoluteSize.X) && !FMath::IsNearlyZero(AbsoluteSize.Y))
+		if (!FMath::IsNearlyZero(LocalSize.X) && !FMath::IsNearlyZero(LocalSize.Y))
 		{
-			if (AbsoluteSize.Y <= AbsoluteSize.X)
+			if (LocalSize.Y <= LocalSize.X)
 			{
-				return FVector2D(1, AbsoluteSize.Y / 3.f);
+				return FVector2D(1, LocalSize.Y / 3.f);
 			}
 
-			return FVector2D(1, AbsoluteSize.X / 3.f + (AbsoluteSize.Y - AbsoluteSize.X) * 0.5f);
+			return FVector2D(1, LocalSize.X / 3.f + (LocalSize.Y - LocalSize.X) * 0.5f);
 		}
 	}
 
@@ -142,16 +142,16 @@ FVector2D SDMTextureUVVisualizerPopout::GetSideBlockSize() const
 {
 	if (Visualizer.IsValid())
 	{
-		const FVector2f AbsoluteSize = Visualizer->GetTickSpaceGeometry().GetLocalSize();
+		const FVector2f LocalSize = Visualizer->GetTickSpaceGeometry().GetLocalSize();
 
-		if (!FMath::IsNearlyZero(AbsoluteSize.X) && !FMath::IsNearlyZero(AbsoluteSize.Y))
+		if (!FMath::IsNearlyZero(LocalSize.X) && !FMath::IsNearlyZero(LocalSize.Y))
 		{
-			if (AbsoluteSize.X <= AbsoluteSize.Y)
+			if (LocalSize.X <= LocalSize.Y)
 			{
-				return FVector2D(AbsoluteSize.X / 3.f, AbsoluteSize.Y / 3.f);
+				return FVector2D(LocalSize.X / 3.f, LocalSize.X / 3.f);
 			}
 
-			return FVector2D(AbsoluteSize.Y / 3.f + (AbsoluteSize.X - AbsoluteSize.Y) * 0.5f, AbsoluteSize.Y / 3.f);
+			return FVector2D(LocalSize.Y / 3.f + (LocalSize.X - LocalSize.Y) * 0.5f, LocalSize.Y / 3.f);
 		}
 	}
 
