Material Designer
- Context-based material model defaults are now available.
- By default decal components will have materials set to the decal domain.
- By default static mesh materials will by set to translucent if they are 2D, opaque if they are 3D.

diff --git a/Engine/Plugins/Experimental/DynamicMaterial/Source/DynamicMaterialEditor/Private/DynamicMaterialEditorModule.cpp b/Engine/Plugins/Experimental/DynamicMaterial/Source/DynamicMaterialEditor/Private/DynamicMaterialEditorModule.cpp
--- a/Engine/Plugins/Experimental/DynamicMaterial/Source/DynamicMaterialEditor/Private/DynamicMaterialEditorModule.cpp
+++ b/Engine/Plugins/Experimental/DynamicMaterial/Source/DynamicMaterialEditor/Private/DynamicMaterialEditorModule.cpp
@@ -15,6 +15,7 @@
 #include "Components/MaterialValues/DMMaterialValueFloat3XYZ.h"
 #include "Components/MaterialValues/DMMaterialValueFloat4.h"
 #include "Components/MaterialValues/DMMaterialValueTexture.h"
+#include "Components/ActorComponent.h"
 #include "Components/PrimitiveComponent.h"
 #include "DetailsPanel/DMMaterialInterfaceTypeCustomizer.h"
 #include "DetailsPanel/DMPropertyTypeCustomizer.h"
@@ -26,13 +27,16 @@
 #include "DynamicMaterialEditorStyle.h"
 #include "DynamicMaterialModule.h"
 #include "Engine/World.h"
+#include "GameFramework/Actor.h"
 #include "LevelEditor/DMLevelEditorIntegration.h"
 #include "Material/DynamicMaterialInstance.h"
 #include "MaterialList.h"
+#include "Model/DMMaterialModelDefaults.h"
 #include "Model/DynamicMaterialModel.h"
 #include "Model/DynamicMaterialModelEditorOnlyData.h"
 #include "Modules/ModuleManager.h"
 #include "PropertyEditorModule.h"
+#include "Model/DMModelCreatedCallback.h"
 #include "Slate/Properties/Editors/SDMPropertyEditBoolValue.h"
 #include "Slate/Properties/Editors/SDMPropertyEditFloat1Value.h"
 #include "Slate/Properties/Editors/SDMPropertyEditFloat2Value.h"
@@ -78,6 +82,7 @@ TMap<UClass*, FDMCreateValueEditWidgetDelegate> FDynamicMaterialEditorModule::Va
 TMap<UClass*, FDMComponentPropertyRowGeneratorDelegate> FDynamicMaterialEditorModule::ComponentPropertyRowGenerators;
 TMap<UClass*, FDMGetObjectMaterialPropertiesDelegate> FDynamicMaterialEditorModule::CustomMaterialPropertyGenerators;
 FDMOnUIValueUpdate FDynamicMaterialEditorModule::OnUIValueUpdate;
+TArray<TSharedRef<IDMMaterialModelCreatedCallback>> FDynamicMaterialEditorModule::OnCreatedCallbacks;
 
 void FDynamicMaterialEditorModule::RegisterValueEditWidgetDelegate(UClass* InClass, FDMCreateValueEditWidgetDelegate ValueEditBodyDelegate)
 {
@@ -182,6 +187,49 @@ void FDynamicMaterialEditorModule::RegisterCustomMaterialPropertyGenerator(UClas
 	CustomMaterialPropertyGenerators.FindOrAdd(InClass) = InGenerator;
 }
 
+void FDynamicMaterialEditorModule::RegisterMaterialModelCreatedCallback(const TSharedRef<IDMMaterialModelCreatedCallback> InCallback)
+{
+	OnCreatedCallbacks.Add(InCallback);
+
+	OnCreatedCallbacks.StableSort(
+		[](const TSharedRef<IDMMaterialModelCreatedCallback>& InA, const TSharedRef<IDMMaterialModelCreatedCallback>& InB)
+		{
+			return InA.Get() < InB.Get();
+		}
+	);
+}
+
+void FDynamicMaterialEditorModule::UnregisterMaterialModelCreatedCallback(const TSharedRef<IDMMaterialModelCreatedCallback> InCallback)
+{
+	OnCreatedCallbacks.Remove(InCallback);
+}
+
+void FDynamicMaterialEditorModule::OnMaterialModelCreated(UDynamicMaterialModel* InModel)
+{
+	if (!IsValid(InModel))
+	{
+		return;
+	}
+
+	UDynamicMaterialModelEditorOnlyData* EditorOnlyData = UDynamicMaterialModelEditorOnlyData::Get(InModel);
+	UObject* Outer = InModel->GetOuter();
+	UActorComponent* OuterComponent = InModel->GetTypedOuter<UActorComponent>();
+	AActor* OuterActor = InModel->GetTypedOuter<AActor>();
+
+	const FDMMaterialModelCreatedCallbackParams Params = {
+		InModel,
+		EditorOnlyData,
+		Outer,
+		OuterComponent,
+		OuterActor
+	};
+
+	for (const TSharedRef<IDMMaterialModelCreatedCallback>& OnCreatedCallback : OnCreatedCallbacks)
+	{
+		OnCreatedCallback->OnModelCreated(Params);
+	}
+}
+
 FDMGetObjectMaterialPropertiesDelegate FDynamicMaterialEditorModule::GetCustomMaterialPropertyGenerator(UClass* InClass)
 {
 	if (InClass)
@@ -225,6 +273,8 @@ void FDynamicMaterialEditorModule::StartupModule()
 
 	FDMLevelEditorIntegration::Initialize();
 
+	FDMMaterialModelDefaults::RegisterDefaultsDelegates();
+
 	RegisterValueEditWidgetDelegate<UDMMaterialValueBool,      SDMPropertyEditBoolValue>();
 	RegisterValueEditWidgetDelegate<UDMMaterialValueFloat1,    SDMPropertyEditFloat1Value>();
 	RegisterValueEditWidgetDelegate<UDMMaterialValueFloat2,    SDMPropertyEditFloat2Value>();
@@ -272,6 +322,8 @@ void FDynamicMaterialEditorModule::ShutdownModule()
 
 	FDMLevelEditorIntegration::Shutdown();
 
+	FDMMaterialModelDefaults::UnregsiterDefaultsDelegates();
+
 	BuildRequestList.Empty();
 
 	UDMMaterialValueTexture::GetDefaultRGBTexture.Unbind();
diff --git a/Engine/Plugins/Experimental/DynamicMaterial/Source/DynamicMaterialEditor/Private/DynamicMaterialEditorModule.h b/Engine/Plugins/Experimental/DynamicMaterial/Source/DynamicMaterialEditor/Private/DynamicMaterialEditorModule.h
--- a/Engine/Plugins/Experimental/DynamicMaterial/Source/DynamicMaterialEditor/Private/DynamicMaterialEditorModule.h
+++ b/Engine/Plugins/Experimental/DynamicMaterial/Source/DynamicMaterialEditor/Private/DynamicMaterialEditorModule.h
@@ -90,6 +90,8 @@ public:
 	//~ Begin IDynamicMaterialEditorModule
 	virtual void OpenEditor(UWorld* InWorld) override;
 	virtual void RegisterCustomMaterialPropertyGenerator(UClass* InClass, FDMGetObjectMaterialPropertiesDelegate InGenerator) override;
+	virtual void RegisterMaterialModelCreatedCallback(const TSharedRef<IDMMaterialModelCreatedCallback> InCallback)  override;
+	virtual void UnregisterMaterialModelCreatedCallback(const TSharedRef<IDMMaterialModelCreatedCallback> InCallback) override;
 	//~ End IDynamicMaterialEditorModule
 
 	//~ Begin IModuleInterface
@@ -117,11 +119,14 @@ public:
 
 	const TSharedRef<FUICommandList>& GetCommandList() const { return CommandList; }
 
+	void OnMaterialModelCreated(UDynamicMaterialModel* InModel);
+
 protected:
 	static TMap<UClass*, FDMCreateValueEditWidgetDelegate> ValueEditWidgetDelegates;
 	static TMap<UClass*, FDMComponentPropertyRowGeneratorDelegate> ComponentPropertyRowGenerators;
 	static TMap<UClass*, FDMGetObjectMaterialPropertiesDelegate> CustomMaterialPropertyGenerators;
 	static FDMOnUIValueUpdate OnUIValueUpdate;
+	static TArray<TSharedRef<IDMMaterialModelCreatedCallback>> OnCreatedCallbacks;
 
 	TSet<FDMBuildRequestEntry> BuildRequestList;
 	TSharedRef<FUICommandList> CommandList;
diff --git a/Engine/Plugins/Experimental/DynamicMaterial/Source/DynamicMaterialEditor/Private/Model/DMMaterialBuildState.cpp b/Engine/Plugins/Experimental/DynamicMaterial/Source/DynamicMaterialEditor/Private/Model/DMMaterialBuildState.cpp
--- a/Engine/Plugins/Experimental/DynamicMaterial/Source/DynamicMaterialEditor/Private/Model/DMMaterialBuildState.cpp
+++ b/Engine/Plugins/Experimental/DynamicMaterial/Source/DynamicMaterialEditor/Private/Model/DMMaterialBuildState.cpp
@@ -384,7 +384,7 @@ void FDMMaterialBuildState::AddStageSourceExpressions(const UDMMaterialStageSour
 {
 	if (ensure(!HasStageSource(InStageSource)))
 	{
-		if (ensure(!InStageSourceExpressions.IsEmpty()))
+		if (!InStageSourceExpressions.IsEmpty())
 		{
 			StageSources.Emplace(InStageSource, InStageSourceExpressions);
 		}
diff --git a/Engine/Plugins/Experimental/DynamicMaterial/Source/DynamicMaterialEditor/Private/Model/DMMaterialModelDefaults.cpp b/Engine/Plugins/Experimental/DynamicMaterial/Source/DynamicMaterialEditor/Private/Model/DMMaterialModelDefaults.cpp
new file mode 100644
--- /dev/null
+++ b/Engine/Plugins/Experimental/DynamicMaterial/Source/DynamicMaterialEditor/Private/Model/DMMaterialModelDefaults.cpp
@@ -0,0 +1,119 @@
+// Copyright Epic Games, Inc. All Rights Reserved.
+
+#include "Model/DMMaterialModelDefaults.h"
+
+#include "Components/DecalComponent.h"
+#include "Components/StaticMeshComponent.h"
+#include "DynamicMaterialEditorModule.h"
+#include "Delegates/IDelegateInstance.h"
+#include "Model/DMModelCreatedCallback.h"
+#include "Model/DynamicMaterialModel.h"
+#include "Model/DynamicMaterialModelEditorOnlyData.h"
+
+namespace UE::DynamicMaterialEditor::Private
+{
+	TSharedPtr<IDMMaterialModelCreatedCallback> DefaultsCallback_StaticMeshComponent;
+	TSharedPtr<IDMMaterialModelCreatedCallback> DefaultsCallback_DecalComponent;
+}
+
+struct FDMMaterialModelCreatedCallback_StaticMeshComponent : FDMMaterialModelCreatedCallbackBase
+{
+	FDMMaterialModelCreatedCallback_StaticMeshComponent(uint32 InPriority)
+		: FDMMaterialModelCreatedCallbackBase(InPriority)
+	{		
+	}
+
+	virtual ~FDMMaterialModelCreatedCallback_StaticMeshComponent() override = default;
+
+	//~ Begin IDMMaterialModelCreatedCallback
+	virtual void OnModelCreated(const FDMMaterialModelCreatedCallbackParams& InParams) override
+	{
+		if (!::IsValid(InParams.EditorOnlyData))
+		{
+			return;
+		}
+
+		UStaticMeshComponent* Component = InParams.MaterialModel->GetTypedOuter<UStaticMeshComponent>();
+
+		if (!::IsValid(Component))
+		{
+			return;
+		}
+
+		FVector Min, Max;
+		Component->GetLocalBounds(Min, Max);
+
+		if (FMath::IsNearlyEqual(Min.X, Max.X)
+			|| FMath::IsNearlyEqual(Min.Y, Max.Y)
+			|| FMath::IsNearlyEqual(Min.Z, Max.Z))
+		{
+			InParams.EditorOnlyData->SetBlendMode(EBlendMode::BLEND_Translucent);
+		}
+		else
+		{
+			InParams.EditorOnlyData->SetBlendMode(EBlendMode::BLEND_Opaque);
+		}
+	}
+	//~ End IDMMaterialModelCreatedCallback
+};
+
+struct FDMMaterialModelCreatedCallback_DecalComponent : FDMMaterialModelCreatedCallbackBase
+{
+	FDMMaterialModelCreatedCallback_DecalComponent(uint32 InPriority)
+		: FDMMaterialModelCreatedCallbackBase(InPriority)
+	{
+	}
+
+	virtual ~FDMMaterialModelCreatedCallback_DecalComponent() override = default;
+
+	//~ Begin IDMMaterialModelCreatedCallback
+	virtual void OnModelCreated(const FDMMaterialModelCreatedCallbackParams& InParams) override
+	{
+		if (!::IsValid(InParams.EditorOnlyData))
+		{
+			return;
+		}
+
+		UDecalComponent* Component = InParams.MaterialModel->GetTypedOuter<UDecalComponent>();
+
+		if (!::IsValid(Component))
+		{
+			return;
+		}
+
+		InParams.EditorOnlyData->SetDomain(EMaterialDomain::MD_DeferredDecal);
+	}
+	//~ End IDMMaterialModelCreatedCallback
+};
+
+void FDMMaterialModelDefaults::RegisterDefaultsDelegates()
+{
+	using namespace UE::DynamicMaterialEditor::Private;
+
+	UnregsiterDefaultsDelegates();
+
+	DefaultsCallback_StaticMeshComponent = IDynamicMaterialEditorModule::Get().RegisterMaterialModelCreatedCallback<FDMMaterialModelCreatedCallback_StaticMeshComponent>(
+		1000
+	);
+
+	DefaultsCallback_DecalComponent = IDynamicMaterialEditorModule::Get().RegisterMaterialModelCreatedCallback<FDMMaterialModelCreatedCallback_DecalComponent>(
+		2000
+	);
+}
+
+void FDMMaterialModelDefaults::UnregsiterDefaultsDelegates()
+{
+	using namespace UE::DynamicMaterialEditor::Private;
+
+	if (DefaultsCallback_StaticMeshComponent.IsValid())
+	{
+		FDynamicMaterialEditorModule::Get().UnregisterMaterialModelCreatedCallback(DefaultsCallback_StaticMeshComponent.ToSharedRef());
+		DefaultsCallback_StaticMeshComponent.Reset();
+	}
+
+	if (DefaultsCallback_DecalComponent.IsValid())
+	{
+		FDynamicMaterialEditorModule::Get().UnregisterMaterialModelCreatedCallback(DefaultsCallback_DecalComponent.ToSharedRef());
+		DefaultsCallback_DecalComponent.Reset();
+	}
+}
diff --git a/Engine/Plugins/Experimental/DynamicMaterial/Source/DynamicMaterialEditor/Private/Model/DMMaterialModelDefaults.h b/Engine/Plugins/Experimental/DynamicMaterial/Source/DynamicMaterialEditor/Private/Model/DMMaterialModelDefaults.h
new file mode 100644
--- /dev/null
+++ b/Engine/Plugins/Experimental/DynamicMaterial/Source/DynamicMaterialEditor/Private/Model/DMMaterialModelDefaults.h
@@ -0,0 +1,13 @@
+// Copyright Epic Games, Inc. All Rights Reserved.
+
+#pragma once
+
+class UDynamicMaterialModel;
+
+class FDMMaterialModelDefaults
+{
+public:
+	static void RegisterDefaultsDelegates();
+
+	static void UnregsiterDefaultsDelegates();
+};
diff --git a/Engine/Plugins/Experimental/DynamicMaterial/Source/DynamicMaterialEditor/Private/Model/DMModelCreatedCallback.cpp b/Engine/Plugins/Experimental/DynamicMaterial/Source/DynamicMaterialEditor/Private/Model/DMModelCreatedCallback.cpp
new file mode 100644
--- /dev/null
+++ b/Engine/Plugins/Experimental/DynamicMaterial/Source/DynamicMaterialEditor/Private/Model/DMModelCreatedCallback.cpp
@@ -0,0 +1,24 @@
+// Copyright Epic Games, Inc. All Rights Reserved.
+
+#include "Model/DMModelCreatedCallback.h"
+
+FDMMaterialModelCreatedCallbackBase::FDMMaterialModelCreatedCallbackBase(uint32 InPriority)
+	: Priority(InPriority)
+{
+}
+
+uint32 FDMMaterialModelCreatedCallbackBase::GetPriority() const
+{
+	return Priority;
+}
+
+FDMMaterialModelCreatedCallbackDelegate::FDMMaterialModelCreatedCallbackDelegate(uint32 InPriority, const FOnModelCreated& InOnModelCreatedDelegate)
+	: FDMMaterialModelCreatedCallbackBase(InPriority)
+	, OnModelCreatedDelegate(InOnModelCreatedDelegate)
+{
+}
+
+void FDMMaterialModelCreatedCallbackDelegate::OnModelCreated(const FDMMaterialModelCreatedCallbackParams& InParams)
+{
+	OnModelCreatedDelegate.ExecuteIfBound(InParams);
+}
diff --git a/Engine/Plugins/Experimental/DynamicMaterial/Source/DynamicMaterialEditor/Private/Model/DynamicMateriaModelFactory.cpp b/Engine/Plugins/Experimental/DynamicMaterial/Source/DynamicMaterialEditor/Private/Model/DynamicMateriaModelFactory.cpp
--- a/Engine/Plugins/Experimental/DynamicMaterial/Source/DynamicMaterialEditor/Private/Model/DynamicMateriaModelFactory.cpp
+++ b/Engine/Plugins/Experimental/DynamicMaterial/Source/DynamicMaterialEditor/Private/Model/DynamicMateriaModelFactory.cpp
@@ -1,5 +1,6 @@
 // Copyright Epic Games, Inc. All Rights Reserved.
 
+#include "DynamicMaterialEditorModule.h"
 #include "Components/DMMaterialComponent.h"
 #include "Model/DynamicMaterialModelFactory.h"
 #include "Model/DynamicMaterialModel.h"
@@ -40,6 +41,8 @@ UObject* UDynamicMaterialModelFactory::FactoryCreateNew(UClass* Class, UObject*
 	const FDMInitializationGuard InitGuard;
 	ModelEditorOnlyData->AddSlot();
 
+	FDynamicMaterialEditorModule::Get().OnMaterialModelCreated(NewModel);
+
 	return NewModel;
 }
 
diff --git a/Engine/Plugins/Experimental/DynamicMaterial/Source/DynamicMaterialEditor/Public/IDynamicMaterialEditorModule.h b/Engine/Plugins/Experimental/DynamicMaterial/Source/DynamicMaterialEditor/Public/IDynamicMaterialEditorModule.h
--- a/Engine/Plugins/Experimental/DynamicMaterial/Source/DynamicMaterialEditor/Public/IDynamicMaterialEditorModule.h
+++ b/Engine/Plugins/Experimental/DynamicMaterial/Source/DynamicMaterialEditor/Public/IDynamicMaterialEditorModule.h
@@ -7,8 +7,10 @@
 #include "Modules/ModuleManager.h"
 
 class UClass;
+class UDynamicMaterialModel;
 class UObject;
 struct FDMObjectMaterialProperty;
+struct IDMMaterialModelCreatedCallback;
 
 DECLARE_DELEGATE_RetVal_OneParam(TArray<FDMObjectMaterialProperty>, FDMGetObjectMaterialPropertiesDelegate, UObject* InObject)
 
@@ -32,4 +34,17 @@ public:
 	virtual void OpenEditor(UWorld* InWorld) = 0;
 
 	virtual void RegisterCustomMaterialPropertyGenerator(UClass* InClass, FDMGetObjectMaterialPropertiesDelegate InGenerator) = 0;
+
+	virtual void RegisterMaterialModelCreatedCallback(const TSharedRef<IDMMaterialModelCreatedCallback> InCallback) = 0;
+
+	virtual void UnregisterMaterialModelCreatedCallback(const TSharedRef<IDMMaterialModelCreatedCallback> InCallback) = 0;
+
+	template<typename InCallbackType, typename... InArgsType
+		UE_REQUIRES(std::is_base_of_v<IDMMaterialModelCreatedCallback, InCallbackType>)>
+	TSharedRef<InCallbackType> RegisterMaterialModelCreatedCallback(InArgsType&&... InArgs)
+	{
+		TSharedRef<InCallbackType> NewCallback = MakeShared<InCallbackType>(Forward<InArgsType>(InArgs)...);
+		RegisterMaterialModelCreatedCallback(NewCallback);
+		return NewCallback;
+	}
 };
diff --git a/Engine/Plugins/Experimental/DynamicMaterial/Source/DynamicMaterialEditor/Public/Model/DMModelCreatedCallback.h b/Engine/Plugins/Experimental/DynamicMaterial/Source/DynamicMaterialEditor/Public/Model/DMModelCreatedCallback.h
new file mode 100644
--- /dev/null
+++ b/Engine/Plugins/Experimental/DynamicMaterial/Source/DynamicMaterialEditor/Public/Model/DMModelCreatedCallback.h
@@ -0,0 +1,69 @@
+// Copyright Epic Games, Inc. All Rights Reserved.
+
+#pragma once
+
+#include "Delegates/Delegate.h"
+#include "Delegates/DelegateCombinations.h"
+#include "HAL/Platform.h"
+#include "Templates/SharedPointer.h"
+
+class AActor;
+class UActorComponent;
+class UDynamicMaterialModelEditorOnlyData;
+class UDynamicMaterialModel;
+class UObject;
+
+struct FDMMaterialModelCreatedCallbackParams
+{
+	UDynamicMaterialModel* MaterialModel;
+	UDynamicMaterialModelEditorOnlyData* EditorOnlyData;
+	UObject* Outer;
+	UActorComponent* OuterComponent;
+	AActor* OuterActor;
+};
+
+// Callback interface
+struct IDMMaterialModelCreatedCallback : public TSharedPtr<IDMMaterialModelCreatedCallback>
+{
+	virtual ~IDMMaterialModelCreatedCallback() = default;
+
+	virtual uint32 GetPriority() const = 0;
+	virtual void OnModelCreated(const FDMMaterialModelCreatedCallbackParams& InParams) = 0;
+
+	virtual bool operator<(const IDMMaterialModelCreatedCallback& InOther) const
+	{
+		return GetPriority() < InOther.GetPriority();
+	}
+};
+
+// Default implementation
+struct FDMMaterialModelCreatedCallbackBase : public IDMMaterialModelCreatedCallback
+{
+	FDMMaterialModelCreatedCallbackBase(uint32 InPriority);
+
+	virtual ~FDMMaterialModelCreatedCallbackBase() override = default;
+
+	//~ Begin IDMMaterialModelCreatedCallback
+	virtual uint32 GetPriority() const override;
+	//~ End IDMMaterialModelCreatedCallback
+
+protected:
+	uint32 Priority;
+};
+
+// Default implementation
+struct FDMMaterialModelCreatedCallbackDelegate : public FDMMaterialModelCreatedCallbackBase
+{
+	DECLARE_DELEGATE_OneParam(FOnModelCreated, const FDMMaterialModelCreatedCallbackParams&)
+
+	FDMMaterialModelCreatedCallbackDelegate(uint32 InPriority, const FOnModelCreated& InOnModelCreatedDelegate);
+
+	virtual ~FDMMaterialModelCreatedCallbackDelegate() override = default;
+
+	//~ Begin IDMMaterialModelCreatedCallback
+	virtual void OnModelCreated(const FDMMaterialModelCreatedCallbackParams& InParams) override;
+	//~ End IDMMaterialModelCreatedCallback
+
+protected:
+	FOnModelCreated OnModelCreatedDelegate;
+};
