Material Designer
- Fixed bug where undoing would not undo the mask texture value automatically set by setting the texture value of the base stage.

diff --git a/Engine/Plugins/Experimental/DynamicMaterial/Source/DynamicMaterialEditor/Private/Components/MaterialStageExpressions/DMMSETextureSampleBase.cpp b/Engine/Plugins/Experimental/DynamicMaterial/Source/DynamicMaterialEditor/Private/Components/MaterialStageExpressions/DMMSETextureSampleBase.cpp
--- a/Engine/Plugins/Experimental/DynamicMaterial/Source/DynamicMaterialEditor/Private/Components/MaterialStageExpressions/DMMSETextureSampleBase.cpp
+++ b/Engine/Plugins/Experimental/DynamicMaterial/Source/DynamicMaterialEditor/Private/Components/MaterialStageExpressions/DMMSETextureSampleBase.cpp
@@ -12,6 +12,7 @@
 #include "Components/MaterialStageInputs/DMMSIThroughput.h"
 #include "Components/MaterialStageInputs/DMMSIValue.h"
 #include "Components/MaterialValues/DMMaterialValueTexture.h"
+#include "CoreGlobals.h"
 #include "Engine/TextureCube.h"
 #include "Materials/MaterialExpressionTextureSample.h"
 #include "Model/DynamicMaterialModelEditorOnlyData.h"
@@ -229,8 +230,8 @@ void UDMMaterialStageExpressionTextureSampleBase::UpdateMask()
 	const UDMMaterialSlot* Slot = Layer->GetSlot();
 	check(Slot);
 
-	UDynamicMaterialModelEditorOnlyData* MaterialMode = Slot->GetMaterialModelEditorOnlyData();
-	check(MaterialMode);
+	UDynamicMaterialModelEditorOnlyData* EditorOnlyData = Slot->GetMaterialModelEditorOnlyData();
+	check(EditorOnlyData);
 
 	UDMMaterialStage* BaseStage = Layer->GetStage(EDMMaterialLayerStage::Base);
 	UDMMaterialStage* MaskStage = Layer->GetStage(EDMMaterialLayerStage::Mask, /* Enabled Only */ true);
@@ -317,16 +318,26 @@ void UDMMaterialStageExpressionTextureSampleBase::UpdateMask()
 						{
 							if (LayerBlendTextureValue->GetClass() == BaseTextureValue->GetClass())
 							{
+								if (GUndo)
+								{
+									LayerBlendTextureValue->Modify();
+								}
+
 								LayerBlendTextureValue->SetValue(BaseTexture);
 							}
 							else
 							{
-								UDMMaterialValueTexture* NewLayerBlendTextureValue = UDMMaterialValueTexture::CreateMaterialValueTexture(MaterialMode, BaseTextureValue->GetValue());
+								UDMMaterialValueTexture* NewLayerBlendTextureValue = UDMMaterialValueTexture::CreateMaterialValueTexture(EditorOnlyData, BaseTextureValue->GetValue());
 								check(NewLayerBlendTextureValue);
 
 								LayerBlendTextureInputValue->SetValue(NewLayerBlendTextureValue);
 							}
 
+							if (GUndo)
+							{
+								MaskStage->Modify();
+							}
+
 							// Set output to alpha
 							MaskStage->UpdateInputMap(2, MaskStageInputConnections[2].Channels[0].SourceIndex,
 								FDMMaterialStageConnectorChannel::WHOLE_CHANNEL, 1, FDMMaterialStageConnectorChannel::WHOLE_CHANNEL, EDMMaterialPropertyType::None);
@@ -342,6 +353,11 @@ void UDMMaterialStageExpressionTextureSampleBase::UpdateMask()
 	// Couldn't find a texture to update, so create a new one.
 
 	// 2nd input, 2nd output (Alpha)
+	if (GUndo)
+	{
+		MaskStage->Modify();
+	}
+
 	UDMMaterialStageInputExpression::ChangeStageInput_Expression(MaskStage, 
 		UDMMaterialStageExpressionTextureSample::StaticClass(), 2, FDMMaterialStageConnectorChannel::WHOLE_CHANNEL,
 		1, FDMMaterialStageConnectorChannel::WHOLE_CHANNEL);
diff --git a/Engine/Plugins/Experimental/DynamicMaterial/Source/DynamicMaterialEditor/Private/Slate/SDMStage.cpp b/Engine/Plugins/Experimental/DynamicMaterial/Source/DynamicMaterialEditor/Private/Slate/SDMStage.cpp
--- a/Engine/Plugins/Experimental/DynamicMaterial/Source/DynamicMaterialEditor/Private/Slate/SDMStage.cpp
+++ b/Engine/Plugins/Experimental/DynamicMaterial/Source/DynamicMaterialEditor/Private/Slate/SDMStage.cpp
@@ -625,6 +625,7 @@ void SDMStage::HandleAssetDragDropOperation(FAssetDragDropOp& AssetDragDropOpera
 
 	if (TextureValue)
 	{
+		TextureValue->Modify();
 		TextureValue->SetValue(Texture);
 	}
 }
