Material Designer
- Changes to the Popout window.

diff --git a/Engine/Plugins/Experimental/DynamicMaterial/Source/DynamicMaterialEditor/Private/Slate/Properties/Generators/DMTextureUVPropertyRowGenerator.cpp b/Engine/Plugins/Experimental/DynamicMaterial/Source/DynamicMaterialEditor/Private/Slate/Properties/Generators/DMTextureUVPropertyRowGenerator.cpp
--- a/Engine/Plugins/Experimental/DynamicMaterial/Source/DynamicMaterialEditor/Private/Slate/Properties/Generators/DMTextureUVPropertyRowGenerator.cpp
+++ b/Engine/Plugins/Experimental/DynamicMaterial/Source/DynamicMaterialEditor/Private/Slate/Properties/Generators/DMTextureUVPropertyRowGenerator.cpp
@@ -8,11 +8,10 @@
 #include "Modules/ModuleManager.h"
 #include "PropertyCustomizationHelpers.h"
 #include "PropertyEditorModule.h"
-#include "Slate/SDMEditor.h"
-#include "Slate/Properties/Editors/SDMPropertyEditFloat.h"
 #include "Slate/Properties/Editors/SDMPropertyEditVector.h"
-#include "Slate/SDMComponentEdit.h"
 #include "Slate/Properties/SDMTextureUVVisualizerProperty.h"
+#include "Slate/SDMComponentEdit.h"
+#include "Slate/SDMEditor.h"
 #include "Styling/SlateIconFinder.h"
 #include "Widgets/Layout/SBox.h"
 
@@ -26,20 +25,8 @@ const TSharedRef<FDMTextureUVPropertyRowGenerator>& FDMTextureUVPropertyRowGener
 
 namespace UE::DynamicMaterialEditor::Private
 {
-	void AddRotationRow(const TSharedRef<SDMComponentEdit>& InComponentEditWidget, UDMTextureUV* InTextureUV, TArray<FDMPropertyHandle>& InOutPropertyRows);
-
-	void AddScaleRow(const TSharedRef<SDMComponentEdit>& InComponentEditWidget, UDMTextureUV* InTextureUV, TArray<FDMPropertyHandle>& InOutPropertyRows);
-
-	TSharedRef<SWidget> CreateScaleLinkWidget(const TSharedRef<SDMComponentEdit>& InComponentEditWidget, UDMTextureUV* InTextureUV);
-	ECheckBoxState GetScaleLinkWidgetState(UDMTextureUV* InTextureUV);
-	void OnScaleLinkWidgetChanged(ECheckBoxState InNewState, UDMTextureUV* InTextureUV);
-
-	TSharedRef<SWidget> CreateScaleValueWidget(const TSharedRef<SDMComponentEdit>& InComponentEditWidget, UDMTextureUV* InTextureUV);
-
-	void AddMirrorRow(const TSharedRef<SDMComponentEdit>& InComponentEditWidget, UDMTextureUV* InTextureUV, TArray<FDMPropertyHandle>& InOutPropertyRows);
-	TSharedRef<SWidget> CreateMirrorExtensionButtons(const TSharedRef<SDMComponentEdit>& InComponentEditWidget, UDMTextureUV* InTextureUV);
-
-	void AddVisualizerRow(const TSharedRef<SDMComponentEdit>& InComponentEditWidget, UDMTextureUV* InTextureUV, TArray<FDMPropertyHandle>& InOutPropertyRows);
+	void AddPropertyRow(const TSharedRef<SWidget>& InComponentEditWidget, UDMTextureUV* InTextureUV, FName InProperty, TArray<FDMPropertyHandle>& InOutPropertyRows);
+	void AddVisualizerRow(const TSharedRef<SWidget>& InComponentEditWidget, UDMTextureUV* InTextureUV, TArray<FDMPropertyHandle>& InOutPropertyRows);
 }
 
 void FDMTextureUVPropertyRowGenerator::AddComponentProperties(const TSharedRef<SDMComponentEdit>& InComponentEditWidget, UDMMaterialComponent* InComponent,
@@ -66,356 +53,60 @@ void FDMTextureUVPropertyRowGenerator::AddComponentProperties(const TSharedRef<S
 
 	using namespace UE::DynamicMaterialEditor::Private;
 
-	const int PropertyCount = InOutPropertyRows.Num();
-
-	AddPropertyEditRows(InComponentEditWidget, InComponent, UDMTextureUV::NAME_Offset, InOutPropertyRows, InOutProcessedObjects);
-	AddRotationRow(InComponentEditWidget, TextureUV, InOutPropertyRows);
-	AddScaleRow(InComponentEditWidget, TextureUV, InOutPropertyRows);
-	AddPropertyEditRows(InComponentEditWidget, InComponent, UDMTextureUV::NAME_Pivot, InOutPropertyRows, InOutProcessedObjects);
-	AddMirrorRow(InComponentEditWidget, TextureUV, InOutPropertyRows);
+	AddPropertyRow(InComponentEditWidget, TextureUV, UDMTextureUV::NAME_Offset, InOutPropertyRows);
+	AddPropertyRow(InComponentEditWidget, TextureUV, UDMTextureUV::NAME_Rotation, InOutPropertyRows);
+	AddPropertyRow(InComponentEditWidget, TextureUV, UDMTextureUV::NAME_Scale, InOutPropertyRows);
+	AddPropertyRow(InComponentEditWidget, TextureUV, UDMTextureUV::NAME_Pivot, InOutPropertyRows);
+	AddPropertyRow(InComponentEditWidget, TextureUV, UDMTextureUV::NAME_bMirrorOnX, InOutPropertyRows);
+	AddPropertyRow(InComponentEditWidget, TextureUV, UDMTextureUV::NAME_bMirrorOnY, InOutPropertyRows);
 	AddVisualizerRow(InComponentEditWidget, TextureUV, InOutPropertyRows);
 }
 
-bool FDMTextureUVPropertyRowGenerator::AllowKeyframeButton(UDMMaterialComponent* InComponent, FProperty* InProperty)
+void FDMTextureUVPropertyRowGenerator::AddPopoutComponentProperties(const TSharedRef<SWidget>& InParentWidget, UDMMaterialComponent* InComponent, TArray<FDMPropertyHandle>& InOutPropertyRows)
 {
-	if (InProperty)
-	{
-		const bool* AddKeyframeButtonPtr = UDMTextureUV::TextureProperties.Find(InProperty->GetFName());
-
-		if (AddKeyframeButtonPtr)
-		{
-			return *AddKeyframeButtonPtr;
-		}
-	}
-
-	return FDMComponentPropertyRowGenerator::AllowKeyframeButton(InComponent, InProperty);
-}
-
-void UE::DynamicMaterialEditor::Private::AddRotationRow(const TSharedRef<SDMComponentEdit>& InComponentEditWidget, UDMTextureUV* InTextureUV, TArray<FDMPropertyHandle>& InOutPropertyRows)
-{
-	InOutPropertyRows.Add(SDMEditor::GetPropertyHandle(&*InComponentEditWidget, InTextureUV, UDMTextureUV::NAME_Rotation));
-}
-
-void UE::DynamicMaterialEditor::Private::AddScaleRow(const TSharedRef<SDMComponentEdit>& InComponentEditWidget, UDMTextureUV* InTextureUV, TArray<FDMPropertyHandle>& InOutPropertyRows)
-{
-	InOutPropertyRows.Add(SDMEditor::GetPropertyHandle(&*InComponentEditWidget, InTextureUV, UDMTextureUV::NAME_Scale));
-}
-
-TSharedRef<SWidget> UE::DynamicMaterialEditor::Private::CreateScaleLinkWidget(const TSharedRef<SDMComponentEdit>& InComponentEditWidget, UDMTextureUV* InTextureUV)
-{
-	const FSlateBrush* Checked = FSlateIconFinder::FindIcon("PropertyWindow.Locked").GetIcon();
-	const FSlateBrush* Unchecked = FSlateIconFinder::FindIcon("PropertyWindow.Unlocked").GetIcon();
-	const FSlateBrush* NoBorder = FSlateIconFinder::FindIcon("NoBorder").GetIcon();
-
-	return SNew(SCheckBox)
-		.IsChecked_Static(&UE::DynamicMaterialEditor::Private::GetScaleLinkWidgetState, InTextureUV)
-		.OnCheckStateChanged_Static(&UE::DynamicMaterialEditor::Private::OnScaleLinkWidgetChanged, InTextureUV)
-		.CheckedImage(Checked)
-		.CheckedHoveredImage(Checked)
-		.CheckedPressedImage(Checked)
-		.UncheckedImage(Unchecked)
-		.UncheckedHoveredImage(Unchecked)
-		.UncheckedPressedImage(Unchecked)
-		.BackgroundImage(NoBorder)
-		.BackgroundHoveredImage(NoBorder)
-		.BackgroundPressedImage(NoBorder);
-}
-
-ECheckBoxState UE::DynamicMaterialEditor::Private::GetScaleLinkWidgetState(UDMTextureUV* InTextureUV)
-{
-	if (!IsValid(InTextureUV))
+	if (!IsValid(InComponent))
 	{
-		return ECheckBoxState::Unchecked;
+		return;
 	}
 
-	return InTextureUV->bLinkScale ? ECheckBoxState::Checked : ECheckBoxState::Unchecked;
-}
+	UDMTextureUV* TextureUV = Cast<UDMTextureUV>(InComponent);
 
-void UE::DynamicMaterialEditor::Private::OnScaleLinkWidgetChanged(ECheckBoxState InNewState, UDMTextureUV* InTextureUV)
-{
-	if (!IsValid(InTextureUV))
+	if (!TextureUV)
 	{
 		return;
 	}
 
-	InTextureUV->bLinkScale = (InNewState == ECheckBoxState::Checked);
+	using namespace UE::DynamicMaterialEditor::Private;
+
+	AddPropertyRow(InParentWidget, TextureUV, UDMTextureUV::NAME_Offset, InOutPropertyRows);
+	AddPropertyRow(InParentWidget, TextureUV, UDMTextureUV::NAME_Rotation, InOutPropertyRows);
+	AddPropertyRow(InParentWidget, TextureUV, UDMTextureUV::NAME_Scale, InOutPropertyRows);
+	AddPropertyRow(InParentWidget, TextureUV, UDMTextureUV::NAME_Pivot, InOutPropertyRows);
+	AddPropertyRow(InParentWidget, TextureUV, UDMTextureUV::NAME_bMirrorOnX, InOutPropertyRows);
+	AddPropertyRow(InParentWidget, TextureUV, UDMTextureUV::NAME_bMirrorOnY, InOutPropertyRows);
 }
 
-namespace UE::DynamicMaterialEditor::Private
+bool FDMTextureUVPropertyRowGenerator::AllowKeyframeButton(UDMMaterialComponent* InComponent, FProperty* InProperty)
 {
-	class SDMPropertyEditVectorScale : public SDMPropertyEditVector
+	if (InProperty)
 	{
-	public:
-		using Super = SDMPropertyEditVector;
-
-		void Construct(const FArguments& InArgs, const TSharedPtr<SDMComponentEdit>& InComponentEditWidget, UDMTextureUV* InTextureUV)
-		{
-			if (ensure(IsValid(InTextureUV)))
-			{
-				Super::Construct(
-					Super::FArguments()
-						.ComponentEditWidget(InComponentEditWidget), 
-					InTextureUV->GetPropertyHandle(UDMTextureUV::NAME_Scale), 
-					2
-				);
-			}
-		}
-
-	protected:
-		FVector2D InitialScale;
-
-		virtual TSharedRef<SWidget> GetComponentWidget(int32 InIndex) override
-		{
-			TSharedRef<SWidget> Widget = Super::GetComponentWidget(InIndex);
-
-			if (InIndex > 0)
-			{
-				TWeakPtr<SDMPropertyEditVectorScale> ThisWeak = SharedThis(this).ToWeakPtr();
-
-				Widget->SetVisibility(TAttribute<EVisibility>::CreateLambda([ThisWeak]()
-					{
-						if (TSharedPtr<SDMPropertyEditVectorScale> This = ThisWeak.Pin())
-						{
-							TArray<UObject*> Outers;
-							This->PropertyHandle->GetOuterObjects(Outers);
-
-							if (Outers.Num() == 1 && IsValid(Outers[0]))
-							{
-								if (UDMTextureUV* TextureUV = Cast<UDMTextureUV>(Outers[0]))
-								{
-									if (TextureUV->bLinkScale && FMath::IsNearlyEqual(TextureUV->GetScale().X, TextureUV->GetScale().Y))
-									{
-										return EVisibility::Hidden;
-									}
-								}
-							}
-						}
-
-						return EVisibility::Visible;
-					}));
-				Widget->SetEnabled(TAttribute<bool>::CreateLambda([ThisWeak]()
-					{
-						if (TSharedPtr<SDMPropertyEditVectorScale> This = ThisWeak.Pin())
-						{
-							TArray<UObject*> Outers;
-							This->PropertyHandle->GetOuterObjects(Outers);
-
-							if (Outers.Num() == 1 && IsValid(Outers[0]))
-							{
-								if (UDMTextureUV* TextureUV = Cast<UDMTextureUV>(Outers[0]))
-								{
-									if (TextureUV->bLinkScale)
-									{
-										return false;
-									}
-								}
-							}
-						}
-
-						return true;
-					}));
-			}
-
-			return Widget;
-		}
-
-		virtual void StartTransaction(FText InDescription) override
-		{
-			Super::StartTransaction(InDescription);
-
-			if (PropertyHandle.IsValid() && PropertyHandle->IsValidHandle())
-			{
-				TArray<UObject*> Outers;
-				PropertyHandle->GetOuterObjects(Outers);
-
-				if (Outers.IsEmpty() == false)
-				{
-					if (UDMTextureUV* TextureUV = Cast<UDMTextureUV>(Outers[0]))
-					{
-						InitialScale = TextureUV->GetScale();
-					}
-				}
-			}
-		}
+		const bool* AddKeyframeButtonPtr = UDMTextureUV::TextureProperties.Find(InProperty->GetFName());
 
-		virtual void OnValueChanged(float InNewValue, int32 InComponent) const override
+		if (AddKeyframeButtonPtr)
 		{
-			if (PropertyHandle.IsValid() && PropertyHandle->IsValidHandle())
-			{
-				TArray<UObject*> Outers;
-				PropertyHandle->GetOuterObjects(Outers);
-
-				if (Outers.IsEmpty() == false)
-				{
-					if (UDMTextureUV* TextureUV = Cast<UDMTextureUV>(Outers[0]))
-					{
-						FVector2D CurrentScale = TextureUV->GetScale();
-
-						if (FMath::IsNearlyEqual(CurrentScale[InComponent], InNewValue))
-						{
-							return;
-						}
-
-						auto IsScaledLinked = [this, TextureUV]()
-						{
-							if (FMath::IsNearlyZero(InitialScale.X) || FMath::IsNearlyZero(InitialScale.Y))
-							{
-								return false;
-							}
-
-							if (!IsValid(TextureUV))
-							{
-								return false;
-							}
-
-							return TextureUV->bLinkScale;
-						};
-
-						if (IsScaledLinked())
-						{
-							switch (InComponent)
-							{
-								case 0: // X
-									CurrentScale.Y = (InNewValue / InitialScale.X) * InitialScale.Y;
-									break;
-
-								case 1: // Y
-									CurrentScale.X = (InNewValue / InitialScale.Y) * InitialScale.X;
-									break;
-
-								default:
-									checkNoEntry();
-							}
-						}
-
-						switch (InComponent)
-						{
-							case 0: // X
-								CurrentScale.X = InNewValue;
-								break;
-
-							case 1: // Y
-								CurrentScale.Y = InNewValue;
-								break;
-
-							default:
-								checkNoEntry();
-						}
-
-						TextureUV->SetScale(CurrentScale);
-					}
-				}
-			}
+			return *AddKeyframeButtonPtr;
 		}
-	};
-}
-
-TSharedRef<SWidget> UE::DynamicMaterialEditor::Private::CreateScaleValueWidget(const TSharedRef<SDMComponentEdit>& InComponentEditWidget, UDMTextureUV* InTextureUV)
-{
-	return SNew(UE::DynamicMaterialEditor::Private::SDMPropertyEditVectorScale, InComponentEditWidget, InTextureUV);
-}
+	}
 
-void UE::DynamicMaterialEditor::Private::AddMirrorRow(const TSharedRef<SDMComponentEdit>& InComponentEditWidget, UDMTextureUV* InTextureUV, TArray<FDMPropertyHandle>& InOutPropertyRows)
-{
-	InOutPropertyRows.Add(SDMEditor::GetPropertyHandle(&*InComponentEditWidget, InTextureUV, UDMTextureUV::NAME_bMirrorOnX));
-	InOutPropertyRows.Add(SDMEditor::GetPropertyHandle(&*InComponentEditWidget, InTextureUV, UDMTextureUV::NAME_bMirrorOnY));
+	return FDMComponentPropertyRowGenerator::AllowKeyframeButton(InComponent, InProperty);
 }
 
-TSharedRef<SWidget> UE::DynamicMaterialEditor::Private::CreateMirrorExtensionButtons(const TSharedRef<SDMComponentEdit>& InComponentEditWidget, UDMTextureUV* InTextureUV)
+void UE::DynamicMaterialEditor::Private::AddPropertyRow(const TSharedRef<SWidget>& InComponentEditWidget, UDMTextureUV* InTextureUV, FName InProperty, TArray<FDMPropertyHandle>& InOutPropertyRows)
 {
-	if (!IsValid(InTextureUV))
-	{
-		return SNullWidget::NullWidget;
-	}
-
-	FProperty* XProperty = InTextureUV->GetClass()->FindPropertyByName(UDMTextureUV::NAME_bMirrorOnX);
-	FProperty* YProperty = InTextureUV->GetClass()->FindPropertyByName(UDMTextureUV::NAME_bMirrorOnY);
-
-	if (!XProperty || !YProperty)
-	{
-		return SNullWidget::NullWidget;
-	}
-
-	TSharedPtr<IPropertyHandle> XPropertyHandle = InTextureUV->GetPropertyHandle(UDMTextureUV::NAME_bMirrorOnX);
-	TSharedPtr<IPropertyHandle> YPropertyHandle = InTextureUV->GetPropertyHandle(UDMTextureUV::NAME_bMirrorOnY);
-	TSharedPtr<IDetailTreeNode> XDetailTreeNode = InTextureUV->GetDetailTreeNode(UDMTextureUV::NAME_bMirrorOnX);
-	TSharedPtr<IDetailTreeNode> YDetailTreeNode = InTextureUV->GetDetailTreeNode(UDMTextureUV::NAME_bMirrorOnY);
-
-	if (!XPropertyHandle.IsValid() || !YPropertyHandle.IsValid() || !XDetailTreeNode.IsValid() || !YDetailTreeNode.IsValid())
-	{
-		return SNullWidget::NullWidget;
-	}
-
-	ensure(XPropertyHandle->GetProperty() == XProperty && YPropertyHandle->GetProperty() == YProperty);
-
-	FPropertyEditorModule& PropertyEditorModule = FModuleManager::GetModuleChecked<FPropertyEditorModule>("PropertyEditor");
-
-	TSharedRef<SWidget> ResetButtonWidget = PropertyCustomizationHelpers::MakeResetButton(
-		FSimpleDelegate::CreateLambda([XPropertyHandle, YPropertyHandle]()
-			{
-				if (XPropertyHandle.IsValid() && XPropertyHandle->IsValidHandle()
-					&& YPropertyHandle.IsValid() && YPropertyHandle->IsValidHandle())
-				{
-					TArray<UObject*> Outers;
-					XPropertyHandle->GetOuterObjects(Outers);
-
-					if (Outers.IsEmpty() == false)
-					{
-						if (UDMTextureUV* TextureUV = Cast<UDMTextureUV>(Outers[0]))
-						{
-							if (TextureUV->GetMirrorOnX() || TextureUV->GetMirrorOnY())
-							{
-								FDMScopedUITransaction Transaction(LOCTEXT("ResetMirror", "Reset Texture UV Mirror to default."));
-								TextureUV->Modify();
-								TextureUV->SetMirrorOnX(false);
-								TextureUV->SetMirrorOnY(false);
-							}
-						}
-					}
-				}
-			})
-	);
-	ResetButtonWidget->SetVisibility(
-		TAttribute<EVisibility>::CreateLambda([XPropertyHandle, YPropertyHandle]()
-			{
-				if (XPropertyHandle.IsValid() && XPropertyHandle->IsValidHandle()
-					&& YPropertyHandle.IsValid() && YPropertyHandle->IsValidHandle())
-				{
-					TArray<UObject*> Outers;
-					XPropertyHandle->GetOuterObjects(Outers);
-
-					if (Outers.IsEmpty() == false)
-					{
-						if (UDMTextureUV* TextureUV = Cast<UDMTextureUV>(Outers[0]))
-						{
-							return (TextureUV->GetMirrorOnX() || TextureUV->GetMirrorOnY())
-								? EVisibility::Visible
-								: EVisibility::Hidden;
-						}
-					}
-				}
-
-				return EVisibility::Hidden;
-			})
-	);
-
-	TSharedRef<SHorizontalBox> ButtonsBox = SNew(SHorizontalBox)
-		+ SHorizontalBox::Slot()
-		[
-			ResetButtonWidget
-		];
-
-	ButtonsBox->AddSlot()
-		[
-			SNew(SBox)
-			.WidthOverride(20.f)
-			.HeightOverride(18.f)
-		];
-
-	return ButtonsBox;
+	InOutPropertyRows.Add(SDMEditor::GetPropertyHandle(&*InComponentEditWidget, InTextureUV, InProperty));
 }
 
-void UE::DynamicMaterialEditor::Private::AddVisualizerRow(const TSharedRef<SDMComponentEdit>& InComponentEditWidget, UDMTextureUV* InTextureUV, TArray<FDMPropertyHandle>& InOutPropertyRows)
+void UE::DynamicMaterialEditor::Private::AddVisualizerRow(const TSharedRef<SWidget>& InComponentEditWidget, UDMTextureUV* InTextureUV, TArray<FDMPropertyHandle>& InOutPropertyRows)
 {
 	// Make sure we don't get a substage
 	UDMMaterialStage* Stage = InTextureUV->GetTypedParent<UDMMaterialStage>(/* Allow Subclasses */ false);
diff --git a/Engine/Plugins/Experimental/DynamicMaterial/Source/DynamicMaterialEditor/Private/Slate/Properties/Generators/DMTextureUVPropertyRowGenerator.h b/Engine/Plugins/Experimental/DynamicMaterial/Source/DynamicMaterialEditor/Private/Slate/Properties/Generators/DMTextureUVPropertyRowGenerator.h
--- a/Engine/Plugins/Experimental/DynamicMaterial/Source/DynamicMaterialEditor/Private/Slate/Properties/Generators/DMTextureUVPropertyRowGenerator.h
+++ b/Engine/Plugins/Experimental/DynamicMaterial/Source/DynamicMaterialEditor/Private/Slate/Properties/Generators/DMTextureUVPropertyRowGenerator.h
@@ -4,6 +4,8 @@
 
 #include "Slate/Properties/Generators/DMComponentPropertyRowGenerator.h"
 
+class SWidget;
+
 class FDMTextureUVPropertyRowGenerator : public FDMComponentPropertyRowGenerator
 {
 public:
@@ -15,5 +17,8 @@ public:
 	virtual void AddComponentProperties(const TSharedRef<SDMComponentEdit>& InComponentEditWidget, UDMMaterialComponent* InComponent,
 		TArray<FDMPropertyHandle>& InOutPropertyRows, TSet<UDMMaterialComponent*>& InOutProcessedObjects) override;
 
+	static void AddPopoutComponentProperties(const TSharedRef<SWidget>& InParentWidget, UDMMaterialComponent* InComponent,
+		TArray<FDMPropertyHandle>& InOutPropertyRows);
+
 	virtual bool AllowKeyframeButton(UDMMaterialComponent* InComponent, FProperty* InProperty) override;
 };
diff --git a/Engine/Plugins/Experimental/DynamicMaterial/Source/DynamicMaterialEditor/Private/Slate/Properties/SDMTextureUVVisualizerPopout.cpp b/Engine/Plugins/Experimental/DynamicMaterial/Source/DynamicMaterialEditor/Private/Slate/Properties/SDMTextureUVVisualizerPopout.cpp
--- a/Engine/Plugins/Experimental/DynamicMaterial/Source/DynamicMaterialEditor/Private/Slate/Properties/SDMTextureUVVisualizerPopout.cpp
+++ b/Engine/Plugins/Experimental/DynamicMaterial/Source/DynamicMaterialEditor/Private/Slate/Properties/SDMTextureUVVisualizerPopout.cpp
@@ -1,48 +1,83 @@
 // Copyright Epic Games, Inc. All Rights Reserved.
 
 #include "Slate/Properties/SDMTextureUVVisualizerPopout.h"
+
+#include "Components/DMMaterialStage.h"
+#include "Components/DMTextureUV.h"
+#include "CustomDetailsViewArgs.h"
+#include "CustomDetailsViewModule.h"
 #include "DetailLayoutBuilder.h"
+#include "DMWorldSubsystem.h"
+#include "DynamicMaterialEditorModule.h"
+#include "Engine/World.h"
+#include "Generators/DMTextureUVPropertyRowGenerator.h"
+#include "ICustomDetailsView.h"
+#include "Items/ICustomDetailsViewCustomItem.h"
+#include "Items/ICustomDetailsViewItem.h"
+#include "Materials/Material.h"
 #include "Slate/Properties/SDMTextureUVVisualizer.h"
+#include "Slate/SDMEditor.h"
+#include "Styling/SlateIconFinder.h"
 #include "Widgets/Colors/SColorBlock.h"
+#include "Widgets/Docking/SDockTab.h"
 #include "Widgets/Input/SButton.h"
 #include "Widgets/Layout/SBox.h"
 #include "Widgets/Text/STextBlock.h"
 
 #define LOCTEXT_NAMESPACE "SDMTextureUVVisualizerPopout"
 
+const FName SDMTextureUVVisualizerPopout::TabId = TEXT("SDMTextureUVVisualizerPopout");
+
+void SDMTextureUVVisualizerPopout::CreatePopout(UDMMaterialStage* InMaterialStage, UDMTextureUV* InTextureUV)
+{
+	if (!IsValid(InMaterialStage) || !IsValid(InTextureUV))
+	{
+		return;
+	}
+
+	if (!FGlobalTabmanager::Get()->HasTabSpawner(TabId))
+	{
+		FGlobalTabmanager::Get()->RegisterNomadTabSpawner(
+			TabId,
+			FOnSpawnTab::CreateLambda(
+				[](const FSpawnTabArgs& InArgs)
+				{
+					TSharedRef<SDockTab> DockTab = SNew(SDockTab)
+						.Label(FText::FromName(TabId))
+						.LabelSuffix(LOCTEXT("TabSuffix", "[UV Vis]"));
+
+					DockTab->SetTabIcon(FSlateIconFinder::FindIconForClass(UMaterial::StaticClass()).GetIcon());
+
+					return DockTab;
+				}
+			)
+		);
+	}
+
+	TSharedPtr<SDockTab> Tab = FGlobalTabmanager::Get()->TryInvokeTab(TabId);
+
+	if (!Tab.IsValid())
+	{
+		return;
+	}
+
+	Tab->ActivateInParent(ETabActivationCause::SetDirectly);
+	Tab->SetLabel(FText::FromString(InTextureUV->GetPathName()));
+	Tab->SetContent(SNew(SDMTextureUVVisualizerPopout, InMaterialStage, InTextureUV));
+}
+
 void SDMTextureUVVisualizerPopout::Construct(const FArguments& InArgs, UDMMaterialStage* InMaterialStage, UDMTextureUV* InTextureUV)
 {
 	check(InMaterialStage);
 	check(InTextureUV);
 
-	SWindow::Construct(
-		SWindow::FArguments()
-		.Title(LOCTEXT("WindowTitle", "Material Designer Texture UV Visualizer"))
-		.ClientSize(FVector2f(1024, 768))
-		.MinWidth(128)
-		.MinHeight(160)
-	);
-
-	SetContent(
-		SNew(SVerticalBox)
-		+ SVerticalBox::Slot()
-		.AutoHeight()
-		.HAlign(EHorizontalAlignment::HAlign_Fill)
-		[
-			SNew(SButton)
-			.OnClicked(this, &SDMTextureUVVisualizerPopout::OnToggleModeClicked)
-			.Content()
-			[
-				SNew(STextBlock)
-				.Text(this, &SDMTextureUVVisualizerPopout::GetModeButtonText)
-				.Font(IDetailLayoutBuilder::GetDetailFontBold())
-			]
-		]
-		+ SVerticalBox::Slot()
-		.VAlign(EVerticalAlignment::VAlign_Fill)
+	ChildSlot
+	[
+		SNew(SHorizontalBox)
+		+ SHorizontalBox::Slot()
 		.HAlign(EHorizontalAlignment::HAlign_Fill)
 		.VAlign(EVerticalAlignment::VAlign_Fill)
-		.Padding(0.f, 3.f, 0.f, 0.f)
+		.Padding(3.f, 0.f, 0.f, 0.f)
 		[
 			SNew(SOverlay)
 			+ SOverlay::Slot()
@@ -55,6 +90,7 @@ void SDMTextureUVVisualizerPopout::Construct(const FArguments& InArgs, UDMMateri
 				SNew(SBox)
 				.HAlign(EHorizontalAlignment::HAlign_Fill)
 				.VAlign(EVerticalAlignment::VAlign_Top)
+				.Visibility(EVisibility::HitTestInvisible)
 				[
 					SNew(SColorBlock)
 					.Color(FLinearColor(0, 0, 0, 0.5))
@@ -66,6 +102,7 @@ void SDMTextureUVVisualizerPopout::Construct(const FArguments& InArgs, UDMMateri
 				SNew(SBox)
 				.HAlign(EHorizontalAlignment::HAlign_Fill)
 				.VAlign(EVerticalAlignment::VAlign_Bottom)
+				.Visibility(EVisibility::HitTestInvisible)
 				[
 					SNew(SColorBlock)
 					.Color(FLinearColor(0, 0, 0, 0.5))
@@ -77,6 +114,7 @@ void SDMTextureUVVisualizerPopout::Construct(const FArguments& InArgs, UDMMateri
 				SNew(SBox)
 				.HAlign(EHorizontalAlignment::HAlign_Left)
 				.VAlign(EVerticalAlignment::VAlign_Center)
+				.Visibility(EVisibility::HitTestInvisible)
 				[
 					SNew(SColorBlock)
 					.Color(FLinearColor(0, 0, 0, 0.5))
@@ -88,6 +126,7 @@ void SDMTextureUVVisualizerPopout::Construct(const FArguments& InArgs, UDMMateri
 				SNew(SBox)
 				.HAlign(EHorizontalAlignment::HAlign_Right)
 				.VAlign(EVerticalAlignment::VAlign_Center)
+				.Visibility(EVisibility::HitTestInvisible)
 				[
 					SNew(SColorBlock)
 					.Color(FLinearColor(0, 0, 0, 0.5))
@@ -95,7 +134,20 @@ void SDMTextureUVVisualizerPopout::Construct(const FArguments& InArgs, UDMMateri
 				]
 			]
 		]
-	);
+		+ SHorizontalBox::Slot()
+		.AutoWidth()
+		.HAlign(EHorizontalAlignment::HAlign_Fill)
+		.VAlign(EVerticalAlignment::VAlign_Top)
+		[
+			SNew(SBox)
+			.WidthOverride(300.f)
+			.HAlign(EHorizontalAlignment::HAlign_Fill)
+			.VAlign(EVerticalAlignment::VAlign_Top)
+			[
+				CreatePropertyWidget(InTextureUV)
+			]
+		]
+	];
 }
 
 FReply SDMTextureUVVisualizerPopout::OnToggleModeClicked()
@@ -158,4 +210,119 @@ FVector2D SDMTextureUVVisualizerPopout::GetSideBlockSize() const
 	return FVector2D::UnitVector;
 }
 
+TSharedRef<SWidget> SDMTextureUVVisualizerPopout::CreatePropertyWidget(UDMTextureUV* InTextureUV)
+{
+	SDMEditor::ClearPropertyHandles(this);
+
+	FCustomDetailsViewArgs Args;
+	Args.KeyframeHandler = nullptr;
+	Args.bAllowGlobalExtensions = true;
+	Args.bAllowResetToDefault = true;
+	Args.bShowCategories = false;
+
+	if (UWorld* World = InTextureUV->GetWorld())
+	{
+		if (UDMWorldSubsystem* WorldSubsystem = World->GetSubsystem<UDMWorldSubsystem>())
+		{
+			Args.KeyframeHandler = WorldSubsystem->GetKeyframeHandler();
+		}
+	}
+
+	TSharedRef<ICustomDetailsView> DetailsView = ICustomDetailsViewModule::Get().CreateCustomDetailsView(Args);
+	const FCustomDetailsViewItemId RootId = DetailsView->GetRootItem()->GetItemId();
+
+	TArray<FDMPropertyHandle> TextureUVPropertyRows;
+	FDMTextureUVPropertyRowGenerator::AddPopoutComponentProperties(SharedThis(this), InTextureUV, TextureUVPropertyRows);
+
+	FDMPropertyHandle EditModeButtonRow;
+	EditModeButtonRow.ValueName = TEXT("EditMode");
+	EditModeButtonRow.NameOverride = LOCTEXT("EditMode", "Edit Mode");
+	EditModeButtonRow.ValueWidget = SNew(SButton)
+		.OnClicked(this, &SDMTextureUVVisualizerPopout::OnToggleModeClicked)
+		.Content()
+		[
+			SNew(STextBlock)
+			.Text(this, &SDMTextureUVVisualizerPopout::GetModeButtonText)
+			.Font(IDetailLayoutBuilder::GetDetailFontBold())
+		];
+
+	TextureUVPropertyRows.Add(EditModeButtonRow);
+
+	for (const FDMPropertyHandle& EditRow : TextureUVPropertyRows)
+	{
+		const bool bHasValidCustomWidget = EditRow.ValueWidget.IsValid() && !EditRow.ValueName.IsNone() && EditRow.NameOverride.IsSet();
+
+		if (!EditRow.DetailTreeNode && !bHasValidCustomWidget)
+		{
+			continue;
+		}
+
+		ECustomDetailsTreeInsertPosition Position = ECustomDetailsTreeInsertPosition::Child;
+
+		if (EditRow.DetailTreeNode)
+		{
+			if (EditRow.DetailTreeNode->CreatePropertyHandle()->HasMetaData("HighPriority"))
+			{
+				Position = ECustomDetailsTreeInsertPosition::FirstChild;
+			}
+			else if (EditRow.DetailTreeNode->CreatePropertyHandle()->HasMetaData("LowPriority"))
+			{
+				Position = ECustomDetailsTreeInsertPosition::LastChild;
+			}
+		}
+
+		if (bHasValidCustomWidget)
+		{
+			TSharedPtr<ICustomDetailsViewCustomItem> Item = DetailsView->CreateCustomItem(
+				EditRow.ValueName, 
+				EditRow.NameOverride.GetValue(), 
+				EditRow.NameToolTipOverride.Get(FText::GetEmpty())
+			);
+
+			if (!Item.IsValid())
+			{
+				continue;
+			}
+
+			Item->SetValueWidget(EditRow.ValueWidget.ToSharedRef());
+			DetailsView->ExtendTree(RootId, Position, Item->AsItem());
+			continue;
+		}
+
+		if (!EditRow.DetailTreeNode)
+		{
+			continue;
+		}
+
+		TSharedRef<ICustomDetailsViewItem> Item = DetailsView->CreateDetailTreeItem(EditRow.DetailTreeNode.ToSharedRef());
+
+		if (EditRow.NameOverride.IsSet())
+		{
+			Item->SetOverrideWidget(
+				ECustomDetailsViewWidgetType::Name,
+				SNew(STextBlock)
+				.Font(IDetailLayoutBuilder::GetDetailFont())
+				.Text(EditRow.NameOverride.GetValue())
+				.ToolTipText(EditRow.NameToolTipOverride.Get(FText::GetEmpty()))
+			);
+		}
+
+		if (EditRow.DetailTreeNode->CreatePropertyHandle()->HasMetaData("NotKeyframeable"))
+		{
+			Item->SetKeyframeEnabled(false);
+		}
+
+		if (EditRow.ResetToDefaultOverride.IsSet())
+		{
+			Item->SetResetToDefaultOverride(EditRow.ResetToDefaultOverride.GetValue());
+		}
+
+		DetailsView->ExtendTree(RootId, Position, Item);
+	}
+
+	DetailsView->RebuildTree(ECustomDetailsViewBuildType::InstantBuild);
+
+	return DetailsView;
+}
+
 #undef LOCTEXT_NAMESPACE
diff --git a/Engine/Plugins/Experimental/DynamicMaterial/Source/DynamicMaterialEditor/Private/Slate/Properties/SDMTextureUVVisualizerPopout.h b/Engine/Plugins/Experimental/DynamicMaterial/Source/DynamicMaterialEditor/Private/Slate/Properties/SDMTextureUVVisualizerPopout.h
--- a/Engine/Plugins/Experimental/DynamicMaterial/Source/DynamicMaterialEditor/Private/Slate/Properties/SDMTextureUVVisualizerPopout.h
+++ b/Engine/Plugins/Experimental/DynamicMaterial/Source/DynamicMaterialEditor/Private/Slate/Properties/SDMTextureUVVisualizerPopout.h
@@ -2,9 +2,12 @@
 
 #pragma once
 
-#include "Widgets/SWindow.h"
+#include "Templates/SharedPointerFwd.h"
+#include "Widgets/SCompoundWidget.h"
 
+class ICustomDetailsView;
 class SDMTextureUVVisualizer;
+class SDockTab;
 class UDMMaterialStage;
 class UDMTextureUV;
 
@@ -16,9 +19,12 @@ class UDMTextureUV;
  * The popout specifically expands the visible area of the preview to 3x the normal size
  * on the smallest axis. The other axis is expanded to match the aspect ratio.
  */
-class SDMTextureUVVisualizerPopout : public SWindow
+class SDMTextureUVVisualizerPopout : public SCompoundWidget
 {
 public:
+	static const FName TabId;
+
+	static void CreatePopout(UDMMaterialStage* InMaterialStage, UDMTextureUV* InTextureUV);
 
 	SLATE_BEGIN_ARGS(SDMTextureUVVisualizerPopout) {}
 	SLATE_END_ARGS()
@@ -36,4 +42,6 @@ protected:
 	FVector2D GetHorizontalBarSize() const;
 
 	FVector2D GetSideBlockSize() const;
+
+	TSharedRef<SWidget> CreatePropertyWidget(UDMTextureUV* InTextureUV);
 };
diff --git a/Engine/Plugins/Experimental/DynamicMaterial/Source/DynamicMaterialEditor/Private/Slate/Properties/SDMTextureUVVisualizerProperty.cpp b/Engine/Plugins/Experimental/DynamicMaterial/Source/DynamicMaterialEditor/Private/Slate/Properties/SDMTextureUVVisualizerProperty.cpp
--- a/Engine/Plugins/Experimental/DynamicMaterial/Source/DynamicMaterialEditor/Private/Slate/Properties/SDMTextureUVVisualizerProperty.cpp
+++ b/Engine/Plugins/Experimental/DynamicMaterial/Source/DynamicMaterialEditor/Private/Slate/Properties/SDMTextureUVVisualizerProperty.cpp
@@ -7,6 +7,7 @@
 #include "DynamicMaterialEditorSettings.h"
 #include "Slate/Properties/SDMTextureUVVisualizer.h"
 #include "Slate/Properties/SDMTextureUVVisualizerPopout.h"
+#include "Widgets/Docking/SDockTab.h"
 #include "Widgets/Input/SButton.h"
 #include "Widgets/Layout/SBox.h"
 #include "Widgets/Text/STextBlock.h"
@@ -32,7 +33,6 @@ void SDMTextureUVVisualizerProperty::Construct(const FArguments& InArgs, UDMMate
 			[
 				SNew(SButton)
 				.OnClicked(this, &SDMTextureUVVisualizerProperty::OnOpenPopoutClicked)
-				.IsEnabled(this, &SDMTextureUVVisualizerProperty::IsPopoutEnabled)
 				.Content()
 				[
 					SNew(STextBlock)
@@ -121,11 +121,6 @@ FText SDMTextureUVVisualizerProperty::GetModeButtonText() const
 	return LOCTEXT("VisualizerOffset", "Offset");
 }
 
-bool SDMTextureUVVisualizerProperty::IsPopoutEnabled() const
-{
-	return !PopoutWindowWeak.IsValid();
-}
-
 FReply SDMTextureUVVisualizerProperty::OnOpenPopoutClicked()
 {
 	if (!Visualizer.IsValid())
@@ -133,11 +128,6 @@ FReply SDMTextureUVVisualizerProperty::OnOpenPopoutClicked()
 		return FReply::Handled();
 	}
 
-	if (PopoutWindowWeak.IsValid())
-	{
-		return FReply::Handled();
-	}
-
 	UDMMaterialStage* Stage = Visualizer->GetStage();
 	UDMTextureUV* TextureUV = Visualizer->GetTextureUV();
 
@@ -146,10 +136,7 @@ FReply SDMTextureUVVisualizerProperty::OnOpenPopoutClicked()
 		return FReply::Handled();
 	}
 
-	TSharedRef<SDMTextureUVVisualizerPopout> PopoutWindow = SNew(SDMTextureUVVisualizerPopout, Stage, TextureUV);
-	FSlateApplication::Get().AddWindow(PopoutWindow, /* Show Immediately */ true);
-
-	PopoutWindowWeak = PopoutWindow;
+	SDMTextureUVVisualizerPopout::CreatePopout(Stage, TextureUV);
 
 	return FReply::Handled();
 }
diff --git a/Engine/Plugins/Experimental/DynamicMaterial/Source/DynamicMaterialEditor/Private/Slate/Properties/SDMTextureUVVisualizerProperty.h b/Engine/Plugins/Experimental/DynamicMaterial/Source/DynamicMaterialEditor/Private/Slate/Properties/SDMTextureUVVisualizerProperty.h
--- a/Engine/Plugins/Experimental/DynamicMaterial/Source/DynamicMaterialEditor/Private/Slate/Properties/SDMTextureUVVisualizerProperty.h
+++ b/Engine/Plugins/Experimental/DynamicMaterial/Source/DynamicMaterialEditor/Private/Slate/Properties/SDMTextureUVVisualizerProperty.h
@@ -26,7 +26,6 @@ public:
 
 protected:
 	TSharedPtr<SDMTextureUVVisualizer> Visualizer;
-	TWeakPtr<SDMTextureUVVisualizerPopout> PopoutWindowWeak;
 
 	FReply OnToggleVisualizerClicked();
 
@@ -34,7 +33,6 @@ protected:
 
 	FText GetModeButtonText() const;
 
-	bool IsPopoutEnabled() const;
 	FReply OnOpenPopoutClicked();
 
 	EVisibility GetVisualizerVisibility() const;
diff --git a/Engine/Plugins/Experimental/DynamicMaterial/Source/DynamicMaterialEditor/Private/Slate/SDMComponentEdit.cpp b/Engine/Plugins/Experimental/DynamicMaterial/Source/DynamicMaterialEditor/Private/Slate/SDMComponentEdit.cpp
--- a/Engine/Plugins/Experimental/DynamicMaterial/Source/DynamicMaterialEditor/Private/Slate/SDMComponentEdit.cpp
+++ b/Engine/Plugins/Experimental/DynamicMaterial/Source/DynamicMaterialEditor/Private/Slate/SDMComponentEdit.cpp
@@ -331,7 +331,11 @@ TSharedRef<SWidget> SDMComponentEdit::CreateEditWidget()
 
 		if (bHasValidCustomWidget)
 		{
-			TSharedPtr<ICustomDetailsViewCustomItem> Item = DetailsView->CreateCustomItem(EditRow.ValueName, EditRow.NameOverride.GetValue(), EditRow.NameToolTipOverride.GetValue());
+			TSharedPtr<ICustomDetailsViewCustomItem> Item = DetailsView->CreateCustomItem(
+				EditRow.ValueName,
+				EditRow.NameOverride.GetValue(),
+				EditRow.NameToolTipOverride.Get(FText::GetEmpty())
+			);
 
 			if (!Item.IsValid())
 			{
