Material Designer
- Textures and material functions can now be dropped into the slot view to add them as a new layer.

diff --git a/Engine/Plugins/Experimental/DynamicMaterial/Source/DynamicMaterialEditor/Private/Slate/Layers/SDMSlotLayerView.cpp b/Engine/Plugins/Experimental/DynamicMaterial/Source/DynamicMaterialEditor/Private/Slate/Layers/SDMSlotLayerView.cpp
--- a/Engine/Plugins/Experimental/DynamicMaterial/Source/DynamicMaterialEditor/Private/Slate/Layers/SDMSlotLayerView.cpp
+++ b/Engine/Plugins/Experimental/DynamicMaterial/Source/DynamicMaterialEditor/Private/Slate/Layers/SDMSlotLayerView.cpp
@@ -1,19 +1,30 @@
 // Copyright Epic Games, Inc. All Rights Reserved.
 
 #include "Slate/Layers/SDMSlotLayerView.h"
+#include "Components/DMMaterialEffect.h"
+#include "Components/DMMaterialEffectStack.h"
 #include "Components/DMMaterialLayer.h"
 #include "Components/DMMaterialStage.h"
+#include "Components/DMMaterialStageFunction.h"
+#include "Components/DMMaterialSubStage.h"
+#include "Components/MaterialStageBlends/DMMSBNormal.h"
+#include "Components/MaterialStageExpressions/DMMSETextureSample.h"
+#include "Components/MaterialStageInputs/DMMSIExpression.h"
+#include "Components/MaterialStageInputs/DMMSIFunction.h"
+#include "Components/MaterialStageInputs/DMMSIValue.h"
+#include "Components/MaterialValues/DMMaterialValueTexture.h"
 #include "DMPrivate.h"
+#include "DragAndDrop/AssetDragDropOp.h"
 #include "DynamicMaterialEditorCommands.h"
+#include "DynamicMaterialEditorModule.h"
 #include "DynamicMaterialEditorSettings.h"
 #include "Framework/Application/SlateApplication.h"
+#include "Materials/MaterialFunctionInterface.h"
 #include "Menus/DMMaterialSlotLayerMenus.h"
 #include "SDMLayerEffectsItem.h"
 #include "Slate/Layers/SDMSlotLayerItem.h"
 #include "Slate/SDMSlot.h"
 #include "ToolMenus.h"
-#include "Components/DMMaterialEffect.h"
-#include "Components/DMMaterialEffectStack.h"
 #include "Widgets/SNullWidget.h"
 
 #define LOCTEXT_NAMESPACE "SDMSlotLayerView"
@@ -490,25 +501,25 @@ int32 SDMSlotLayerView::GetLayerItemIndex(const TSharedPtr<FDMMaterialLayerRefer
 	return INDEX_NONE;
 }
 
-FCursorReply SDMSlotLayerView::OnCursorQuery(const FGeometry& MyGeometry, const FPointerEvent& CursorEvent) const
+FCursorReply SDMSlotLayerView::OnCursorQuery(const FGeometry& InGeometry, const FPointerEvent& InCursorEvent) const
 {
-	if (IsRightClickScrolling() && CursorEvent.IsMouseButtonDown(EKeys::RightMouseButton))
+	if (IsRightClickScrolling() && InCursorEvent.IsMouseButtonDown(EKeys::RightMouseButton))
 	{
 		// We hide the native cursor as we'll be drawing the software EMouseCursor::GrabHandClosed cursor
 		return FCursorReply::Cursor(EMouseCursor::None);
 	}
-	return SListView::OnCursorQuery(MyGeometry, CursorEvent);
+	return SListView::OnCursorQuery(InGeometry, InCursorEvent);
 }
 
-FReply SDMSlotLayerView::OnMouseMove(const FGeometry& MyGeometry, const FPointerEvent& MouseEvent)
+FReply SDMSlotLayerView::OnMouseMove(const FGeometry& InGeometry, const FPointerEvent& InMouseEvent)
 {
-	if (MouseEvent.IsMouseButtonDown(EKeys::RightMouseButton) && !MouseEvent.IsTouchEvent())
+	if (InMouseEvent.IsMouseButtonDown(EKeys::RightMouseButton) && !InMouseEvent.IsTouchEvent())
 	{
 		// We only care about deltas along the scroll axis
-		FTableViewDimensions CursorDeltaDimensions(Orientation, MouseEvent.GetCursorDelta());
+		FTableViewDimensions CursorDeltaDimensions(Orientation, InMouseEvent.GetCursorDelta());
 		CursorDeltaDimensions.LineAxis = 0.f;
 
-		const float ScrollByAmount = CursorDeltaDimensions.ScrollAxis / MyGeometry.Scale;
+		const float ScrollByAmount = CursorDeltaDimensions.ScrollAxis / InGeometry.Scale;
 
 		// If scrolling with the right mouse button, we need to remember how much we scrolled.
 		// If we did not scroll at all, we will bring up the context menu when the mouse is released.
@@ -527,7 +538,7 @@ FReply SDMSlotLayerView::OnMouseMove(const FGeometry& MyGeometry, const FPointer
 
 			TickScrollDelta -= ScrollByAmount;
 
-			const float AmountScrolled = this->ScrollBy(MyGeometry, -ScrollByAmount, AllowOverscroll);
+			const float AmountScrolled = this->ScrollBy(InGeometry, -ScrollByAmount, AllowOverscroll);
 
 			FReply Reply = FReply::Handled();
 
@@ -536,7 +547,7 @@ FReply SDMSlotLayerView::OnMouseMove(const FGeometry& MyGeometry, const FPointer
 			if (this->HasMouseCapture() == false)
 			{
 				Reply.CaptureMouse(AsShared()).UseHighPrecisionMouseMovement(AsShared());
-				SoftwareCursorPosition = MyGeometry.AbsoluteToLocal(MouseEvent.GetScreenSpacePosition());
+				SoftwareCursorPosition = InGeometry.AbsoluteToLocal(InMouseEvent.GetScreenSpacePosition());
 				bShowSoftwareCursor    = true;
 			}
 
@@ -550,12 +561,12 @@ FReply SDMSlotLayerView::OnMouseMove(const FGeometry& MyGeometry, const FPointer
 		}
 	}
 
-	return SListView::OnMouseMove(MyGeometry, MouseEvent);
+	return SListView::OnMouseMove(InGeometry, InMouseEvent);
 }
 
-FReply SDMSlotLayerView::OnMouseButtonUp(const FGeometry& MyGeometry, const FPointerEvent& MouseEvent)
+FReply SDMSlotLayerView::OnMouseButtonUp(const FGeometry& InGeometry, const FPointerEvent& InMouseEvent)
 {
-	if (MouseEvent.GetEffectingButton() == EKeys::RightMouseButton)
+	if (InMouseEvent.GetEffectingButton() == EKeys::RightMouseButton)
 	{
 		FReply Reply = FReply::Handled().ReleaseMouseCapture();
 		AmountScrolledWhileRightMouseDown = 0;
@@ -564,8 +575,8 @@ FReply SDMSlotLayerView::OnMouseButtonUp(const FGeometry& MyGeometry, const FPoi
 		// If we have mouse capture, snap the mouse back to the closest location that is within the list's bounds
 		if (HasMouseCapture())
 		{
-			const FSlateRect ListScreenSpaceRect = MyGeometry.GetLayoutBoundingRect();
-			const FVector2D CursorPosition = MyGeometry.LocalToAbsolute(SoftwareCursorPosition);
+			const FSlateRect ListScreenSpaceRect = InGeometry.GetLayoutBoundingRect();
+			const FVector2D CursorPosition = InGeometry.LocalToAbsolute(SoftwareCursorPosition);
 
 			const FIntPoint BestPositionInList(
 				FMath::RoundToInt(FMath::Clamp(CursorPosition.X, ListScreenSpaceRect.Left, ListScreenSpaceRect.Right)),
@@ -578,7 +589,7 @@ FReply SDMSlotLayerView::OnMouseButtonUp(const FGeometry& MyGeometry, const FPoi
 		return Reply;
 	}
 
-	if (MouseEvent.GetEffectingButton() == EKeys::LeftMouseButton)
+	if (InMouseEvent.GetEffectingButton() == EKeys::LeftMouseButton)
 	{
 		bool bHovering = false;
 		FVector2f MousePosition = FSlateApplication::Get().GetCursorPos();
@@ -595,7 +606,119 @@ FReply SDMSlotLayerView::OnMouseButtonUp(const FGeometry& MyGeometry, const FPoi
 		}
 	}
 
-	return SListView::OnMouseButtonUp(MyGeometry, MouseEvent);
+	return SListView::OnMouseButtonUp(InGeometry, InMouseEvent);
+}
+
+FReply SDMSlotLayerView::OnDrop(const FGeometry& InGeometry, const FDragDropEvent& InDragDropEvent)
+{
+	if (const TSharedPtr<FAssetDragDropOp> AssetDragDropOp = InDragDropEvent.GetOperationAs<FAssetDragDropOp>())
+	{
+		HandleAssetDragDropOperation(*AssetDragDropOp);
+		return FReply::Handled();
+	}
+
+	return FReply::Unhandled();
+}
+
+void SDMSlotLayerView::HandleAssetDragDropOperation(FAssetDragDropOp& InAssetDragDropOperation)
+{
+	TSharedPtr<SDMSlot> SlotWidget = SlotWidgetWeak.Pin();
+
+	if (!SlotWidget.IsValid())
+	{
+		return;
+	}
+
+	UDMMaterialSlot* Slot = MaterialSlotWeak.Get();
+
+	if (!IsValid(Slot))
+	{
+		return;
+	}
+
+	const TArray<FAssetData>& Assets = InAssetDragDropOperation.GetAssets();
+
+	if (Assets.IsEmpty())
+	{
+		return;
+	}
+
+	UObject* Asset = Assets[0].GetAsset();
+
+	if (!IsValid(Asset))
+	{
+		return;
+	}
+
+	if (UTexture* Texture = Cast<UTexture>(Asset))
+	{
+		FDMScopedUITransaction Transaction(LOCTEXT("DropTexture", "Material Designer Drop Texture"));
+		Slot->Modify();
+
+		UDMMaterialStage* NewStage = UDMMaterialStageBlend::CreateStage(UDMMaterialStageBlendNormal::StaticClass());
+		SlotWidget->AddNewLayer(NewStage);
+
+		UDMMaterialStageInputExpression* InputExpression = UDMMaterialStageInputExpression::ChangeStageInput_Expression(
+			NewStage,
+			UDMMaterialStageExpressionTextureSample::StaticClass(),
+			UDMMaterialStageBlend::InputB,
+			FDMMaterialStageConnectorChannel::WHOLE_CHANNEL,
+			0,
+			FDMMaterialStageConnectorChannel::WHOLE_CHANNEL
+		);
+
+		UDMMaterialSubStage* SubStage = InputExpression->GetSubStage();
+
+		if (ensure(SubStage))
+		{
+			UDMMaterialStageInputValue* InputValue = UDMMaterialStageInputValue::ChangeStageInput_NewLocalValue(
+				SubStage,
+				0,
+				FDMMaterialStageConnectorChannel::WHOLE_CHANNEL,
+				EDMValueType::VT_Texture,
+				FDMMaterialStageConnectorChannel::WHOLE_CHANNEL
+			);
+
+			if (ensure(InputValue))
+			{
+				UDMMaterialValueTexture* InputTexture = Cast<UDMMaterialValueTexture>(InputValue->GetValue());
+
+				if (ensure(InputTexture))
+				{
+
+					InputTexture->SetValue(Texture);
+				}
+			}
+		}
+	}
+	else if (UMaterialFunctionInterface* MaterialFunction = Cast<UMaterialFunctionInterface>(Asset))
+	{
+		FDMScopedUITransaction Transaction(LOCTEXT("DropFunction", "Material Designer Drop Material Function"));
+		Slot->Modify();
+
+		UDMMaterialStage* NewStage = UDMMaterialStageBlend::CreateStage(UDMMaterialStageBlendNormal::StaticClass());
+		UDMMaterialLayerObject* Layer = SlotWidget->AddNewLayer(NewStage);
+
+		if (ensure(Layer))
+		{
+			UDMMaterialStageInputFunction* NewFunction = UDMMaterialStageInputFunction::ChangeStageInput_Function(
+				NewStage,
+				MaterialFunction,
+				UDMMaterialStageBlend::InputB,
+				FDMMaterialStageConnectorChannel::WHOLE_CHANNEL,
+				0,
+				FDMMaterialStageConnectorChannel::WHOLE_CHANNEL
+			);
+
+			// The function was invalid and was removed. Remove the layer, refresh the appropriate ui parts.
+			if (!NewFunction->GetMaterialFunction())
+			{
+				Slot->RemoveLayer(Layer);
+				SlotWidget->InvalidateMainWidget();
+				SlotWidget->InvalidateComponentEditWidget();
+			}
+		}
+	}
 }
 
 void SDMSlotLayerView::RequestListRefresh()
@@ -705,9 +828,9 @@ bool SDMSlotLayerView::RemoveLayerItem(const TSharedPtr<FDMMaterialLayerReferenc
 	return true;
 }
 
-void SDMSlotLayerView::Tick(const FGeometry& AllottedGeometry, const double InCurrentTime, const float InDeltaTime)
+void SDMSlotLayerView::Tick(const FGeometry& InAllottedGeometry, const double InCurrentTime, const float InDeltaTime)
 {
-	SListView<TSharedPtr<FDMMaterialLayerReference, ESPMode::ThreadSafe>>::Tick(AllottedGeometry, InCurrentTime, InDeltaTime);
+	SListView<TSharedPtr<FDMMaterialLayerReference, ESPMode::ThreadSafe>>::Tick(InAllottedGeometry, InCurrentTime, InDeltaTime);
 
 	if (bPostRegenSelect && LayerItems.IsValidIndex(PostRegenSelectedLayerIndex))
 	{
diff --git a/Engine/Plugins/Experimental/DynamicMaterial/Source/DynamicMaterialEditor/Private/Slate/Layers/SDMSlotLayerView.h b/Engine/Plugins/Experimental/DynamicMaterial/Source/DynamicMaterialEditor/Private/Slate/Layers/SDMSlotLayerView.h
--- a/Engine/Plugins/Experimental/DynamicMaterial/Source/DynamicMaterialEditor/Private/Slate/Layers/SDMSlotLayerView.h
+++ b/Engine/Plugins/Experimental/DynamicMaterial/Source/DynamicMaterialEditor/Private/Slate/Layers/SDMSlotLayerView.h
@@ -55,10 +55,11 @@ public:
 	bool RemoveLayerItem(const TSharedPtr<FDMMaterialLayerReference>& InLayerItem);
 
 	//~ Begin SWidget
-	virtual void Tick(const FGeometry& AllottedGeometry, const double InCurrentTime, const float InDeltaTime) override;
-	virtual FCursorReply OnCursorQuery(const FGeometry& MyGeometry, const FPointerEvent& CursorEvent) const override;
-	virtual FReply OnMouseMove(const FGeometry& MyGeometry, const FPointerEvent& MouseEvent) override;
-	virtual FReply OnMouseButtonUp(const FGeometry& MyGeometry, const FPointerEvent& MouseEvent) override;
+	virtual void Tick(const FGeometry& InAllottedGeometry, const double InCurrentTime, const float InDeltaTime) override;
+	virtual FCursorReply OnCursorQuery(const FGeometry& InGeometry, const FPointerEvent& InCursorEvent) const override;
+	virtual FReply OnMouseMove(const FGeometry& InGeometry, const FPointerEvent& InMouseEvent) override;
+	virtual FReply OnMouseButtonUp(const FGeometry& InGeometry, const FPointerEvent& InMouseEvent) override;
+	virtual FReply OnDrop(const FGeometry& InGeometry, const FDragDropEvent& InDragDropEvent) override;
 	//~ End SWidget
 
 	virtual void RequestListRefresh() override;
@@ -114,4 +115,6 @@ protected:
 	void ExecuteSelectLayerStage(EDMMaterialLayerStage InStageType);
 
 	void OnUndo();
+
+	void HandleAssetDragDropOperation(FAssetDragDropOp& InAssetDragDropOperation);
 };
diff --git a/Engine/Plugins/Experimental/DynamicMaterial/Source/DynamicMaterialEditor/Private/Slate/SDMStage.cpp b/Engine/Plugins/Experimental/DynamicMaterial/Source/DynamicMaterialEditor/Private/Slate/SDMStage.cpp
--- a/Engine/Plugins/Experimental/DynamicMaterial/Source/DynamicMaterialEditor/Private/Slate/SDMStage.cpp
+++ b/Engine/Plugins/Experimental/DynamicMaterial/Source/DynamicMaterialEditor/Private/Slate/SDMStage.cpp
@@ -384,8 +384,6 @@ FReply SDMStage::OnDrop(const FGeometry& InMyGeometry, const FDragDropEvent& InD
 		return FReply::Unhandled();
 	}
 
-	SCompoundWidget::OnDrop(InMyGeometry, InDragDropEvent);
-
 	if (!StageWeak.IsValid())
 	{
 		return FReply::Unhandled();
@@ -398,8 +396,7 @@ FReply SDMStage::OnDrop(const FGeometry& InMyGeometry, const FDragDropEvent& InD
 		HandleAssetDragDropOperation(*AssetDragDropOp);
 		Reply = FReply::Handled();
 	}
-
-	if (OnAcceptDropEvent.IsBound())
+	else if (OnAcceptDropEvent.IsBound())
 	{
 		return OnAcceptDropEvent.Execute(InDragDropEvent, SharedThis(this));
 	}
