Material Designer
- A wizard now appears on first open of a material.
- The top settings in the material designer have been moved to the settings menu at the top-right
- Added two-sided flag to materials
- Re-arranged the material designer dropdown menu slightly. And fixed some issues.

diff --git a/Engine/Plugins/Experimental/DynamicMaterial/Source/DynamicMaterialEditor/Private/Menus/DMToolBarMenus.cpp b/Engine/Plugins/Experimental/DynamicMaterial/Source/DynamicMaterialEditor/Private/Menus/DMToolBarMenus.cpp
--- a/Engine/Plugins/Experimental/DynamicMaterial/Source/DynamicMaterialEditor/Private/Menus/DMToolBarMenus.cpp
+++ b/Engine/Plugins/Experimental/DynamicMaterial/Source/DynamicMaterialEditor/Private/Menus/DMToolBarMenus.cpp
@@ -32,9 +32,362 @@
 namespace UE::DynamicMaterialEditor::Private
 {
 	static FName ToolBarEditorLayoutMenuName = TEXT("MaterialDesigner.EditorLayout");
-	static FName ToolBarPreviewOptionsSectionName = TEXT("PreviewOptions");
-	static FName ToolBarTooltipOptionsSectionName = TEXT("TooltipOptions");
-	static FName ToolBarExportSectionName = TEXT("Export");
+	static FName ToolBarMaterialInstanceSectionName = TEXT("MaterialInstance");
+	static FName ToolBarMaterialExportSectionName = TEXT("MaterialExport");
+	static FName ToolBarMaterialDesignerSettingsSectionName = TEXT("MaterialDesignerSettings");
+
+	void AddToolBarMaterialInstanceDomainMenu_Execute(const FToolMenuContext& InContext, EMaterialDomain InDomain)
+	{
+		if (UDMMenuContext* Context = InContext.FindContext<UDMMenuContext>())
+		{
+			if (UDynamicMaterialModel* MaterialModel = Context->GetModel())
+			{
+				if (UDynamicMaterialModelEditorOnlyData* EditorOnlyData = UDynamicMaterialModelEditorOnlyData::Get(MaterialModel))
+				{
+					EditorOnlyData->SetDomain(InDomain);
+				}
+			}
+		}
+	}
+
+	bool AddToolBarMaterialInstanceDomainMenu_CanExecute(const FToolMenuContext& InContext, EMaterialDomain InDomain)
+	{
+		if (UDMMenuContext* Context = InContext.FindContext<UDMMenuContext>())
+		{
+			if (UDynamicMaterialModel* MaterialModel = Context->GetModel())
+			{
+				return !!UDynamicMaterialModelEditorOnlyData::Get(MaterialModel);
+			}
+		}
+
+		return false;
+	}
+
+	ECheckBoxState AddToolBarMaterialInstanceDomainMenu_GetCheckState(const FToolMenuContext& InContext, EMaterialDomain InDomain)
+	{
+		if (UDMMenuContext* Context = InContext.FindContext<UDMMenuContext>())
+		{
+			if (UDynamicMaterialModel* MaterialModel = Context->GetModel())
+			{
+				if (UDynamicMaterialModelEditorOnlyData* EditorOnlyData = UDynamicMaterialModelEditorOnlyData::Get(MaterialModel))
+				{
+					return EditorOnlyData->GetDomain() == InDomain
+						? ECheckBoxState::Checked
+						: ECheckBoxState::Unchecked;
+				}
+			}
+		}
+
+		return ECheckBoxState::Undetermined;
+	}
+
+	void AddToolBarMaterialInstanceDomainMenu(UToolMenu* InMenu)
+	{
+		UEnum* DomainEnum = StaticEnum<EMaterialDomain>();
+
+		FToolMenuSection& NewSection = InMenu->AddSection("MaterialDomains", LOCTEXT("MaterialDomains", "Material Domains"));
+
+		for (EMaterialDomain Domain : UDynamicMaterialModelEditorOnlyData::SupportedDomains)
+		{
+			FToolUIAction DomainAction;
+			DomainAction.ExecuteAction = FToolMenuExecuteAction::CreateStatic(&AddToolBarMaterialInstanceDomainMenu_Execute, Domain);
+			DomainAction.CanExecuteAction = FToolMenuCanExecuteAction::CreateStatic(&AddToolBarMaterialInstanceDomainMenu_CanExecute, Domain);
+			DomainAction.GetActionCheckState = FToolMenuGetActionCheckState::CreateStatic(&AddToolBarMaterialInstanceDomainMenu_GetCheckState, Domain);
+
+			NewSection.AddMenuEntry(
+				DomainEnum->GetNameByValue(Domain),
+				DomainEnum->GetDisplayNameTextByValue(Domain),
+				FText::GetEmpty(),
+				TAttribute<FSlateIcon>(),
+				FToolUIActionChoice(	DomainAction),
+				EUserInterfaceActionType::RadioButton
+			);
+		}
+	}
+
+	void AddToolBarMaterialInstanceBlendModeMenu_Execute(const FToolMenuContext& InContext, EBlendMode InBlendMode)
+	{
+		if (UDMMenuContext* Context = InContext.FindContext<UDMMenuContext>())
+		{
+			if (UDynamicMaterialModel* MaterialModel = Context->GetModel())
+			{
+				if (UDynamicMaterialModelEditorOnlyData* EditorOnlyData = UDynamicMaterialModelEditorOnlyData::Get(MaterialModel))
+				{
+					EditorOnlyData->SetBlendMode(InBlendMode);
+				}
+			}
+		}
+	}
+
+	bool AddToolBarMaterialInstanceBlendModeMenu_CanExecute(const FToolMenuContext& InContext, EBlendMode InBlendMode)
+	{
+		if (UDMMenuContext* Context = InContext.FindContext<UDMMenuContext>())
+		{
+			if (UDynamicMaterialModel* MaterialModel = Context->GetModel())
+			{
+				return !!UDynamicMaterialModelEditorOnlyData::Get(MaterialModel);
+			}
+		}
+
+		return false;
+	}
+
+	ECheckBoxState AddToolBarMaterialInstanceBlendModeMenu_GetCheckState(const FToolMenuContext& InContext, EBlendMode InBlendMode)
+	{
+		if (UDMMenuContext* Context = InContext.FindContext<UDMMenuContext>())
+		{
+			if (UDynamicMaterialModel* MaterialModel = Context->GetModel())
+			{
+				if (UDynamicMaterialModelEditorOnlyData* EditorOnlyData = UDynamicMaterialModelEditorOnlyData::Get(MaterialModel))
+				{
+					return EditorOnlyData->GetBlendMode() == InBlendMode
+						? ECheckBoxState::Checked
+						: ECheckBoxState::Unchecked;
+				}
+			}
+		}
+
+		return ECheckBoxState::Undetermined;
+	}
+
+	void AddToolBarMaterialInstanceBlendModeMenu(UToolMenu* InMenu)
+	{
+		UEnum* BlendModeEnum = StaticEnum<EBlendMode>();
+
+		FToolMenuSection& NewSection = InMenu->AddSection("MaterialBlendModes", LOCTEXT("MaterialBlendModes", "Material Blend Modes"));
+
+		for (EBlendMode BlendMode : UDynamicMaterialModelEditorOnlyData::SupportedBlendModes)
+		{
+			FToolUIAction BlendModeAction;
+			BlendModeAction.ExecuteAction = FToolMenuExecuteAction::CreateStatic(&AddToolBarMaterialInstanceBlendModeMenu_Execute, BlendMode);
+			BlendModeAction.CanExecuteAction = FToolMenuCanExecuteAction::CreateStatic(&AddToolBarMaterialInstanceBlendModeMenu_CanExecute, BlendMode);
+			BlendModeAction.GetActionCheckState = FToolMenuGetActionCheckState::CreateStatic(&AddToolBarMaterialInstanceBlendModeMenu_GetCheckState, BlendMode);
+
+			NewSection.AddMenuEntry(
+				BlendModeEnum->GetNameByValue(BlendMode),
+				BlendModeEnum->GetDisplayNameTextByValue(BlendMode),
+				FText::GetEmpty(),
+				TAttribute<FSlateIcon>(),
+				FToolUIActionChoice(BlendModeAction),
+				EUserInterfaceActionType::RadioButton
+			);
+		}
+	}
+
+	void AddToolBarMaterialInstanceUnlitMenu_Execute(const FToolMenuContext& InContext)
+	{
+		if (UDMMenuContext* Context = InContext.FindContext<UDMMenuContext>())
+		{
+			if (UDynamicMaterialModel* MaterialModel = Context->GetModel())
+			{
+				if (UDynamicMaterialModelEditorOnlyData* EditorOnlyData = UDynamicMaterialModelEditorOnlyData::Get(MaterialModel))
+				{
+					switch (EditorOnlyData->GetShadingModel())
+					{
+						case EDMMaterialShadingModel::DefaultLit:
+							EditorOnlyData->SetShadingModel(EDMMaterialShadingModel::Unlit);
+							break;
+
+						case EDMMaterialShadingModel::Unlit:
+							EditorOnlyData->SetShadingModel(EDMMaterialShadingModel::DefaultLit);
+							break;
+					}
+				}
+			}
+		}
+	}
+
+	bool AddToolBarMaterialInstanceUnlitMenu_CanExecute(const FToolMenuContext& InContext)
+	{
+		if (UDMMenuContext* Context = InContext.FindContext<UDMMenuContext>())
+		{
+			if (UDynamicMaterialModel* MaterialModel = Context->GetModel())
+			{
+				return !!UDynamicMaterialModelEditorOnlyData::Get(MaterialModel);
+			}
+		}
+
+		return false;
+	}
+
+	ECheckBoxState AddToolBarMaterialInstanceUnlitMenu_GetCheckState(const FToolMenuContext& InContext)
+	{
+		if (UDMMenuContext* Context = InContext.FindContext<UDMMenuContext>())
+		{
+			if (UDynamicMaterialModel* MaterialModel = Context->GetModel())
+			{
+				if (UDynamicMaterialModelEditorOnlyData* EditorOnlyData = UDynamicMaterialModelEditorOnlyData::Get(MaterialModel))
+				{
+					switch (EditorOnlyData->GetShadingModel())
+					{
+						case EDMMaterialShadingModel::DefaultLit:
+							return ECheckBoxState::Unchecked;
+
+						case EDMMaterialShadingModel::Unlit:
+							return ECheckBoxState::Checked;
+					}
+				}
+			}
+		}
+
+		return ECheckBoxState::Undetermined;
+	}
+
+	void AddToolBarMaterialInstanceAnimatedMenu_Execute(const FToolMenuContext& InContext)
+	{
+		if (UDMMenuContext* Context = InContext.FindContext<UDMMenuContext>())
+		{
+			if (UDynamicMaterialModel* MaterialModel = Context->GetModel())
+			{
+				if (UDynamicMaterialModelEditorOnlyData* EditorOnlyData = UDynamicMaterialModelEditorOnlyData::Get(MaterialModel))
+				{
+					EditorOnlyData->SetPixelAnimationFlag(!EditorOnlyData->IsPixelAnimationFlagSet());
+				}
+			}
+		}
+	}
+
+	bool AddToolBarMaterialInstanceAnimatedMenu_CanExecute(const FToolMenuContext& InContext)
+	{
+		if (UDMMenuContext* Context = InContext.FindContext<UDMMenuContext>())
+		{
+			if (UDynamicMaterialModel* MaterialModel = Context->GetModel())
+			{
+				return !!UDynamicMaterialModelEditorOnlyData::Get(MaterialModel);
+			}
+		}
+
+		return false;
+	}
+
+	ECheckBoxState AddToolBarMaterialInstanceAnimatedMenu_GetCheckState(const FToolMenuContext& InContext)
+	{
+		if (UDMMenuContext* Context = InContext.FindContext<UDMMenuContext>())
+		{
+			if (UDynamicMaterialModel* MaterialModel = Context->GetModel())
+			{
+				if (UDynamicMaterialModelEditorOnlyData* EditorOnlyData = UDynamicMaterialModelEditorOnlyData::Get(MaterialModel))
+				{
+					return EditorOnlyData->IsPixelAnimationFlagSet()
+						? ECheckBoxState::Checked
+						: ECheckBoxState::Unchecked;
+				}
+			}
+		}
+
+		return ECheckBoxState::Undetermined;
+	}
+
+	void AddToolBarMaterialInstanceTwoSidedMenu_Execute(const FToolMenuContext& InContext)
+	{
+		if (UDMMenuContext* Context = InContext.FindContext<UDMMenuContext>())
+		{
+			if (UDynamicMaterialModel* MaterialModel = Context->GetModel())
+			{
+				if (UDynamicMaterialModelEditorOnlyData* EditorOnlyData = UDynamicMaterialModelEditorOnlyData::Get(MaterialModel))
+				{
+					EditorOnlyData->SetTwoSidedFlag(!EditorOnlyData->IsTwoSidedFlagSet());
+				}
+			}
+		}
+	}
+
+	bool AddToolBarMaterialInstanceTwoSidedMenu_CanExecute(const FToolMenuContext& InContext)
+	{
+		if (UDMMenuContext* Context = InContext.FindContext<UDMMenuContext>())
+		{
+			if (UDynamicMaterialModel* MaterialModel = Context->GetModel())
+			{
+				return !!UDynamicMaterialModelEditorOnlyData::Get(MaterialModel);
+			}
+		}
+
+		return false;
+	}
+
+	ECheckBoxState AddToolBarMaterialInstanceTwoSidedMenu_GetCheckState(const FToolMenuContext& InContext)
+	{
+		if (UDMMenuContext* Context = InContext.FindContext<UDMMenuContext>())
+		{
+			if (UDynamicMaterialModel* MaterialModel = Context->GetModel())
+			{
+				if (UDynamicMaterialModelEditorOnlyData* EditorOnlyData = UDynamicMaterialModelEditorOnlyData::Get(MaterialModel))
+				{
+					return EditorOnlyData->IsTwoSidedFlagSet()
+						? ECheckBoxState::Checked
+						: ECheckBoxState::Unchecked;
+				}
+			}
+		}
+
+		return ECheckBoxState::Undetermined;
+	}
+
+	void AddToolBarInstanceMenu(UToolMenu* InMenu)
+	{
+		if (!IsValid(InMenu) || InMenu->ContainsSection(ToolBarMaterialInstanceSectionName))
+		{
+			return;
+		}
+
+		FToolMenuSection& NewSection = InMenu->AddSection(ToolBarMaterialInstanceSectionName, LOCTEXT("MaterialInstanceSection", "Material Instance"));
+
+		NewSection.AddSubMenu(
+			"Domain",
+			LOCTEXT("MaterialInstanceDomain", "Material Domain"),
+			LOCTEXT("MaterialInstanceDomain_ToolTip", "Set the material domain for the active material designer instance."),
+			FNewToolMenuDelegate::CreateStatic(&AddToolBarMaterialInstanceDomainMenu)
+		);
+
+		NewSection.AddSubMenu(
+			"BlendMode",
+			LOCTEXT("MaterialInstanceBlendMode", "Material Blend Mode"),
+			LOCTEXT("MaterialInstanceBlendMode_ToolTip", "Set the material blend mode for the active material designer instance."),
+			FNewToolMenuDelegate::CreateStatic(&AddToolBarMaterialInstanceBlendModeMenu)
+		);
+
+		FToolUIAction UnlitAction;
+		UnlitAction.ExecuteAction = FToolMenuExecuteAction::CreateStatic(&AddToolBarMaterialInstanceUnlitMenu_Execute);
+		UnlitAction.CanExecuteAction = FToolMenuCanExecuteAction::CreateStatic(&AddToolBarMaterialInstanceUnlitMenu_CanExecute);
+		UnlitAction.GetActionCheckState = FToolMenuGetActionCheckState::CreateStatic(&AddToolBarMaterialInstanceUnlitMenu_GetCheckState);
+
+		NewSection.AddMenuEntry(
+			"Unlit",
+			LOCTEXT("Unlit", "Unlit"),
+			LOCTEXT("UnlitTooltip", "Whether this material requires light to be seen."),
+			TAttribute<FSlateIcon>(),
+			FToolUIActionChoice(UnlitAction),
+			EUserInterfaceActionType::ToggleButton
+		);
+
+		FToolUIAction AnimatedAction;
+		AnimatedAction.ExecuteAction = FToolMenuExecuteAction::CreateStatic(&AddToolBarMaterialInstanceAnimatedMenu_Execute);
+		AnimatedAction.CanExecuteAction = FToolMenuCanExecuteAction::CreateStatic(&AddToolBarMaterialInstanceAnimatedMenu_CanExecute);
+		AnimatedAction.GetActionCheckState = FToolMenuGetActionCheckState::CreateStatic(&AddToolBarMaterialInstanceAnimatedMenu_GetCheckState);
+
+		NewSection.AddMenuEntry(
+			"Animated",
+			LOCTEXT("Animated", "Animated"),
+			LOCTEXT("AnimatedTooltip", "Enables motion vectors in supported materials."),
+			TAttribute<FSlateIcon>(),
+			FToolUIActionChoice(AnimatedAction),
+			EUserInterfaceActionType::ToggleButton
+		);
+
+		FToolUIAction TwoSidedAction;
+		TwoSidedAction.ExecuteAction = FToolMenuExecuteAction::CreateStatic(&AddToolBarMaterialInstanceTwoSidedMenu_Execute);
+		TwoSidedAction.CanExecuteAction = FToolMenuCanExecuteAction::CreateStatic(&AddToolBarMaterialInstanceTwoSidedMenu_CanExecute);
+		TwoSidedAction.GetActionCheckState = FToolMenuGetActionCheckState::CreateStatic(&AddToolBarMaterialInstanceTwoSidedMenu_GetCheckState);
+
+		NewSection.AddMenuEntry(
+			"TwoSided",
+			LOCTEXT("TwoSided", "Two Sided"),
+			LOCTEXT("TwoSidedTooltip", "Whether this material will render on both the back and front of geometry."),
+			TAttribute<FSlateIcon>(),
+			FToolUIActionChoice(TwoSidedAction),
+			EUserInterfaceActionType::ToggleButton
+		);
+	}
 
 	void OpenMaterialEditorFromContext(UDMMenuContext* InMenuContext)
 	{
@@ -291,14 +644,14 @@ namespace UE::DynamicMaterialEditor::Private
 		);
 	}
 
-	void AddToolbarExportMenu(UToolMenu* InMenu)
+	void AddToolBarExportMenu(UToolMenu* InMenu)
 	{
-		if (!IsValid(InMenu) || InMenu->ContainsSection(ToolBarExportSectionName))
+		if (!IsValid(InMenu) || InMenu->ContainsSection(ToolBarMaterialExportSectionName))
 		{
 			return;
 		}
 
-		const UDMMenuContext* const MenuContext = InMenu->FindContext<UDMMenuContext>();
+		UDMMenuContext* const MenuContext = InMenu->FindContext<UDMMenuContext>();
 
 		if (!MenuContext)
 		{
@@ -327,7 +680,17 @@ namespace UE::DynamicMaterialEditor::Private
 			return;
 		}
 
-		FToolMenuSection& NewSection = InMenu->AddSection("Export", LOCTEXT("ExportSection", "Export"));
+		FToolMenuSection& NewSection = InMenu->AddSection(ToolBarMaterialExportSectionName, LOCTEXT("ExportSection", "Export"));
+
+		NewSection.AddMenuEntry(NAME_None,
+			LOCTEXT("OpenInUEMaterialEditor", "Open in Standard Material Editor"),
+			LOCTEXT("OpenInUEMaterialEditorTooltip", "Opens the currently editing generated Material Designer Instance material in the standard material editor."),
+			FSlateIcon(),
+			FUIAction(FExecuteAction::CreateStatic(
+				&OpenMaterialEditorFromContext,
+				MenuContext
+			))
+		);
 
 		if (bAllowInstanceExport)
 		{
@@ -367,12 +730,12 @@ namespace UE::DynamicMaterialEditor::Private
 
 	void AddToolBarTooltipOptionsSection(UToolMenu* InMenu)
 	{
-		if (!IsValid(InMenu) || InMenu->ContainsSection(ToolBarTooltipOptionsSectionName))
+		if (!IsValid(InMenu))
 		{
 			return;
 		}
 
-		FToolMenuSection& NewSection = InMenu->AddSection(ToolBarTooltipOptionsSectionName, LOCTEXT("TooltipOptionsSection", "Tooltip Options"));
+		FToolMenuSection& NewSection = InMenu->AddSection("TooltipOptions", LOCTEXT("TooltipOptionsSection", "Tooltip Options"));
 
 		AddToolBarBoolOptionMenuEntry(NewSection,
 			GET_MEMBER_NAME_CHECKED(UDynamicMaterialEditorSettings, bShowTooltipPreview),
@@ -408,24 +771,20 @@ namespace UE::DynamicMaterialEditor::Private
 
 	void AddToolBarPreviewOptionsSection(UToolMenu* InMenu)
 	{
-		if (!IsValid(InMenu) || InMenu->ContainsSection(ToolBarPreviewOptionsSectionName))
+		if (!IsValid(InMenu))
 		{
 			return;
 		}
 
-		FToolMenuSection& NewSection = InMenu->AddSection(ToolBarPreviewOptionsSectionName, LOCTEXT("PreviewOptionsSection", "Preview Options"));
+		FToolMenuSection& NewSection = InMenu->AddSection("PreviewOptions", LOCTEXT("PreviewOptionsSection", "Preview Options"));
 
 		AddToolBarIntOptionMenuEntry(NewSection,
 			GET_MEMBER_NAME_CHECKED(UDynamicMaterialEditorSettings, LayerPreviewSize)
 		);
 
-		NewSection.AddSeparator(NAME_None);
-
 		AddToolBarIntOptionMenuEntry(NewSection,
 			GET_MEMBER_NAME_CHECKED(UDynamicMaterialEditorSettings, DetailsPreviewSize)
 		);
-
-		NewSection.AddSeparator(NAME_None);
 	}
 
 	void AddToolBarAdvancedSection(UToolMenu* InMenu)
@@ -436,8 +795,6 @@ namespace UE::DynamicMaterialEditor::Private
 
 		InMenu->AddDynamicSection(NAME_None, FNewToolMenuDelegate::CreateStatic(&AddToolBarTooltipOptionsSection));
 
-		NewSection.AddSeparator(NAME_None);
-
 		NewSection.AddMenuEntry(NAME_None,
 			LOCTEXT("ResetAllSettingsToDefaults", "Reset All To Defaults"),
 			LOCTEXT("ResetAllSettingsToDefaultsTooltip", "Resets all the Material Designer settings to their default values."),
@@ -449,23 +806,14 @@ namespace UE::DynamicMaterialEditor::Private
 		);
 	}
 
-	void AddToolBarEditorLayoutMenu(UToolMenu* InMenu)
+	void AddToolBarSettingsMenu(UToolMenu* InMenu)
 	{
-		AddToolbarExportMenu(InMenu);
-
-		FToolMenuSection& NewSection = InMenu->AddSection("MaterialDesigner", LOCTEXT("MaterialDesignerSection", "Material Designer"));
-
-		UDMMenuContext* MenuContext = InMenu->FindContext<UDMMenuContext>();
+		if (!IsValid(InMenu) || InMenu->ContainsSection(ToolBarMaterialDesignerSettingsSectionName))
+		{
+			return;
+		}
 
-		NewSection.AddMenuEntry(NAME_None,
-			LOCTEXT("OpenInUEMaterialEditor", "Open in UE Material Editor..."),
-			LOCTEXT("OpenInUEMaterialEditorTooltip", "Opens the currently editing generated Material Designer Instance material in the Unreal Engine material editor."),
-			FSlateIcon(),
-			FUIAction(FExecuteAction::CreateStatic(
-				&OpenMaterialEditorFromContext,
-				MenuContext
-			))
-		);
+		FToolMenuSection& NewSection = InMenu->AddSection(ToolBarMaterialDesignerSettingsSectionName, LOCTEXT("MaterialDesignerSection", "Material Designer"));
 
 		NewSection.AddSubMenu(
 			"AdvancedSettings",
@@ -475,8 +823,8 @@ namespace UE::DynamicMaterialEditor::Private
 		);
 
 		NewSection.AddMenuEntry(NAME_None,
-			LOCTEXT("OpenSettings", "Material Designer Settings..."),
-			LOCTEXT("OpenSettingsTooltip", "Opens the Material Designer settings window."),
+			LOCTEXT("OpenSettings", "Material Designer Editor Settings"),
+			LOCTEXT("OpenSettingsTooltip", "Opens the Editor Settings and navigates to Material Designer section."),
 			FSlateIcon(FAppStyle::GetAppStyleSetName(), "FoliageEditMode.Settings"),
 			FUIAction(FExecuteAction::CreateUObject(
 				UDynamicMaterialEditorSettings::Get(),
@@ -484,6 +832,13 @@ namespace UE::DynamicMaterialEditor::Private
 			)
 		);
 	}
+
+	void AddToolBarEditorLayoutMenu(UToolMenu* InMenu)
+	{
+		AddToolBarInstanceMenu(InMenu);
+		AddToolBarExportMenu(InMenu);
+		AddToolBarSettingsMenu(InMenu);
+	}
 }
 
 TSharedRef<SWidget> FDMToolBarMenus::MakeEditorLayoutMenu(const TSharedPtr<SDMEditor>& InEditorWidget)
@@ -500,7 +855,7 @@ TSharedRef<SWidget> FDMToolBarMenus::MakeEditorLayoutMenu(const TSharedPtr<SDMEd
 		}
 
 		FToolMenuSection& NewSection = NewToolMenu->AddDynamicSection(
-			"MaterialDesigner", 
+			"MaterialDesignerSettings", 
 			FNewToolMenuDelegate::CreateStatic(&AddToolBarEditorLayoutMenu)
 		);
 	}
diff --git a/Engine/Plugins/Experimental/DynamicMaterial/Source/DynamicMaterialEditor/Private/Model/DynamicMaterialModelEditorOnlyData.cpp b/Engine/Plugins/Experimental/DynamicMaterial/Source/DynamicMaterialEditor/Private/Model/DynamicMaterialModelEditorOnlyData.cpp
--- a/Engine/Plugins/Experimental/DynamicMaterial/Source/DynamicMaterialEditor/Private/Model/DynamicMaterialModelEditorOnlyData.cpp
+++ b/Engine/Plugins/Experimental/DynamicMaterial/Source/DynamicMaterialEditor/Private/Model/DynamicMaterialModelEditorOnlyData.cpp
@@ -64,6 +64,22 @@ const FString UDynamicMaterialModelEditorOnlyData::RGBSlotPathToken     = FStrin
 const FString UDynamicMaterialModelEditorOnlyData::OpacitySlotPathToken = FString(TEXT("OpacitySlot"));
 const FString UDynamicMaterialModelEditorOnlyData::PropertiesPathToken  = FString(TEXT("Properties"));
 
+const TArray<EMaterialDomain> UDynamicMaterialModelEditorOnlyData::SupportedDomains =
+{
+	EMaterialDomain::MD_Surface,
+	EMaterialDomain::MD_PostProcess,
+	EMaterialDomain::MD_DeferredDecal
+};
+
+const TArray<EBlendMode> UDynamicMaterialModelEditorOnlyData::SupportedBlendModes =
+{
+	EBlendMode::BLEND_Opaque,
+	EBlendMode::BLEND_Masked,
+	EBlendMode::BLEND_Translucent,
+	EBlendMode::BLEND_Additive,
+	EBlendMode::BLEND_Modulate
+};
+
 UDynamicMaterialModelEditorOnlyData* UDynamicMaterialModelEditorOnlyData::Get(UDynamicMaterialModel* InModel)
 {
 	if (InModel)
@@ -96,6 +112,8 @@ UDynamicMaterialModelEditorOnlyData::UDynamicMaterialModelEditorOnlyData()
 	BlendMode = EBlendMode::BLEND_Translucent;
 	ShadingModel = EDMMaterialShadingModel::Unlit;
 	bPixelAnimationFlag = false;
+	bTwoSidedFlag = false;
+	bUseWizard = true;
 
 	Properties.Emplace(EDMMaterialPropertyType::BaseColor,           CreateDefaultSubobject<UDMMaterialPropertyBaseColor>(          "MaterialProperty_BaseColor"));
 	Properties.Emplace(EDMMaterialPropertyType::EmissiveColor,       CreateDefaultSubobject<UDMMaterialPropertyEmissiveColor>(      "MaterialProperty_EmissiveColor"));
@@ -217,6 +235,7 @@ void UDynamicMaterialModelEditorOnlyData::BuildMaterial(bool bInDirtyAssets)
 	MaterialModel->DynamicMaterial->MaterialDomain = Domain;
 	MaterialModel->DynamicMaterial->BlendMode = BlendMode;
 	MaterialModel->DynamicMaterial->bHasPixelAnimation = bPixelAnimationFlag;
+	MaterialModel->DynamicMaterial->TwoSided = bTwoSidedFlag;
 
 	switch (ShadingModel)
 	{
@@ -605,6 +624,16 @@ TSharedRef<FDMMaterialBuildState> UDynamicMaterialModelEditorOnlyData::CreateBui
 	return MakeShared<FDMMaterialBuildState>(InMaterialToBuild, MaterialModel, bInDirtyAssets);
 }
 
+bool UDynamicMaterialModelEditorOnlyData::NeedsWizard() const
+{
+	return bUseWizard;
+}
+
+void UDynamicMaterialModelEditorOnlyData::OnWizardComplete()
+{
+	bUseWizard = false;
+}
+
 UDMMaterialComponent* UDynamicMaterialModelEditorOnlyData::GetSubComponentByPath(FDMComponentPath& InPath,
 	const FDMComponentPathSegment& InPathSegment) const
 {
@@ -804,6 +833,66 @@ void UDynamicMaterialModelEditorOnlyData::SetShadingModel(EDMMaterialShadingMode
 	RequestMaterialBuild();
 }
 
+bool UDynamicMaterialModelEditorOnlyData::IsPixelAnimationFlagSet() const
+{
+	return bPixelAnimationFlag;
+}
+
+void UDynamicMaterialModelEditorOnlyData::SetPixelAnimationFlag(bool bInFlagValue)
+{
+	if (bPixelAnimationFlag == bInFlagValue)
+	{
+		return;
+	}
+
+	bPixelAnimationFlag = bInFlagValue;
+
+	if (MaterialModel)
+	{
+		if (UMaterial* Material = MaterialModel->GetGeneratedMaterial())
+		{
+			if (GUndo)
+			{
+				Material->Modify();
+			}
+
+			Material->bHasPixelAnimation = bPixelAnimationFlag;
+		}
+	}
+
+	RequestMaterialBuild();
+}
+
+bool UDynamicMaterialModelEditorOnlyData::IsTwoSidedFlagSet() const
+{
+	return bTwoSidedFlag;
+}
+
+void UDynamicMaterialModelEditorOnlyData::SetTwoSidedFlag(bool bInFlagValue)
+{
+	if (bTwoSidedFlag == bInFlagValue)
+	{
+		return;
+	}
+
+	bTwoSidedFlag = bInFlagValue;
+
+	if (MaterialModel)
+	{
+		if (UMaterial* Material = MaterialModel->GetGeneratedMaterial())
+		{
+			if (GUndo)
+			{
+				Material->Modify();
+			}
+
+			Material->TwoSided = bTwoSidedFlag;
+		}
+	}
+
+	RequestMaterialBuild();
+}
+
 void UDynamicMaterialModelEditorOnlyData::OpenMaterialEditor() const
 {
 	if (!IsValid(MaterialModel) || !IsValid(MaterialModel->DynamicMaterial))
@@ -1088,36 +1177,6 @@ void UDynamicMaterialModelEditorOnlyData::UnassignMaterialProperty(EDMMaterialPr
 	RequestMaterialBuild();
 }
 
-bool UDynamicMaterialModelEditorOnlyData::IsPixelAnimationFlagSet() const
-{
-	return bPixelAnimationFlag;
-}
-
-void UDynamicMaterialModelEditorOnlyData::SetPixelAnimationFlag(bool bInFlagValue)
-{
-	if (bPixelAnimationFlag == bInFlagValue)
-	{
-		return;
-	}
-
-	bPixelAnimationFlag = bInFlagValue;
-
-	if (MaterialModel)
-	{
-		if (UMaterial* Material = MaterialModel->GetGeneratedMaterial())
-		{
-			if (GUndo)
-			{
-				Material->Modify();
-			}
-
-			Material->bHasPixelAnimation = bPixelAnimationFlag;
-		}
-	}
-
-	RequestMaterialBuild();
-}
-
 void UDynamicMaterialModelEditorOnlyData::NotifyPostChange(const FPropertyChangedEvent& PropertyChangedEvent, class FEditPropertyChain* PropertyThatChanged)
 {
 	RequestMaterialBuild();
@@ -1142,6 +1201,8 @@ void UDynamicMaterialModelEditorOnlyData::PostLoad()
 	SetFlags(RF_Transactional);
 
 	ReinitComponents();
+
+	bUseWizard = false;
 }
 
 void UDynamicMaterialModelEditorOnlyData::PostEditUndo()
diff --git a/Engine/Plugins/Experimental/DynamicMaterial/Source/DynamicMaterialEditor/Private/Slate/Properties/SDMBlendMode.cpp b/Engine/Plugins/Experimental/DynamicMaterial/Source/DynamicMaterialEditor/Private/Slate/Properties/SDMBlendMode.cpp
deleted file mode 100644
diff --git a/Engine/Plugins/Experimental/DynamicMaterial/Source/DynamicMaterialEditor/Private/Slate/Properties/SDMBlendMode.h b/Engine/Plugins/Experimental/DynamicMaterial/Source/DynamicMaterialEditor/Private/Slate/Properties/SDMBlendMode.h
deleted file mode 100644
diff --git a/Engine/Plugins/Experimental/DynamicMaterial/Source/DynamicMaterialEditor/Private/Slate/Properties/SDMDomain.cpp b/Engine/Plugins/Experimental/DynamicMaterial/Source/DynamicMaterialEditor/Private/Slate/Properties/SDMDomain.cpp
deleted file mode 100644
diff --git a/Engine/Plugins/Experimental/DynamicMaterial/Source/DynamicMaterialEditor/Private/Slate/Properties/SDMDomain.h b/Engine/Plugins/Experimental/DynamicMaterial/Source/DynamicMaterialEditor/Private/Slate/Properties/SDMDomain.h
deleted file mode 100644
diff --git a/Engine/Plugins/Experimental/DynamicMaterial/Source/DynamicMaterialEditor/Private/Slate/SDMEditor.cpp b/Engine/Plugins/Experimental/DynamicMaterial/Source/DynamicMaterialEditor/Private/Slate/SDMEditor.cpp
--- a/Engine/Plugins/Experimental/DynamicMaterial/Source/DynamicMaterialEditor/Private/Slate/SDMEditor.cpp
+++ b/Engine/Plugins/Experimental/DynamicMaterial/Source/DynamicMaterialEditor/Private/Slate/SDMEditor.cpp
@@ -2,9 +2,7 @@
 
 #include "SDMEditor.h"
 #include "AssetThumbnail.h"
-#include "AssetToolsModule.h"
 #include "Components/DMMaterialLayer.h"
-#include "Components/DMMaterialProperty.h"
 #include "Components/DMMaterialSlot.h"
 #include "Components/DMMaterialValue.h"
 #include "Components/MaterialStageExpressions/DMMSETextureSample.h"
@@ -24,30 +22,24 @@
 #include "Framework/Commands/InputChord.h"
 #include "HAL/PlatformApplicationMisc.h"
 #include "IDetailTreeNode.h"
-#include "IDocumentation.h"
 #include "IPropertyRowGenerator.h"
 #include "MaterialDomain.h"
 #include "Menus/DMToolBarMenus.h"
 #include "Misc/CoreDelegates.h"
 #include "Model/DynamicMaterialModel.h"
 #include "Model/DynamicMaterialModelEditorOnlyData.h"
-#include "PropertyCustomizationHelpers.h"
 #include "PropertyHandle.h"
-#include "SEnumCombo.h"
 #include "Slate/Properties/Editors/SDMPropertyEditOpacity.h"
-#include "Slate/Properties/SDMBlendMode.h"
-#include "Slate/Properties/SDMDomain.h"
 #include "Slate/Properties/SDMMaterialParameters.h"
-#include "Slate/SDMSlot.h"
 #include "Slate/SDMComponentEdit.h"
+#include "Slate/SDMMaterialWizard.h"
+#include "Slate/SDMSlot.h"
 #include "Slate/SDMToolBar.h"
 #include "SlateOptMacros.h"
 #include "Styling/StyleColors.h"
 #include "ThumbnailRendering/ThumbnailManager.h"
-#include "Widgets/Colors/SColorBlock.h"
 #include "Widgets/Input/SButton.h"
 #include "Widgets/Layout/SBox.h"
-#include "Widgets/Layout/SConstraintCanvas.h"
 #include "Widgets/Layout/SExpandableArea.h"
 #include "Widgets/Layout/SScrollBox.h"
 #include "Widgets/SBoxPanel.h"
@@ -293,13 +285,23 @@ void SDMEditor::SetMaterialModel(UDynamicMaterialModel* InMaterialModel)
 		return;
 	}
 
-	Container->SetContent(CreateMainLayout());
+	UDynamicMaterialModelEditorOnlyData* EditorOnlyData = UDynamicMaterialModelEditorOnlyData::Get(InMaterialModel);
+
+	if (EditorOnlyData && EditorOnlyData->NeedsWizard())
+	{
+		Container->SetContent(SNew(SDMMaterialWizard, SharedThis(this)));
+		EditorOnlyData->OnWizardComplete();
+	}
+	else
+	{
+		Container->SetContent(CreateMainLayout());
+	}
 
-	if (UDynamicMaterialModelEditorOnlyData* ModelEditorOnlyData = UDynamicMaterialModelEditorOnlyData::Get(InMaterialModel))
+	if (EditorOnlyData)
 	{
-		ModelEditorOnlyData->GetOnMaterialBuiltDelegate().AddSP(this, &SDMEditor::OnMaterialBuilt);
-		ModelEditorOnlyData->GetOnValueListUpdateDelegate().AddSP(this, &SDMEditor::OnValuesUpdated);
-		ModelEditorOnlyData->GetOnSlotListUpdateDelegate().AddSP(this, &SDMEditor::OnSlotsUpdated);
+		EditorOnlyData->GetOnMaterialBuiltDelegate().AddSP(this, &SDMEditor::OnMaterialBuilt);
+		EditorOnlyData->GetOnValueListUpdateDelegate().AddSP(this, &SDMEditor::OnValuesUpdated);
+		EditorOnlyData->GetOnSlotListUpdateDelegate().AddSP(this, &SDMEditor::OnSlotsUpdated);
 	}
 }
 
@@ -555,26 +557,6 @@ TSharedRef<SWidget> SDMEditor::CreateMainLayout()
 		.AutoHeight()
 		.HAlign(HAlign_Fill)
 		.VAlign(VAlign_Fill)
-		[
-			SAssignNew(GlobalOpacityContainer, SBox)
-			.HAlign(HAlign_Fill)
-			.VAlign(VAlign_Fill)
-			[
-				SNew(SBorder)
-				.HAlign(HAlign_Fill)
-				.VAlign(VAlign_Center)
-				.Padding(3.0f, 3.0f, 3.0f, 3.0f)
-				.BorderImage(FDynamicMaterialEditorStyle::GetBrush("Border.Bottom"))
-				.BorderBackgroundColor(FLinearColor(1, 1, 1, 0.05f))
-				[
-					CreateMaterialSettingsRow()
-				]
-			]
-		]
-		+ SVerticalBox::Slot()
-		.AutoHeight()
-		.HAlign(HAlign_Fill)
-		.VAlign(VAlign_Fill)
 		[
 			SAssignNew(GlobalOpacityContainer, SBox)
 			.HAlign(HAlign_Fill)
@@ -614,95 +596,6 @@ TSharedRef<SWidget> SDMEditor::CreateMainLayout()
 		];
 }
 
-TSharedRef<SWidget> SDMEditor::CreateMaterialSettingsRow()
-{
-	return 
-		SNew(SHorizontalBox)
-		+ SHorizontalBox::Slot()
-		.FillWidth(1.0f)
-		.HAlign(HAlign_Fill)
-		.VAlign(VAlign_Center)
-		.Padding(5.0f, 0.0f, 0.0f, 2.0f)
-		[
-			SNew(SDMDomain)
-			.ToolTipText(LOCTEXT("MaterialDesignerInstanceDomainTooltip", "Change the Material Designer Instance Domain"))
-			.SelectedItem(this, &SDMEditor::GetSelectedDomain)
-			.OnSelectedItemChanged(this, &SDMEditor::OnDomainChanged)
-			.IsEnabled(this, &SDMEditor::CanChangeDomain)
-		]
-		+ SHorizontalBox::Slot()
-		.FillWidth(1.0f)
-		.HAlign(HAlign_Fill)
-		.VAlign(VAlign_Center)
-		.Padding(5.0f, 0.0f, 0.0f, 2.0f)
-		[
-			SNew(SDMBlendMode)
-			.ToolTipText(LOCTEXT("MaterialDesignerInstanceBlendModeTooltip", "Change the Material Designer Instance Blend Mode"))
-			.SelectedItem(this, &SDMEditor::GetSelectedBlendMode)
-			.OnSelectedItemChanged(this, &SDMEditor::OnBlendModeChanged)
-			.IsEnabled(this, &SDMEditor::CanChangeBlendType)
-		]
-		+ SHorizontalBox::Slot()
-		.AutoWidth()
-		.HAlign(HAlign_Right)
-		.VAlign(VAlign_Center)
-		.Padding(10.f, 0.0f, 5.0f, 2.0f)
-		[
-			SNew(SHorizontalBox)
-			.ToolTipText(LOCTEXT("MaterialDesignerInstanceUnlitTooltip", "Toggle between Default Lit and Unlit for this Material Designer Instance."))
-			+ SHorizontalBox::Slot()
-			.AutoWidth()
-			.HAlign(HAlign_Left)
-			.VAlign(VAlign_Center)
-			.Padding(0.0f, 0.0f, 5.0f, 0.0f)
-			[
-				SNew(STextBlock)
-				.Text(LOCTEXT("MaterialDesignerInstanceUnlit", "Unlit"))
-			]
-			+ SHorizontalBox::Slot()
-			.AutoWidth()
-			.HAlign(HAlign_Left)
-			.VAlign(VAlign_Center)
-			[
-				SNew(SCheckBox)
-				.IsChecked(this, &SDMEditor::IsMaterialUnlit)
-				.OnCheckStateChanged(this, &SDMEditor::OnMaterialUnlitChanged)
-				.IsEnabled(this, &SDMEditor::CanChangeMaterialShadingModel)
-			]
-		]
-		+ SHorizontalBox::Slot()
-		.AutoWidth()
-		.HAlign(HAlign_Right)
-		.VAlign(VAlign_Center)
-		.Padding(10.0f, 0.0f, 0.0f, 2.0f)
-		[
-			SNew(SHorizontalBox)
-			.ToolTipText(LOCTEXT("MaterialDesignerInstanceTypeTooltip", "Enables the material's TSR pixel animation flag. Not available in translucent blend modes."))
-			+ SHorizontalBox::Slot()
-			.AutoWidth()
-			.HAlign(HAlign_Left)
-			.VAlign(VAlign_Center)
-			[
-				SNew(SBox)
-				.Padding(0.f, 0.f, 5.f, 0.f)
-				[
-					SNew(STextBlock)
-					.Text(LOCTEXT("MaterialDesignerInstanceAnimaed", "Animated"))
-				]
-			]
-			+ SHorizontalBox::Slot()
-			.AutoWidth()
-			.HAlign(HAlign_Left)
-			.VAlign(VAlign_Center)
-			[
-				SNew(SCheckBox)
-				.IsChecked(this, &SDMEditor::IsMaterialAnimated)
-				.IsEnabled(this, &SDMEditor::CanMaterialBeAnimated)
-				.OnCheckStateChanged(this, &SDMEditor::OnMaterialAnimatedChanged)
-			]
-		];
-}
-
 TSharedRef<SWidget> SDMEditor::CreateGlobalOpacityWidget()
 {
 	UDMMaterialValueFloat1* OpacityValue = nullptr;
diff --git a/Engine/Plugins/Experimental/DynamicMaterial/Source/DynamicMaterialEditor/Private/Slate/SDMEditor.h b/Engine/Plugins/Experimental/DynamicMaterial/Source/DynamicMaterialEditor/Private/Slate/SDMEditor.h
--- a/Engine/Plugins/Experimental/DynamicMaterial/Source/DynamicMaterialEditor/Private/Slate/SDMEditor.h
+++ b/Engine/Plugins/Experimental/DynamicMaterial/Source/DynamicMaterialEditor/Private/Slate/SDMEditor.h
@@ -134,7 +134,6 @@ protected:
 	void BindCommands();
 
 	TSharedRef<SWidget> CreateMainLayout();
-	TSharedRef<SWidget> CreateMaterialSettingsRow();
 	TSharedRef<SWidget> CreateGlobalOpacityWidget();
 	TSharedRef<SWidget> CreateParametersArea();
 
diff --git a/Engine/Plugins/Experimental/DynamicMaterial/Source/DynamicMaterialEditor/Private/Slate/SDMMaterialWizard.cpp b/Engine/Plugins/Experimental/DynamicMaterial/Source/DynamicMaterialEditor/Private/Slate/SDMMaterialWizard.cpp
new file mode 100644
--- /dev/null
+++ b/Engine/Plugins/Experimental/DynamicMaterial/Source/DynamicMaterialEditor/Private/Slate/SDMMaterialWizard.cpp
@@ -0,0 +1,622 @@
+// Copyright Epic Games, Inc. All Rights Reserved.
+
+#include "Slate/SDMMaterialWizard.h"
+
+#include "DMDefs.h"
+#include "DynamicMaterialEditorStyle.h"
+#include "Engine/EngineTypes.h"
+#include "MaterialDomain.h"
+#include "Model/DynamicMaterialModel.h"
+#include "Model/DynamicMaterialModelEditorOnlyData.h"
+#include "SDMEditor.h"
+#include "Styling/SlateTypes.h"
+#include "Widgets/Input/SButton.h"
+#include "Widgets/Input/SCheckBox.h"
+#include "Widgets/Layout/SBox.h"
+#include "Widgets/Layout/SWrapBox.h"
+#include "Widgets/SBoxPanel.h"
+#include "Widgets/Text/STextBlock.h"
+
+#define LOCTEXT_NAMESPACE "SDMMaterialWizard"
+
+namespace UE::DynamicMaterialEditor::Private
+{
+	constexpr float SeparationDistance = 20.f;
+	constexpr float TitleContentDistance = 5.f;
+	static const FMargin ButtonPadding = FMargin(10.f, 5.f);
+	static const FVector2D WrapBoxSlotPadding = FVector2D(5, 5);
+}
+
+void SDMMaterialWizard::Construct(const FArguments& InArgs, const TSharedRef<SDMEditor>& InEditor)
+{
+	EditorWeak = InEditor;
+
+	ChildSlot
+		[
+			CreateLayout()
+		];
+}
+
+TSharedPtr<SDMEditor> SDMMaterialWizard::GetEditor() const
+{
+	return EditorWeak.Pin();
+}
+
+UDynamicMaterialModel* SDMMaterialWizard::GetMaterialModel() const
+{
+	if (TSharedPtr<SDMEditor> Editor = EditorWeak.Pin())
+	{
+		return Editor->GetMaterialModel();
+	}
+
+	return nullptr;
+}
+
+TSharedRef<SWidget> SDMMaterialWizard::CreateLayout()
+{
+	using namespace UE::DynamicMaterialEditor::Private;
+
+	return SNew(SBox)
+		.Padding(SeparationDistance)
+		.HAlign(EHorizontalAlignment::HAlign_Fill)
+		[
+			SNew(SVerticalBox)
+			+ SVerticalBox::Slot()
+			.AutoHeight()
+			.HAlign(HAlign_Center)
+			[
+				SNew(STextBlock)
+				.TextStyle(FDynamicMaterialEditorStyle::Get(), "ActorNameBig")
+				.Text(LOCTEXT("MaterialWizard", "Material Wizard"))
+			]
+
+			+ SVerticalBox::Slot()
+			.AutoHeight()
+			.HAlign(HAlign_Left)
+			.Padding(0.0f, SeparationDistance, 0.0f, TitleContentDistance)
+			[
+				SNew(STextBlock)
+				.Text(LOCTEXT("MaterialDomain", "Material Domain"))
+			]
+
+			+ SVerticalBox::Slot()
+			.AutoHeight()
+			.HAlign(EHorizontalAlignment::HAlign_Fill)
+			[
+				CreateMaterialDomainOptions()
+			]
+
+			+ SVerticalBox::Slot()
+			.AutoHeight()
+			.HAlign(HAlign_Left)
+			.Padding(0.0f, SeparationDistance, 0.0f, TitleContentDistance)
+			[
+				SNew(STextBlock)
+				.Text(LOCTEXT("BlendMode", "Blend Mode"))
+			]
+
+			+ SVerticalBox::Slot()
+			.AutoHeight()
+			.HAlign(EHorizontalAlignment::HAlign_Fill)
+			[
+				CreateBlendModeOptions()
+			]
+
+			+ SVerticalBox::Slot()
+			.AutoHeight()
+			.HAlign(HAlign_Left)
+			.Padding(0.0f, SeparationDistance, 0.0f, TitleContentDistance)
+			[
+				SNew(STextBlock)
+				.Text(LOCTEXT("ShadingModel", "Shading Model"))
+			]
+
+			+ SVerticalBox::Slot()
+			.AutoHeight()
+			.HAlign(EHorizontalAlignment::HAlign_Fill)
+			[
+				CreateShadingModelOptions()
+			]
+
+			+ SVerticalBox::Slot()
+			.AutoHeight()
+			.HAlign(HAlign_Left)
+			.Padding(0.0f, SeparationDistance, 0.0f, TitleContentDistance)
+			[
+				SNew(STextBlock)
+				.Text(LOCTEXT("MotionVectors", "Motion Vectors"))
+			]
+
+			+ SVerticalBox::Slot()
+			.AutoHeight()
+			.HAlign(EHorizontalAlignment::HAlign_Fill)
+			[
+				CreateAnimationOptions()
+			]
+
+			+ SVerticalBox::Slot()
+			.AutoHeight()
+			.HAlign(HAlign_Left)
+			.Padding(0.0f, SeparationDistance, 0.0f, TitleContentDistance)
+			[
+				SNew(STextBlock)
+				.Text(LOCTEXT("Geometry", "Geometry"))
+			]
+
+			+ SVerticalBox::Slot()
+			.AutoHeight()
+			.HAlign(EHorizontalAlignment::HAlign_Fill)
+			[
+				CreateTwoSidedOptions()
+			]
+
+			+ SVerticalBox::Slot()
+			.AutoHeight()
+			.HAlign(HAlign_Left)
+			.Padding(0.0f, SeparationDistance, 0.0f, TitleContentDistance)
+			[
+				SNew(STextBlock)
+				.Text(LOCTEXT("Channels", "Channels"))
+			]
+
+			+ SVerticalBox::Slot()
+			.AutoHeight()
+			.HAlign(EHorizontalAlignment::HAlign_Fill)
+			.Padding(0.0f, 0.f, 0.0f, TitleContentDistance)
+			[
+				CreateChannelPresets()
+			]
+
+			+ SVerticalBox::Slot()
+			.AutoHeight()
+			.HAlign(EHorizontalAlignment::HAlign_Fill)
+			[
+				CreateChannelList()
+			]
+
+			+ SVerticalBox::Slot()
+			.AutoHeight()
+			.HAlign(EHorizontalAlignment::HAlign_Fill)
+			.Padding(0.0f, SeparationDistance, 0.0f, 0.f)
+			[
+				CreateAcceptButton()
+			]
+		];
+}
+
+TSharedRef<SWidget> SDMMaterialWizard::CreateMaterialDomainOptions()
+{
+	using namespace UE::DynamicMaterialEditor::Private;
+
+	UEnum* DomainEnum = StaticEnum<EMaterialDomain>();
+
+	TSharedRef<SWrapBox> DomainOptions = SNew(SWrapBox)
+		.UseAllottedSize(true)
+		.InnerSlotPadding(WrapBoxSlotPadding)
+		.Orientation(EOrientation::Orient_Horizontal);
+
+	for (EMaterialDomain Domain : UDynamicMaterialModelEditorOnlyData::SupportedDomains)
+	{
+		DomainOptions->AddSlot()
+			[
+				SNew(SCheckBox)
+				.Style(FAppStyle::Get(), "DetailsView.SectionButton")
+				.HAlign(EHorizontalAlignment::HAlign_Center)
+				.Padding(ButtonPadding)
+				.IsEnabled(this, &SDMMaterialWizard::Domain_IsEnabled)
+				.IsChecked(this, &SDMMaterialWizard::Domain_GetState, Domain)
+				.OnCheckStateChanged(this, &SDMMaterialWizard::Domain_OnChange, Domain)
+				[
+					SNew(STextBlock)
+					.Text(DomainEnum->GetDisplayNameTextByValue(Domain))
+				]
+			];
+	}
+
+	return DomainOptions;
+}
+
+TSharedRef<SWidget> SDMMaterialWizard::CreateBlendModeOptions()
+{
+	using namespace UE::DynamicMaterialEditor::Private;
+
+	UEnum* BlendModeEnum = StaticEnum<EBlendMode>();
+
+	TSharedRef<SWrapBox> BlendModeOptions = SNew(SWrapBox)
+		.UseAllottedSize(true)
+		.InnerSlotPadding(WrapBoxSlotPadding)
+		.Orientation(EOrientation::Orient_Horizontal);
+
+	for (EBlendMode BlendMode : UDynamicMaterialModelEditorOnlyData::SupportedBlendModes)
+	{
+		BlendModeOptions->AddSlot()
+			[
+				SNew(SCheckBox)
+				.Style(FAppStyle::Get(), "DetailsView.SectionButton")
+				.HAlign(EHorizontalAlignment::HAlign_Center)
+				.Padding(ButtonPadding)
+				.IsEnabled(this, &SDMMaterialWizard::BlendMode_IsEnabled)
+				.IsChecked(this, &SDMMaterialWizard::BlendMode_GetState, BlendMode)
+				.OnCheckStateChanged(this, &SDMMaterialWizard::BlendMode_OnChange, BlendMode)
+				[
+					SNew(STextBlock)
+					.Text(BlendModeEnum->GetDisplayNameTextByValue(BlendMode))
+				]
+			];
+	}
+
+	return BlendModeOptions;
+}
+
+TSharedRef<SWidget> SDMMaterialWizard::CreateShadingModelOptions()
+{
+	using namespace UE::DynamicMaterialEditor::Private;
+
+	UEnum* MaterialShadingModelEnum = StaticEnum<EMaterialShadingModel>();
+
+	TSharedRef<SWrapBox> MaterialShadingModelOptions = SNew(SWrapBox)
+		.UseAllottedSize(true)
+		.InnerSlotPadding(WrapBoxSlotPadding)
+		.Orientation(EOrientation::Orient_Horizontal);
+
+	for (EDMMaterialShadingModel MaterialShadingModel : {EDMMaterialShadingModel::Unlit, EDMMaterialShadingModel::DefaultLit})
+	{
+		MaterialShadingModelOptions->AddSlot()
+			[
+				SNew(SCheckBox)
+				.Style(FAppStyle::Get(), "DetailsView.SectionButton")
+				.HAlign(EHorizontalAlignment::HAlign_Center)
+				.Padding(ButtonPadding)
+				.IsEnabled(this, &SDMMaterialWizard::Unlit_IsEnabled)
+				.IsChecked(this, &SDMMaterialWizard::Unlit_GetState, MaterialShadingModel)
+				.OnCheckStateChanged(this, &SDMMaterialWizard::Unlit_OnChange, MaterialShadingModel)
+				[
+					SNew(STextBlock)
+					.Text(MaterialShadingModelEnum->GetDisplayNameTextByValue(static_cast<int64>(MaterialShadingModel)))
+				]
+			];
+	}
+
+	return MaterialShadingModelOptions;
+}
+
+TSharedRef<SWidget> SDMMaterialWizard::CreateAnimationOptions()
+{
+	using namespace UE::DynamicMaterialEditor::Private;
+
+	TSharedRef<SWrapBox> AnimatedOptions = SNew(SWrapBox)
+		.UseAllottedSize(true)
+		.InnerSlotPadding(WrapBoxSlotPadding)
+		.Orientation(EOrientation::Orient_Horizontal);
+
+	AnimatedOptions->AddSlot()
+		[
+			SNew(SCheckBox)
+				.Style(FAppStyle::Get(), "DetailsView.SectionButton")
+				.HAlign(EHorizontalAlignment::HAlign_Center)
+				.Padding(ButtonPadding)
+				.IsEnabled(this, &SDMMaterialWizard::Animated_IsEnabled)
+				.IsChecked(this, &SDMMaterialWizard::Animated_GetState, true)
+				.OnCheckStateChanged(this, &SDMMaterialWizard::Animated_OnChange, true)
+				[
+					SNew(STextBlock)
+					.Text(LOCTEXT("Animated", "Animated"))
+				]
+		];
+
+	AnimatedOptions->AddSlot()
+		[
+			SNew(SCheckBox)
+				.Style(FAppStyle::Get(), "DetailsView.SectionButton")
+				.HAlign(EHorizontalAlignment::HAlign_Center)
+				.Padding(ButtonPadding)
+				.IsEnabled(this, &SDMMaterialWizard::Animated_IsEnabled)
+				.IsChecked(this, &SDMMaterialWizard::Animated_GetState, false)
+				.OnCheckStateChanged(this, &SDMMaterialWizard::Animated_OnChange, false)
+				[
+					SNew(STextBlock)
+					.Text(LOCTEXT("Static", "Static"))
+				]
+		];
+
+	return AnimatedOptions;
+}
+
+TSharedRef<SWidget> SDMMaterialWizard::CreateTwoSidedOptions()
+{
+	using namespace UE::DynamicMaterialEditor::Private;
+
+	TSharedRef<SWrapBox> TwoSidedOptions = SNew(SWrapBox)
+		.UseAllottedSize(true)
+		.InnerSlotPadding(WrapBoxSlotPadding)
+		.Orientation(EOrientation::Orient_Horizontal);
+
+	TwoSidedOptions->AddSlot()
+		[
+			SNew(SCheckBox)
+				.Style(FAppStyle::Get(), "DetailsView.SectionButton")
+				.HAlign(EHorizontalAlignment::HAlign_Center)
+				.Padding(ButtonPadding)
+				.IsEnabled(this, &SDMMaterialWizard::TwoSided_IsEnabled)
+				.IsChecked(this, &SDMMaterialWizard::TwoSided_GetState, true)
+				.OnCheckStateChanged(this, &SDMMaterialWizard::TwoSided_OnChange, true)
+				[
+					SNew(STextBlock)
+					.Text(LOCTEXT("TwoSided", "Two Sided"))
+				]
+		];
+
+	TwoSidedOptions->AddSlot()
+		[
+			SNew(SCheckBox)
+				.Style(FAppStyle::Get(), "DetailsView.SectionButton")
+				.HAlign(EHorizontalAlignment::HAlign_Center)
+				.Padding(ButtonPadding)
+				.IsEnabled(this, &SDMMaterialWizard::TwoSided_IsEnabled)
+				.IsChecked(this, &SDMMaterialWizard::TwoSided_GetState, false)
+				.OnCheckStateChanged(this, &SDMMaterialWizard::TwoSided_OnChange, false)
+				[
+					SNew(STextBlock)
+					.Text(LOCTEXT("OneSided", "One Sided"))
+				]
+		];
+
+	return TwoSidedOptions;
+}
+
+TSharedRef<SWidget> SDMMaterialWizard::CreateChannelPresets()
+{
+	using namespace UE::DynamicMaterialEditor::Private;
+
+	return SNew(SHorizontalBox)
+		+ SHorizontalBox::Slot()
+		.AutoWidth()
+		.Padding(0.f, 0.f, TitleContentDistance, 0.f)
+		[
+			SNew(STextBlock)
+			.Text(LOCTEXT("Presets", "Presets"))
+		]
+		+ SHorizontalBox::Slot()
+		[
+			SNew(STextBlock)
+			.Text(INVTEXT("-"))
+		];
+}
+
+TSharedRef<SWidget> SDMMaterialWizard::CreateChannelList()
+{
+	return SNew(STextBlock)
+		.Text(INVTEXT("Full Channel List"));
+}
+
+TSharedRef<SWidget> SDMMaterialWizard::CreateAcceptButton()
+{
+	using namespace UE::DynamicMaterialEditor::Private;
+
+	return SNew(SBox)
+		.HAlign(HAlign_Center)
+		[
+			SNew(SButton)
+			.ButtonStyle(FAppStyle::Get(), "PrimaryButton")
+			.ContentPadding(ButtonPadding)
+			.OnClicked(this, &SDMMaterialWizard::Accept_OnClick)
+			[
+				SNew(STextBlock)
+				.Text(LOCTEXT("Continue", "Continue"))
+			]
+		];
+}
+
+bool SDMMaterialWizard::Domain_IsEnabled() const
+{
+	if (UDynamicMaterialModel* MaterialModel = GetMaterialModel())
+	{
+		return !!UDynamicMaterialModelEditorOnlyData::Get(MaterialModel);
+	}
+
+	return false;
+}
+
+ECheckBoxState SDMMaterialWizard::Domain_GetState(EMaterialDomain InDomain) const
+{
+	if (UDynamicMaterialModel* MaterialModel = GetMaterialModel())
+	{
+		if (UDynamicMaterialModelEditorOnlyData* EditorOnlyData = UDynamicMaterialModelEditorOnlyData::Get(MaterialModel))
+		{
+			return EditorOnlyData->GetDomain() == InDomain
+				? ECheckBoxState::Checked
+				: ECheckBoxState::Unchecked;
+		}
+	}
+
+	return ECheckBoxState::Undetermined;
+}
+
+void SDMMaterialWizard::Domain_OnChange(ECheckBoxState InState, EMaterialDomain InDomain)
+{
+	if (InState == ECheckBoxState::Checked)
+	{
+		if (UDynamicMaterialModel* MaterialModel = GetMaterialModel())
+		{
+			if (UDynamicMaterialModelEditorOnlyData* EditorOnlyData = UDynamicMaterialModelEditorOnlyData::Get(MaterialModel))
+			{
+				EditorOnlyData->SetDomain(InDomain);
+			}
+		}
+	}
+}
+
+bool SDMMaterialWizard::BlendMode_IsEnabled() const
+{
+	if (UDynamicMaterialModel* MaterialModel = GetMaterialModel())
+	{
+		return !!UDynamicMaterialModelEditorOnlyData::Get(MaterialModel);
+	}
+
+	return false;
+}
+
+ECheckBoxState SDMMaterialWizard::BlendMode_GetState(EBlendMode InBlendMode) const
+{
+	if (UDynamicMaterialModel* MaterialModel = GetMaterialModel())
+	{
+		if (UDynamicMaterialModelEditorOnlyData* EditorOnlyData = UDynamicMaterialModelEditorOnlyData::Get(MaterialModel))
+		{
+			return EditorOnlyData->GetBlendMode() == InBlendMode
+				? ECheckBoxState::Checked
+				: ECheckBoxState::Unchecked;
+		}
+	}
+
+	return ECheckBoxState::Undetermined;
+}
+
+void SDMMaterialWizard::BlendMode_OnChange(ECheckBoxState InState, EBlendMode InBlendMode)
+{
+	if (InState == ECheckBoxState::Checked)
+	{
+		if (UDynamicMaterialModel* MaterialModel = GetMaterialModel())
+		{
+			if (UDynamicMaterialModelEditorOnlyData* EditorOnlyData = UDynamicMaterialModelEditorOnlyData::Get(MaterialModel))
+			{
+				EditorOnlyData->SetBlendMode(InBlendMode);
+			}
+		}
+	}
+}
+
+bool SDMMaterialWizard::Unlit_IsEnabled() const
+{
+	if (UDynamicMaterialModel* MaterialModel = GetMaterialModel())
+	{
+		return !!UDynamicMaterialModelEditorOnlyData::Get(MaterialModel);
+	}
+
+	return false;
+}
+
+ECheckBoxState SDMMaterialWizard::Unlit_GetState(EDMMaterialShadingModel InValue) const
+{
+	if (UDynamicMaterialModel* MaterialModel = GetMaterialModel())
+	{
+		if (UDynamicMaterialModelEditorOnlyData* EditorOnlyData = UDynamicMaterialModelEditorOnlyData::Get(MaterialModel))
+		{
+			return EditorOnlyData->GetShadingModel() == InValue
+				? ECheckBoxState::Checked
+				: ECheckBoxState::Unchecked;
+		}
+	}
+
+	return ECheckBoxState::Undetermined;
+}
+
+void SDMMaterialWizard::Unlit_OnChange(ECheckBoxState InState, EDMMaterialShadingModel InValue)
+{
+	if (InState == ECheckBoxState::Checked)
+	{
+		if (UDynamicMaterialModel* MaterialModel = GetMaterialModel())
+		{
+			if (UDynamicMaterialModelEditorOnlyData* EditorOnlyData = UDynamicMaterialModelEditorOnlyData::Get(MaterialModel))
+			{
+				EditorOnlyData->SetShadingModel(InValue);
+			}
+		}
+	}
+}
+
+bool SDMMaterialWizard::Animated_IsEnabled() const
+{
+	if (UDynamicMaterialModel* MaterialModel = GetMaterialModel())
+	{
+		return !!UDynamicMaterialModelEditorOnlyData::Get(MaterialModel);
+	}
+
+	return false;
+}
+
+ECheckBoxState SDMMaterialWizard::Animated_GetState(bool bInValue) const
+{
+	if (UDynamicMaterialModel* MaterialModel = GetMaterialModel())
+	{
+		if (UDynamicMaterialModelEditorOnlyData* EditorOnlyData = UDynamicMaterialModelEditorOnlyData::Get(MaterialModel))
+		{
+			return EditorOnlyData->IsPixelAnimationFlagSet() == bInValue
+				? ECheckBoxState::Checked
+				: ECheckBoxState::Unchecked;
+		}
+	}
+
+	return ECheckBoxState::Undetermined;
+}
+
+void SDMMaterialWizard::Animated_OnChange(ECheckBoxState InState, bool bInValue)
+{
+	if (InState == ECheckBoxState::Checked)
+	{
+		if (UDynamicMaterialModel* MaterialModel = GetMaterialModel())
+		{
+			if (UDynamicMaterialModelEditorOnlyData* EditorOnlyData = UDynamicMaterialModelEditorOnlyData::Get(MaterialModel))
+			{
+				EditorOnlyData->SetPixelAnimationFlag(bInValue);
+			}
+		}
+	}
+}
+
+bool SDMMaterialWizard::TwoSided_IsEnabled() const
+{
+	if (UDynamicMaterialModel* MaterialModel = GetMaterialModel())
+	{
+		return !!UDynamicMaterialModelEditorOnlyData::Get(MaterialModel);
+	}
+
+	return false;
+}
+
+ECheckBoxState SDMMaterialWizard::TwoSided_GetState(bool bInValue) const
+{
+	if (UDynamicMaterialModel* MaterialModel = GetMaterialModel())
+	{
+		if (UDynamicMaterialModelEditorOnlyData* EditorOnlyData = UDynamicMaterialModelEditorOnlyData::Get(MaterialModel))
+		{
+			return EditorOnlyData->IsTwoSidedFlagSet() == bInValue
+				? ECheckBoxState::Checked
+				: ECheckBoxState::Unchecked;
+		}
+	}
+
+	return ECheckBoxState::Undetermined;
+}
+
+void SDMMaterialWizard::TwoSided_OnChange(ECheckBoxState InState, bool bInValue)
+{
+	if (InState == ECheckBoxState::Checked)
+	{
+		if (UDynamicMaterialModel* MaterialModel = GetMaterialModel())
+		{
+			if (UDynamicMaterialModelEditorOnlyData* EditorOnlyData = UDynamicMaterialModelEditorOnlyData::Get(MaterialModel))
+			{
+				EditorOnlyData->SetTwoSidedFlag(bInValue);
+			}
+		}
+	}
+}
+
+FReply SDMMaterialWizard::Accept_OnClick()
+{
+	if (TSharedPtr<SDMEditor> Editor = GetEditor())
+	{
+		if (UDynamicMaterialModel* MaterialModel = GetMaterialModel())
+		{
+			if (UDynamicMaterialModelEditorOnlyData* EditorOnlyData = UDynamicMaterialModelEditorOnlyData::Get(MaterialModel))
+			{
+				EditorOnlyData->OnWizardComplete();
+				Editor->SetMaterialModel(MaterialModel); // Refresh display
+			}
+		}
+	}
+
+	return FReply::Handled();
+}
+
+#undef LOCTEXT_NAMESPACE
diff --git a/Engine/Plugins/Experimental/DynamicMaterial/Source/DynamicMaterialEditor/Private/Slate/SDMMaterialWizard.h b/Engine/Plugins/Experimental/DynamicMaterial/Source/DynamicMaterialEditor/Private/Slate/SDMMaterialWizard.h
new file mode 100644
--- /dev/null
+++ b/Engine/Plugins/Experimental/DynamicMaterial/Source/DynamicMaterialEditor/Private/Slate/SDMMaterialWizard.h
@@ -0,0 +1,62 @@
+// Copyright Epic Games, Inc. All Rights Reserved.
+
+#pragma once
+
+#include "Templates/SharedPointer.h"
+#include "Widgets/SCompoundWidget.h"
+
+class SDMEditor;
+class SWidget;
+class UDynamicMaterialModel;
+enum class ECheckBoxState : uint8;
+enum class EDMMaterialShadingModel : uint8;
+enum EBlendMode : int;
+enum EMaterialDomain : int;
+
+class SDMMaterialWizard : public SCompoundWidget
+{
+public:
+	SLATE_BEGIN_ARGS(SDMMaterialWizard) {}
+	SLATE_END_ARGS()
+
+	void Construct(const FArguments& InArgs, const TSharedRef<SDMEditor>& InEditor);
+
+	TSharedPtr<SDMEditor> GetEditor() const;
+
+	UDynamicMaterialModel* GetMaterialModel() const;
+
+protected:
+	TWeakPtr<SDMEditor> EditorWeak;
+
+	TSharedRef<SWidget> CreateLayout();
+	TSharedRef<SWidget> CreateMaterialDomainOptions();
+	TSharedRef<SWidget> CreateBlendModeOptions();
+	TSharedRef<SWidget> CreateShadingModelOptions();
+	TSharedRef<SWidget> CreateAnimationOptions();
+	TSharedRef<SWidget> CreateTwoSidedOptions();
+	TSharedRef<SWidget> CreateChannelPresets();
+	TSharedRef<SWidget> CreateChannelList();
+	TSharedRef<SWidget> CreateAcceptButton();
+
+	bool Domain_IsEnabled() const;
+	ECheckBoxState Domain_GetState(EMaterialDomain InDomain) const;
+	void Domain_OnChange(ECheckBoxState InState, EMaterialDomain InDomain);
+
+	bool BlendMode_IsEnabled() const;
+	ECheckBoxState BlendMode_GetState(EBlendMode InBlendMode) const;
+	void BlendMode_OnChange(ECheckBoxState InState, EBlendMode InBlendMode);
+
+	bool Unlit_IsEnabled() const;
+	ECheckBoxState Unlit_GetState(EDMMaterialShadingModel InValue) const;
+	void Unlit_OnChange(ECheckBoxState InState, EDMMaterialShadingModel InValue);
+
+	bool Animated_IsEnabled() const;
+	ECheckBoxState Animated_GetState(bool bInValue) const;
+	void Animated_OnChange(ECheckBoxState InState, bool bInValue);
+
+	bool TwoSided_IsEnabled() const;
+	ECheckBoxState TwoSided_GetState(bool bInValue) const;
+	void TwoSided_OnChange(ECheckBoxState InState, bool bInValue);
+
+	FReply Accept_OnClick();
+};
diff --git a/Engine/Plugins/Experimental/DynamicMaterial/Source/DynamicMaterialEditor/Public/Model/DynamicMaterialModelEditorOnlyData.h b/Engine/Plugins/Experimental/DynamicMaterial/Source/DynamicMaterialEditor/Public/Model/DynamicMaterialModelEditorOnlyData.h
--- a/Engine/Plugins/Experimental/DynamicMaterial/Source/DynamicMaterialEditor/Public/Model/DynamicMaterialModelEditorOnlyData.h
+++ b/Engine/Plugins/Experimental/DynamicMaterial/Source/DynamicMaterialEditor/Public/Model/DynamicMaterialModelEditorOnlyData.h
@@ -51,6 +51,9 @@ public:
 	static const FString OpacitySlotPathToken; // @TODO This will probably need to change when other slot types are opened up.
 	static const FString PropertiesPathToken;
 
+	static const TArray<EMaterialDomain> SupportedDomains;
+	static const TArray<EBlendMode> SupportedBlendModes;
+
 	static UDynamicMaterialModelEditorOnlyData* Get(UDynamicMaterialModel* InModel);
 	static UDynamicMaterialModelEditorOnlyData* Get(TWeakObjectPtr<UDynamicMaterialModel> InModelWeak);
 	static UDynamicMaterialModelEditorOnlyData* Get(const TScriptInterface<IDynamicMaterialModelEditorOnlyDataInterface>& InInterface);
@@ -85,6 +88,18 @@ public:
 	UFUNCTION(BlueprintCallable, Category = "Material Designer")
 	void SetShadingModel(EDMMaterialShadingModel InShadingModel);
 
+	UFUNCTION(BlueprintPure, Category = "Material Designer")
+	bool IsPixelAnimationFlagSet() const;
+
+	UFUNCTION(BlueprintCallable, Category = "Material Designer")
+	void SetPixelAnimationFlag(bool bInFlagValue);
+
+	UFUNCTION(BlueprintPure, Category = "Material Designer")
+	bool IsTwoSidedFlagSet() const;
+
+	UFUNCTION(BlueprintCallable, Category = "Material Designer")
+	void SetTwoSidedFlag(bool bInFlagValue);
+
 	UFUNCTION(BlueprintCallable, Category = "Material Designer")
 	void OpenMaterialEditor() const;
 
@@ -117,12 +132,6 @@ public:
 
 	UFUNCTION(BlueprintCallable, Category = "Material Designer")
 	void UnassignMaterialProperty(EDMMaterialPropertyType Property);
-
-	UFUNCTION(BlueprintPure, Category = "Material Designer")
-	bool IsPixelAnimationFlagSet() const;
-
-	UFUNCTION(BlueprintCallable, Category = "Material Designer")
-	void SetPixelAnimationFlag(bool bInFlagValue);
 	
 	FDMOnMaterialBuilt& GetOnMaterialBuiltDelegate() { return OnMaterialBuiltDelegate; }
 	FDMOnValueListUpdated& GetOnValueListUpdateDelegate() { return OnValueListUpdateDelegate; }
@@ -135,6 +144,10 @@ public:
 
 	TSharedRef<FDMMaterialBuildState> CreateBuildState(UMaterial* InMaterialToBuild, bool bInDirtyAssets = true) const;
 
+	bool NeedsWizard() const;
+
+	void OnWizardComplete();
+
 	//~ Begin FNotifyHook
 	virtual void NotifyPreChange(class FEditPropertyChain* PropertyAboutToChange) {}
 	virtual void NotifyPostChange(const FPropertyChangedEvent& PropertyChangedEvent, class FEditPropertyChain* PropertyThatChanged);
@@ -180,6 +193,12 @@ protected:
 	UPROPERTY(EditInstanceOnly, BlueprintReadOnly, Category = "Material Designer")
 	EDMMaterialShadingModel ShadingModel = EDMMaterialShadingModel::Unlit;
 
+	UPROPERTY(VisibleInstanceOnly, BlueprintReadOnly, Category = "Material Designer")
+	bool bPixelAnimationFlag;
+
+	UPROPERTY(VisibleInstanceOnly, BlueprintReadOnly, Category = "Material Designer")
+	bool bTwoSidedFlag;
+
 	UPROPERTY(VisibleInstanceOnly, BlueprintReadOnly, Instanced, Category = "Material Designer")
 	TMap<EDMMaterialPropertyType, TObjectPtr<UDMMaterialProperty>> Properties;
 
@@ -195,8 +214,12 @@ protected:
 	UPROPERTY(VisibleInstanceOnly, BlueprintReadOnly, Category = "Material Designer")
 	bool bCreateMaterialPackage;
 
-	UPROPERTY(VisibleInstanceOnly, BlueprintReadOnly, Category = "Material Designer")
-	bool bPixelAnimationFlag;
+	/**
+	 * This is just a quick, transient flag that marks this model as needing to
+	 * open the wizard when it is opened by the material designer. It is initially
+	 * set to true, but set to false on first opening and on post load.
+	 */
+	bool bUseWizard;
 
 	FDMOnMaterialBuilt OnMaterialBuiltDelegate;
 	FDMOnValueListUpdated OnValueListUpdateDelegate;
