Material Designer
- Added a 'source' object to the error log / toast function.
- Fix for the material function and material stage function classes checking their inputs and outputs for a valid throughput.

diff --git a/Engine/Plugins/Experimental/DynamicMaterial/Source/DynamicMaterialEditor/Private/Components/DMMaterialEffectFunction.cpp b/Engine/Plugins/Experimental/DynamicMaterial/Source/DynamicMaterialEditor/Private/Components/DMMaterialEffectFunction.cpp
--- a/Engine/Plugins/Experimental/DynamicMaterial/Source/DynamicMaterialEditor/Private/Components/DMMaterialEffectFunction.cpp
+++ b/Engine/Plugins/Experimental/DynamicMaterial/Source/DynamicMaterialEditor/Private/Components/DMMaterialEffectFunction.cpp
@@ -15,6 +15,7 @@
 #include "Materials/MaterialExpressionFunctionOutput.h"
 #include "Materials/MaterialExpressionMaterialFunctionCall.h"
 #include "Materials/MaterialFunctionInterface.h"
+#include "MaterialValueType.h"
 #include "Model/DMMaterialBuildState.h"
 #include "Model/DynamicMaterialModelEditorOnlyData.h"
 
@@ -285,14 +286,42 @@ bool UDMMaterialEffectFunction::NeedsFunctionInit() const
 
 	if (Outputs.IsEmpty())
 	{
-		UE::DynamicMaterialEditor::Private::LogError(TEXT("Effect Function must have at least one output."));
+		UE::DynamicMaterialEditor::Private::LogError(TEXT("Effect Function must have at least one output."), true, this);
 		return false;
 	}
 
-	if (!Inputs[0].ExpressionInput || !Outputs[0].ExpressionOutput
-		|| Inputs[0].ExpressionInput->InputType != Outputs[0].ExpressionOutput->GetOutputType(0))
+	const uint32 InputType = Inputs[0].ExpressionInput ? Inputs[0].ExpressionInput->GetInputType(0) : MCT_Unknown;
+	const bool bInputTypeIsFloat = (InputType & MCT_Float) != 0;
+
+	const uint32 OutputType = Outputs[0].ExpressionOutput ? Outputs[0].ExpressionOutput->GetOutputType(0) : MCT_Unknown;
+	const bool bOutputTypeIsFloat = (InputType & MCT_Float) != 0;
+
+	/**
+	 * First input and output types must match.
+	 * Previously MCT_Float (a combination of float 1, 2, 3 and 4) caused a equality check to fail. It is now more rigorous.
+	 */
+	bool bValidThroughput = true;
+
+	if (InputType == MCT_Unknown || OutputType == MCT_Unknown)
+	{
+		bValidThroughput = false;
+	}
+	else if (bInputTypeIsFloat && bOutputTypeIsFloat)
+	{
+		// If they are both float types and either of them are MCT_Float then they are a match.
+		if (InputType != MCT_Float && OutputType != MCT_Float)
+		{
+			bValidThroughput = InputType == OutputType;
+		}
+	}
+	else if (InputType != OutputType)
+	{
+		bValidThroughput = false;
+	}
+
+	if (!bValidThroughput)
 	{
-		UE::DynamicMaterialEditor::Private::LogError(TEXT("Effect Function's first input must match its first output."));
+		UE::DynamicMaterialEditor::Private::LogError(TEXT("Effect Function's first input must match its first output."), true, this);
 		return false;
 	}
 
@@ -307,7 +336,7 @@ bool UDMMaterialEffectFunction::NeedsFunctionInit() const
 
 		if (!IsValid(FunctionInput))
 		{
-			UE::DynamicMaterialEditor::Private::LogError(TEXT("Effect Function has missing input object."));
+			UE::DynamicMaterialEditor::Private::LogError(TEXT("Effect Function has missing input object."), true, this);
 			return false;
 		}
 
@@ -323,7 +352,7 @@ bool UDMMaterialEffectFunction::NeedsFunctionInit() const
 					break;
 
 				default:
-					UE::DynamicMaterialEditor::Private::LogError(TEXT("Effect Function has invalid first input -  must be a scalar or vector."));
+					UE::DynamicMaterialEditor::Private::LogError(TEXT("Effect Function has invalid first input -  must be a scalar or vector."), true, this);
 					return false;
 			}
 
@@ -357,7 +386,7 @@ bool UDMMaterialEffectFunction::NeedsFunctionInit() const
 				break;
 
 			default:
-				UE::DynamicMaterialEditor::Private::LogError(TEXT("Effect Function has invalid input type -  must be a scalar, vector or texture."));
+				UE::DynamicMaterialEditor::Private::LogError(TEXT("Effect Function has invalid input type -  must be a scalar, vector or texture."), true, this);
 				return false;
 		}
 
@@ -407,14 +436,14 @@ void UDMMaterialEffectFunction::InitFunction()
 
 	if (Inputs.IsEmpty())
 	{
-		UE::DynamicMaterialEditor::Private::LogError(TEXT("Effect Function must have at least one input."));
+		UE::DynamicMaterialEditor::Private::LogError(TEXT("Effect Function must have at least one input."), true, this);
 		MaterialFunctionPtr = nullptr;
 		return;
 	}
 
 	if (Outputs.IsEmpty())
 	{
-		UE::DynamicMaterialEditor::Private::LogError(TEXT("Effect Function must have at least one output."));
+		UE::DynamicMaterialEditor::Private::LogError(TEXT("Effect Function must have at least one output."), true, this);
 		MaterialFunctionPtr = nullptr;
 		return;
 	}
@@ -429,7 +458,7 @@ void UDMMaterialEffectFunction::InitFunction()
 
 		if (!IsValid(FunctionInput))
 		{
-			UE::DynamicMaterialEditor::Private::LogError(TEXT("Effect Function has missing input object."));
+			UE::DynamicMaterialEditor::Private::LogError(TEXT("Effect Function has missing input object."), true, this);
 			MaterialFunctionPtr = nullptr;
 			return;
 		}
@@ -456,7 +485,7 @@ void UDMMaterialEffectFunction::InitFunction()
 					break;
 
 				default:
-					UE::DynamicMaterialEditor::Private::LogError(TEXT("Effect Function has invalid first input -  must be a scalar or vector."));
+					UE::DynamicMaterialEditor::Private::LogError(TEXT("Effect Function has invalid first input -  must be a scalar or vector."), true, this);
 					MaterialFunctionPtr = nullptr;
 					return;
 			}
@@ -492,7 +521,7 @@ void UDMMaterialEffectFunction::InitFunction()
 				break;
 
 			default:
-				UE::DynamicMaterialEditor::Private::LogError(TEXT("Effect Function has invalid input type -  must be a scalar, vector or texture."));
+				UE::DynamicMaterialEditor::Private::LogError(TEXT("Effect Function has invalid input type -  must be a scalar, vector or texture."), true, this);
 				MaterialFunctionPtr = nullptr;
 				return;
 		}
diff --git a/Engine/Plugins/Experimental/DynamicMaterial/Source/DynamicMaterialEditor/Private/Components/DMMaterialStageFunction.cpp b/Engine/Plugins/Experimental/DynamicMaterial/Source/DynamicMaterialEditor/Private/Components/DMMaterialStageFunction.cpp
--- a/Engine/Plugins/Experimental/DynamicMaterial/Source/DynamicMaterialEditor/Private/Components/DMMaterialStageFunction.cpp
+++ b/Engine/Plugins/Experimental/DynamicMaterial/Source/DynamicMaterialEditor/Private/Components/DMMaterialStageFunction.cpp
@@ -17,6 +17,7 @@
 #include "Materials/MaterialExpressionFunctionOutput.h"
 #include "Materials/MaterialExpressionMaterialFunctionCall.h"
 #include "Materials/MaterialFunctionInterface.h"
+#include "MaterialValueType.h"
 #include "Model/DMMaterialBuildState.h"
 #include "Model/DMMaterialBuildUtils.h"
 #include "Model/DynamicMaterialModelEditorOnlyData.h"
@@ -291,14 +292,42 @@ bool UDMMaterialStageFunction::NeedsFunctionInit() const
 
 	if (Outputs.IsEmpty())
 	{
-		UE::DynamicMaterialEditor::Private::LogError(TEXT("Function must have at least one output."));
+		UE::DynamicMaterialEditor::Private::LogError(TEXT("Function must have at least one output."), true, this);
 		return false;
 	}
 
-	if (!Inputs[0].ExpressionInput || !Outputs[0].ExpressionOutput
-		|| Inputs[0].ExpressionInput->InputType != Outputs[0].ExpressionOutput->GetOutputType(0))
+	const uint32 InputType = Inputs[0].ExpressionInput ? Inputs[0].ExpressionInput->GetInputType(0) : MCT_Unknown;
+	const bool bInputTypeIsFloat = (InputType & MCT_Float) != 0;
+
+	const uint32 OutputType = Outputs[0].ExpressionOutput ? Outputs[0].ExpressionOutput->GetOutputType(0) : MCT_Unknown;
+	const bool bOutputTypeIsFloat = (InputType & MCT_Float) != 0;
+
+	/**
+	 * First input and output types must match.
+	 * Previously MCT_Float (a combination of float 1, 2, 3 and 4) caused a equality check to fail. It is now more rigorous.
+	 */
+	bool bValidThroughput = true;
+
+	if (InputType == MCT_Unknown || OutputType == MCT_Unknown)
+	{
+		bValidThroughput = false;
+	}
+	else if (bInputTypeIsFloat && bOutputTypeIsFloat)
+	{
+		// If they are both float types and either of them are MCT_Float then they are a match.
+		if (InputType != MCT_Float && OutputType != MCT_Float)
+		{
+			bValidThroughput = InputType == OutputType;
+		}
+	}
+	else if (InputType != OutputType)
+	{
+		bValidThroughput = false;
+	}
+
+	if (!bValidThroughput)
 	{
-		UE::DynamicMaterialEditor::Private::LogError(TEXT("Function's first input must match its first output."));
+		UE::DynamicMaterialEditor::Private::LogError(TEXT("Function's first input must match its first output."), true, this);
 		return false;
 	}
 
@@ -313,7 +342,7 @@ bool UDMMaterialStageFunction::NeedsFunctionInit() const
 
 		if (!IsValid(FunctionInput))
 		{
-			UE::DynamicMaterialEditor::Private::LogError(TEXT("Function has missing input object."));
+			UE::DynamicMaterialEditor::Private::LogError(TEXT("Function has missing input object."), true, this);
 			return false;
 		}
 
@@ -327,7 +356,7 @@ bool UDMMaterialStageFunction::NeedsFunctionInit() const
 					break;
 
 				default:
-					UE::DynamicMaterialEditor::Private::LogError(TEXT("Function has invalid first input - must be a scalar or vector3."));
+					UE::DynamicMaterialEditor::Private::LogError(TEXT("Function has invalid first input - must be a scalar or vector3."), true, this);
 					return false;
 			}
 
@@ -361,7 +390,7 @@ bool UDMMaterialStageFunction::NeedsFunctionInit() const
 				break;
 
 			default:
-				UE::DynamicMaterialEditor::Private::LogError(TEXT("Function has invalid input type - must be a scalar, vector or texture."));
+				UE::DynamicMaterialEditor::Private::LogError(TEXT("Function has invalid input type - must be a scalar, vector or texture."), true, this);
 				return false;
 		}
 
@@ -419,14 +448,14 @@ void UDMMaterialStageFunction::InitFunction()
 
 	if (Inputs.IsEmpty())
 	{
-		UE::DynamicMaterialEditor::Private::LogError(TEXT("Function must have at least one input."));
+		UE::DynamicMaterialEditor::Private::LogError(TEXT("Function must have at least one input."), true, this);
 		MaterialFunction = nullptr;
 		return;
 	}
 
 	if (Outputs.IsEmpty())
 	{
-		UE::DynamicMaterialEditor::Private::LogError(TEXT("Function must have at least one output."));
+		UE::DynamicMaterialEditor::Private::LogError(TEXT("Function must have at least one output."), true, this);
 		MaterialFunction = nullptr;
 		return;
 	}
@@ -448,7 +477,7 @@ void UDMMaterialStageFunction::InitFunction()
 
 		if (!IsValid(FunctionInput))
 		{
-			UE::DynamicMaterialEditor::Private::LogError(TEXT("Function has missing input object."));
+			UE::DynamicMaterialEditor::Private::LogError(TEXT("Function has missing input object."), true, this);
 			InputConnectors.SetNum(1);
 			MaterialFunction = nullptr;
 			return;
@@ -464,7 +493,7 @@ void UDMMaterialStageFunction::InitFunction()
 					break;
 
 				default:
-					UE::DynamicMaterialEditor::Private::LogError(TEXT("Function has invalid first input - must be a scalar or vector."));
+					UE::DynamicMaterialEditor::Private::LogError(TEXT("Function has invalid first input - must be a scalar or vector."), true, this);
 					InputConnectors.SetNum(1);
 					MaterialFunction = nullptr;
 					return;
@@ -508,7 +537,7 @@ void UDMMaterialStageFunction::InitFunction()
 				break;
 
 			default:
-				UE::DynamicMaterialEditor::Private::LogError(TEXT("Function has invalid input type - must be a scalar, vector or texture."));
+				UE::DynamicMaterialEditor::Private::LogError(TEXT("Function has invalid input type - must be a scalar, vector or texture."), true, this);
 				InputConnectors.SetNum(1);
 				MaterialFunction = nullptr;
 				return;
diff --git a/Engine/Plugins/Experimental/DynamicMaterial/Source/DynamicMaterialEditor/Private/Components/DMMaterialSubStage.cpp b/Engine/Plugins/Experimental/DynamicMaterial/Source/DynamicMaterialEditor/Private/Components/DMMaterialSubStage.cpp
--- a/Engine/Plugins/Experimental/DynamicMaterial/Source/DynamicMaterialEditor/Private/Components/DMMaterialSubStage.cpp
+++ b/Engine/Plugins/Experimental/DynamicMaterial/Source/DynamicMaterialEditor/Private/Components/DMMaterialSubStage.cpp
@@ -55,7 +55,7 @@ void UDMMaterialSubStage::PostEditorDuplicate(UDynamicMaterialModel* InMaterialM
 	}
 	else
 	{
-		UE::DynamicMaterialEditor::Private::LogError(TEXT("Wrong parent component passed to substage."));
+		UE::DynamicMaterialEditor::Private::LogError(TEXT("Wrong parent component passed to substage."), true, this);
 		ParentStage = nullptr;
 		ParentComponent = nullptr;
 		Super::PostEditorDuplicate(InMaterialModel, InParent);
diff --git a/Engine/Plugins/Experimental/DynamicMaterial/Source/DynamicMaterialEditor/Private/DMBlueprintFunctionLibrary.cpp b/Engine/Plugins/Experimental/DynamicMaterial/Source/DynamicMaterialEditor/Private/DMBlueprintFunctionLibrary.cpp
--- a/Engine/Plugins/Experimental/DynamicMaterial/Source/DynamicMaterialEditor/Private/DMBlueprintFunctionLibrary.cpp
+++ b/Engine/Plugins/Experimental/DynamicMaterial/Source/DynamicMaterialEditor/Private/DMBlueprintFunctionLibrary.cpp
@@ -267,7 +267,7 @@ bool UDMBlueprintFunctionLibrary::ExportMaterialInstance(UDynamicMaterialModel*
 	UDynamicMaterialInstance* Instance = InMaterialModel->GetDynamicMaterialInstance();
 	if (!Instance)
 	{
-		UE::DynamicMaterialEditor::Private::LogError(TEXT("Failed to find a Material Designer Instance to export."));
+		UE::DynamicMaterialEditor::Private::LogError(TEXT("Failed to find a Material Designer Instance to export."), true, InMaterialModel);
 		return false;
 	}
 
diff --git a/Engine/Plugins/Experimental/DynamicMaterial/Source/DynamicMaterialEditor/Private/DMPrivate.cpp b/Engine/Plugins/Experimental/DynamicMaterial/Source/DynamicMaterialEditor/Private/DMPrivate.cpp
--- a/Engine/Plugins/Experimental/DynamicMaterial/Source/DynamicMaterialEditor/Private/DMPrivate.cpp
+++ b/Engine/Plugins/Experimental/DynamicMaterial/Source/DynamicMaterialEditor/Private/DMPrivate.cpp
@@ -9,6 +9,8 @@
 #include "Framework/Notifications/NotificationManager.h"
 #include "Widgets/Notifications/SNotificationList.h"
 
+#define LOCTEXT_NAMESPACE "FDMMaterialLayerReference"
+
 namespace UE::DynamicMaterialEditor::Private
 {
 	void SetMask(FExpressionInput& InInputConnector, const FExpressionOutput& InOutputConnector, int32 InChannelOverride)
@@ -70,14 +72,23 @@ namespace UE::DynamicMaterialEditor::Private
 	}
 
 	static bool bAllowUIFeedback = false;
+	static FText LogErrorObjectFormat = LOCTEXT("LogErrorObjectFormat", "%s (Source: %s)");
 
-	void LogError(const FString& InMessage, bool bInToast)
+	void LogError(const FString& InMessage, bool bInToast, const UObject* InSource)
 	{
-		UE_LOG(LogDynamicMaterialEditor, Error, TEXT("%s"), *InMessage);
+		const FString Message = IsValid(InSource)
+			? FString::Printf(TEXT("%s (Source: %s)"), *InMessage, *InSource->GetPathName())
+			: InMessage;
+
+		UE_LOG(LogDynamicMaterialEditor, Error, TEXT("%s"), *Message);
 
 		if (bAllowUIFeedback && bInToast)
 		{
-			FNotificationInfo Info(FText::FromString(InMessage));
+			const FText MessageText = IsValid(InSource)
+				? FText::Format(LogErrorObjectFormat, FText::FromString(InMessage), FText::FromString(InSource->GetPathName()))
+				: FText::FromString(InMessage);
+
+			FNotificationInfo Info(MessageText);
 			Info.ExpireDuration = 5.0f;
 			FSlateNotificationManager::Get().AddNotification(Info);
 		}
@@ -144,3 +155,5 @@ FDMScopedUITransaction::FDMScopedUITransaction(const FText& InSessionName, bool
 	, UIFeedbackGuard(TGuardValue<bool>(UE::DynamicMaterialEditor::Private::bAllowUIFeedback, true))
 {
 }
+
+#undef LOCTEXT_NAMESPACE
diff --git a/Engine/Plugins/Experimental/DynamicMaterial/Source/DynamicMaterialEditor/Private/DMPrivate.h b/Engine/Plugins/Experimental/DynamicMaterial/Source/DynamicMaterialEditor/Private/DMPrivate.h
--- a/Engine/Plugins/Experimental/DynamicMaterial/Source/DynamicMaterialEditor/Private/DMPrivate.h
+++ b/Engine/Plugins/Experimental/DynamicMaterial/Source/DynamicMaterialEditor/Private/DMPrivate.h
@@ -30,7 +30,7 @@ namespace UE::DynamicMaterialEditor::Private
 
 	bool IsCustomMaterialProperty(EDMMaterialPropertyType InMaterialProperty);
 
-	void LogError(const FString& InMessage, bool bInToast = true);
+	void LogError(const FString& InMessage, bool bInToast = true, const UObject* InSource = nullptr);
 }
 
 struct FDMMaterialLayerReference
