Material Designer
- Moved new-slot default layout to material property classes.
- Added property short names.
- Added material property presets configurable in the editor settings.
- Updated code to support the property presets.
- Updated UI to support more than 2 material properties.

diff --git a/Engine/Plugins/Experimental/DynamicMaterial/Source/DynamicMaterial/Public/DMDefs.h b/Engine/Plugins/Experimental/DynamicMaterial/Source/DynamicMaterial/Public/DMDefs.h
--- a/Engine/Plugins/Experimental/DynamicMaterial/Source/DynamicMaterial/Public/DMDefs.h
+++ b/Engine/Plugins/Experimental/DynamicMaterial/Source/DynamicMaterial/Public/DMDefs.h
@@ -10,21 +10,22 @@
 class UTexture;
 class UDMMaterialStage;
 
+/** Some short names provided so they fit nicely on buttons! */
 UENUM(BlueprintType)
 enum class EDMMaterialPropertyType : uint8
 {
 	None = 0,
-	BaseColor,
-	EmissiveColor,
+	BaseColor UMETA(ShortName = "RGB"),
+	EmissiveColor UMETA(ShortName = "RGB"),
 	Opacity,
-	OpacityMask,
+	OpacityMask UMETA(ShortName = "Mask"),
 	Roughness,
 	Specular,
 	Metallic,
 	Normal,
-	PixelDepthOffset,
-	WorldPositionOffset,
-	AmbientOcclusion,
+	PixelDepthOffset UMETA(ShortName = "PDO"),
+	WorldPositionOffset UMETA(ShortName = "WPO"),
+	AmbientOcclusion UMETA(ShortName = "AO"),
 	Anisotropy,
 	Refraction,
 	Tangent,
diff --git a/Engine/Plugins/Experimental/DynamicMaterial/Source/DynamicMaterialEditor/Private/Components/DMMaterialSlot.cpp b/Engine/Plugins/Experimental/DynamicMaterial/Source/DynamicMaterialEditor/Private/Components/DMMaterialSlot.cpp
--- a/Engine/Plugins/Experimental/DynamicMaterial/Source/DynamicMaterialEditor/Private/Components/DMMaterialSlot.cpp
+++ b/Engine/Plugins/Experimental/DynamicMaterial/Source/DynamicMaterialEditor/Private/Components/DMMaterialSlot.cpp
@@ -254,7 +254,6 @@ void UDMMaterialSlot::UpdateOutputConnectorTypes()
 
 	if (LayerObjects.IsEmpty())
 	{
-		UpdateMaterialProperties();
 		return;
 	}
 
@@ -302,24 +301,6 @@ void UDMMaterialSlot::UpdateMaterialProperties()
 
 	if (LayerObjects.IsEmpty())
 	{
-		const TArray<EDMMaterialPropertyType> SlotMaterialProperties = ModelEditorOnlyData->GetMaterialPropertiesForSlot(this);
-
-		for (EDMMaterialPropertyType MaterialProperty : SlotMaterialProperties)
-		{
-			UDMMaterialSlot* CurrentSlot = ModelEditorOnlyData->GetSlotForMaterialProperty(MaterialProperty);
-			check(CurrentSlot == nullptr || CurrentSlot == this);
-
-			if (CurrentSlot == this)
-			{
-				if (GUndo)
-				{
-					ModelEditorOnlyData->Modify();
-				}
-
-				ModelEditorOnlyData->UnassignMaterialProperty(MaterialProperty);
-			}
-		}
-
 		return;
 	}
 
@@ -465,9 +446,6 @@ UDMMaterialLayerObject* UDMMaterialSlot::AddLayerWithMask(EDMMaterialPropertyTyp
 		InNewMask->Modify();
 	}
 
-	InNewBase->SetBeingEdited(false);
-	InNewMask->SetBeingEdited(false);
-
 	UDMMaterialLayerObject* NewLayer = UDMMaterialLayerObject::CreateLayer(this, InMaterialProperty, {InNewBase, InNewMask});
 	LayerObjects.Add(NewLayer);
 
@@ -529,21 +507,46 @@ bool UDMMaterialSlot::PasteLayer(UDMMaterialLayerObject* InLayer)
 	return true;
 }
 
-bool UDMMaterialSlot::RemoveLayer(UDMMaterialLayerObject* InLayer)
+bool UDMMaterialSlot::CanRemoveLayer(const UDMMaterialLayerObject* InLayer) const
 {
 	if (!IsComponentValid())
 	{
 		return false;
 	}
 
-	if (LayerObjects.Num() == 1)
+	check(InLayer);
+	check(InLayer->GetSlot() == this);
+
+	const EDMMaterialPropertyType LayerProperty = InLayer->GetMaterialProperty();
+
+	if (LayerProperty == EDMMaterialPropertyType::BaseColor || LayerProperty == EDMMaterialPropertyType::EmissiveColor)
+	{
+		int32 LayerPropertyCount = 0;
+
+		for (UDMMaterialLayerObject* Layer : LayerObjects)
+		{
+			if (Layer->GetMaterialProperty() == LayerProperty)
+			{
+				++LayerPropertyCount;
+			}
+		}
+
+		if (LayerPropertyCount == 1)
+		{
+			return false;
+		}
+	}
+
+	return InLayer->FindIndex() != INDEX_NONE;
+}
+
+bool UDMMaterialSlot::RemoveLayer(UDMMaterialLayerObject* InLayer)
+{
+	if (!CanRemoveLayer(InLayer))
 	{
 		return false;
 	}
 
-	check(InLayer);
-	check(InLayer->GetSlot() == this);
-
 	const int32 LayerIndex = InLayer->FindIndex();
 
 	if (LayerIndex == INDEX_NONE)
@@ -660,9 +663,7 @@ bool UDMMaterialSlot::ReferencedBySlot(UDMMaterialSlot* InOtherSlot)
 	check(InOtherSlot);
 	check(InOtherSlot != this);
 
-	int32* CountPtr = SlotsReferencedBy.Find(InOtherSlot);
-
-	if (CountPtr)
+	if (int32* CountPtr = SlotsReferencedBy.Find(InOtherSlot))
 	{
 		++(*CountPtr);
 		return false;
@@ -1157,8 +1158,8 @@ void UDMMaterialSlot::UpdatePreviewMaterialProperty(const TSharedRef<FDMMaterial
 	}
 }
 
-bool UDMMaterialSlot::SetLayerMaterialPropertyAndReplaceOthers(UDMMaterialLayerObject* InLayer, EDMMaterialPropertyType InMaterialProperty, 
-	EDMMaterialPropertyType InReplaceWithProperty)
+bool UDMMaterialSlot::SetLayerMaterialPropertyAndReplaceOthers(UDMMaterialLayerObject* InLayer, EDMMaterialPropertyType InPropertyFrom, 
+	EDMMaterialPropertyType InPropertyTo)
 {
 	if (!IsComponentValid())
 	{
@@ -1168,7 +1169,7 @@ bool UDMMaterialSlot::SetLayerMaterialPropertyAndReplaceOthers(UDMMaterialLayerO
 	UDynamicMaterialModelEditorOnlyData* ModelEditorOnlyData = GetMaterialModelEditorOnlyData();
 	check(ModelEditorOnlyData);
 
-	UDMMaterialSlot* CurrentSlot = ModelEditorOnlyData->GetSlotForMaterialProperty(InMaterialProperty);
+	UDMMaterialSlot* CurrentSlot = ModelEditorOnlyData->GetSlotForMaterialProperty(InPropertyFrom);
 
 	if (!CurrentSlot || CurrentSlot == this)
 	{
@@ -1177,7 +1178,7 @@ bool UDMMaterialSlot::SetLayerMaterialPropertyAndReplaceOthers(UDMMaterialLayerO
 			InLayer->Modify();
 		}
 
-		InLayer->SetMaterialProperty(InMaterialProperty);
+		InLayer->SetMaterialProperty(InPropertyFrom);
 		return false; // Could be caused by asynchronous input
 	}
 
@@ -1186,14 +1187,14 @@ bool UDMMaterialSlot::SetLayerMaterialPropertyAndReplaceOthers(UDMMaterialLayerO
 		
 		for (TObjectPtr<UDMMaterialLayerObject>& Layer : CurrentSlot->LayerObjects)
 		{
-			if (Layer->GetMaterialProperty() == InMaterialProperty)
+			if (Layer->GetMaterialProperty() == InPropertyFrom)
 			{
 				if (GUndo)
 				{
 					Layer->Modify();
 				}
 
-				Layer->SetMaterialProperty(InReplaceWithProperty);
+				Layer->SetMaterialProperty(InPropertyTo);
 			}
 
 			if (UDMMaterialStage* BaseStage = Layer->GetStage(EDMMaterialLayerStage::Base))
@@ -1207,10 +1208,10 @@ bool UDMMaterialSlot::SetLayerMaterialPropertyAndReplaceOthers(UDMMaterialLayerO
 						FDMMaterialStageConnectorChannel& Channel = CurrentSlotStageInputMap[InputIdx].Channels[ChannelIdx];
 
 						if (Channel.SourceIndex == FDMMaterialStageConnectorChannel::PREVIOUS_STAGE
-							&& Channel.MaterialProperty == InMaterialProperty)
+							&& Channel.MaterialProperty == InPropertyFrom)
 						{
 							// Delve into class internals, avoiding the const issues above.
-							Channel.MaterialProperty = InReplaceWithProperty;
+							Channel.MaterialProperty = InPropertyTo;
 						}
 					}
 				}
@@ -1226,7 +1227,75 @@ bool UDMMaterialSlot::SetLayerMaterialPropertyAndReplaceOthers(UDMMaterialLayerO
 		}
 	}
 
-	return InLayer->SetMaterialProperty(InMaterialProperty);
+	return InLayer->SetMaterialProperty(InPropertyFrom);
+}
+
+bool UDMMaterialSlot::ChangeMaterialProperty(EDMMaterialPropertyType InPropertyFrom, EDMMaterialPropertyType InReplaceWithProperty)
+{
+	if (!IsComponentValid())
+	{
+		return false;
+	}
+
+	UDynamicMaterialModelEditorOnlyData* ModelEditorOnlyData = GetMaterialModelEditorOnlyData();
+	check(ModelEditorOnlyData);
+
+	FDMUpdateGuard Guard;
+	UDMMaterialLayerObject* FirstLayerObject = nullptr;
+
+	for (TObjectPtr<UDMMaterialLayerObject>& Layer : LayerObjects)
+	{
+		if (Layer->GetMaterialProperty() == InPropertyFrom)
+		{
+			if (GUndo)
+			{
+				Layer->Modify();
+			}
+
+			Layer->SetMaterialProperty(InReplaceWithProperty);
+
+			if (!FirstLayerObject)
+			{
+				FirstLayerObject = Layer;
+			}
+		}
+
+		if (UDMMaterialStage* BaseStage = Layer->GetStage(EDMMaterialLayerStage::Base))
+		{
+			TArray<FDMMaterialStageConnection>& CurrentSlotStageInputMap = BaseStage->GetInputConnectionMap();
+
+			for (int32 InputIdx = 0; InputIdx < CurrentSlotStageInputMap.Num(); ++InputIdx)
+			{
+				for (int32 ChannelIdx = 0; ChannelIdx < CurrentSlotStageInputMap[InputIdx].Channels.Num(); ++ChannelIdx)
+				{
+					FDMMaterialStageConnectorChannel& Channel = CurrentSlotStageInputMap[InputIdx].Channels[ChannelIdx];
+
+					if (Channel.SourceIndex == FDMMaterialStageConnectorChannel::PREVIOUS_STAGE
+						&& Channel.MaterialProperty == InPropertyFrom)
+					{
+						// Delve into class internals, avoiding the const issues above.
+						Channel.MaterialProperty = InReplaceWithProperty;
+					}
+				}
+			}
+		}
+	}
+
+	ModelEditorOnlyData->UnassignMaterialProperty(InPropertyFrom);
+	ModelEditorOnlyData->AssignMaterialPropertyToSlot(InReplaceWithProperty, this);
+
+	if (const FDMMaterialSlotOutputConnectorTypes* ConnectorTypes = OutputConnectorTypes.Find(InPropertyFrom))
+	{
+		OutputConnectorTypes.Emplace(InReplaceWithProperty, *ConnectorTypes);
+		OutputConnectorTypes.Remove(InPropertyFrom);
+	}
+
+	if (FirstLayerObject)
+	{
+		FirstLayerObject->Update(EDMUpdateType::Structure);
+	}
+
+	return true;
 }
 
 void UDMMaterialSlot::DoClean()
diff --git a/Engine/Plugins/Experimental/DynamicMaterial/Source/DynamicMaterialEditor/Private/Components/MaterialProperties/DMMPBaseColor.cpp b/Engine/Plugins/Experimental/DynamicMaterial/Source/DynamicMaterialEditor/Private/Components/MaterialProperties/DMMPBaseColor.cpp
--- a/Engine/Plugins/Experimental/DynamicMaterial/Source/DynamicMaterialEditor/Private/Components/MaterialProperties/DMMPBaseColor.cpp
+++ b/Engine/Plugins/Experimental/DynamicMaterial/Source/DynamicMaterialEditor/Private/Components/MaterialProperties/DMMPBaseColor.cpp
@@ -1,8 +1,32 @@
 // Copyright Epic Games, Inc. All Rights Reserved.
 
 #include "Components/MaterialProperties/DMMPBaseColor.h"
+
+#include "Components/DMMaterialSlot.h"
+#include "Components/DMMaterialStage.h"
+#include "Components/DMMaterialStageBlend.h"
+#include "Components/DMMaterialStageThroughputLayerBlend.h"
+#include "Components/MaterialStageBlends/DMMSBNormal.h"
+#include "Components/MaterialStageExpressions/DMMSETextureSample.h"
+#include "Components/MaterialStageInputs/DMMSIExpression.h"
 #include "Model/DynamicMaterialModelEditorOnlyData.h"
 
+void UDMMaterialPropertyBaseColor::AddDefaultRGBLayer(EDMMaterialPropertyType InMaterialProperty, UDMMaterialSlot* InSlot)
+{
+	UDMMaterialStage* DefaultStage = UDMMaterialStageBlend::CreateStage(UDMMaterialStageBlendNormal::StaticClass());
+	check(DefaultStage);
+
+	UDMMaterialStage* MaskStage = UDMMaterialStageThroughputLayerBlend::CreateStage();
+	check(MaskStage);
+
+	InSlot->AddLayerWithMask(InMaterialProperty, DefaultStage, MaskStage);
+
+	UDMMaterialStageInputExpression* BaseInputExpression = UDMMaterialStageInputExpression::ChangeStageInput_Expression(DefaultStage,
+		UDMMaterialStageExpressionTextureSample::StaticClass(), UDMMaterialStageBlendNormal::InputB,
+		FDMMaterialStageConnectorChannel::WHOLE_CHANNEL, 0,
+		FDMMaterialStageConnectorChannel::WHOLE_CHANNEL);
+}
+
 UDMMaterialPropertyBaseColor::UDMMaterialPropertyBaseColor()
 	: UDMMaterialProperty(
 		EDMMaterialPropertyType(EDMMaterialPropertyType::BaseColor),
@@ -19,3 +43,8 @@ UMaterialExpression* UDMMaterialPropertyBaseColor::GetDefaultInput(const TShared
 {
 	return CreateConstant(InBuildState, FVector::OneVector);
 }
+
+void UDMMaterialPropertyBaseColor::OnSlotAdded(UDMMaterialSlot* InSlot)
+{
+	AddDefaultRGBLayer(EDMMaterialPropertyType::BaseColor, InSlot);
+}
diff --git a/Engine/Plugins/Experimental/DynamicMaterial/Source/DynamicMaterialEditor/Private/Components/MaterialProperties/DMMPEmissiveColor.cpp b/Engine/Plugins/Experimental/DynamicMaterial/Source/DynamicMaterialEditor/Private/Components/MaterialProperties/DMMPEmissiveColor.cpp
--- a/Engine/Plugins/Experimental/DynamicMaterial/Source/DynamicMaterialEditor/Private/Components/MaterialProperties/DMMPEmissiveColor.cpp
+++ b/Engine/Plugins/Experimental/DynamicMaterial/Source/DynamicMaterialEditor/Private/Components/MaterialProperties/DMMPEmissiveColor.cpp
@@ -1,6 +1,8 @@
 // Copyright Epic Games, Inc. All Rights Reserved.
 
 #include "Components/MaterialProperties/DMMPEmissiveColor.h"
+
+#include "Components/MaterialProperties/DMMPBaseColor.h"
 #include "Model/DynamicMaterialModelEditorOnlyData.h"
 
 UDMMaterialPropertyEmissiveColor::UDMMaterialPropertyEmissiveColor()
@@ -27,3 +29,8 @@ UMaterialExpression* UDMMaterialPropertyEmissiveColor::GetDefaultInput(
 {
 	return CreateConstant(InBuildState, FVector::OneVector);
 }
+
+void UDMMaterialPropertyEmissiveColor::OnSlotAdded(UDMMaterialSlot* InSlot)
+{
+	UDMMaterialPropertyBaseColor::AddDefaultRGBLayer(EDMMaterialPropertyType::BaseColor, InSlot);
+}
diff --git a/Engine/Plugins/Experimental/DynamicMaterial/Source/DynamicMaterialEditor/Private/Components/MaterialProperties/DMMPOpacity.cpp b/Engine/Plugins/Experimental/DynamicMaterial/Source/DynamicMaterialEditor/Private/Components/MaterialProperties/DMMPOpacity.cpp
--- a/Engine/Plugins/Experimental/DynamicMaterial/Source/DynamicMaterialEditor/Private/Components/MaterialProperties/DMMPOpacity.cpp
+++ b/Engine/Plugins/Experimental/DynamicMaterial/Source/DynamicMaterialEditor/Private/Components/MaterialProperties/DMMPOpacity.cpp
@@ -1,8 +1,93 @@
 // Copyright Epic Games, Inc. All Rights Reserved.
 
 #include "Components/MaterialProperties/DMMPOpacity.h"
+
+#include "Components/DMMaterialLayer.h"
+#include "Components/DMMaterialSlot.h"
+#include "Components/DMMaterialStage.h"
+#include "Components/DMMaterialStageThroughputLayerBlend.h"
+#include "Components/DMMaterialSubStage.h"
+#include "Components/MaterialStageBlends/DMMSBNormal.h"
+#include "Components/MaterialStageExpressions/DMMSETextureSample.h"
+#include "Components/MaterialStageInputs/DMMSIExpression.h"
+#include "Components/MaterialStageInputs/DMMSIValue.h"
+#include "Components/MaterialValues/DMMaterialValueTexture.h"
+#include "DynamicMaterialEditorSettings.h"
 #include "Model/DynamicMaterialModelEditorOnlyData.h"
 
+void UDMMaterialPropertyOpacity::AddDefaultOpacityLayer(EDMMaterialPropertyType InMaterialProperty, UDMMaterialSlot* InSlot)
+{
+	UDMMaterialStage* DefaultStage = UDMMaterialStageBlend::CreateStage(UDMMaterialStageBlendNormal::StaticClass());
+	check(DefaultStage);
+
+	UDMMaterialStage* MaskStage = UDMMaterialStageThroughputLayerBlend::CreateStage();
+	check(MaskStage);
+
+	InSlot->AddLayerWithMask(InMaterialProperty, DefaultStage, MaskStage);
+
+	UDMMaterialStageInputExpression* BaseInputExpression = UDMMaterialStageInputExpression::ChangeStageInput_Expression(
+		DefaultStage,
+		UDMMaterialStageExpressionTextureSample::StaticClass(),
+		UDMMaterialStageBlendNormal::InputB,
+		FDMMaterialStageConnectorChannel::WHOLE_CHANNEL, 
+		0,
+		FDMMaterialStageConnectorChannel::WHOLE_CHANNEL
+	);
+
+	if (UDMMaterialStageExpressionTextureSample* BaseInputTextureSample = Cast<UDMMaterialStageExpressionTextureSample>(BaseInputExpression->GetMaterialStageExpression()))
+	{
+		BaseInputTextureSample->SetClampTextureEnabled(true);
+
+		if (UDMMaterialStage* BaseTextureInputStage = BaseInputExpression->GetSubStage())
+		{
+			const TArray<UDMMaterialStageInput*> BaseTextureStageInputs = BaseTextureInputStage->GetInputs();
+
+			for (UDMMaterialStageInput* BaseTextureStageInput : BaseTextureStageInputs)
+			{
+				if (UDMMaterialStageInputValue* BaseTextureInputValue = Cast<UDMMaterialStageInputValue>(BaseTextureStageInput))
+				{
+					if (UDMMaterialValueTexture* BaseTextureValue = Cast<UDMMaterialValueTexture>(BaseTextureInputValue->GetValue()))
+					{
+						BaseTextureValue->SetDefaultValue(UDynamicMaterialEditorSettings::Get()->DefaultOpacitySlotMask.LoadSynchronous());
+						BaseTextureValue->ApplyDefaultValue();
+						break;
+					}
+				}
+			}
+		}
+	}
+
+	const TArray<UDMMaterialStageInput*> MaskStageInputs = MaskStage->GetInputs();
+
+	for (UDMMaterialStageInput* MaskStageInput : MaskStageInputs)
+	{
+		if (UDMMaterialStageInputExpression* MaskInputExpression = Cast<UDMMaterialStageInputExpression>(MaskStageInput))
+		{
+			if (UDMMaterialStageExpressionTextureSample* MaskInputTextureSample = Cast<UDMMaterialStageExpressionTextureSample>(MaskInputExpression->GetMaterialStageExpression()))
+			{
+				MaskInputTextureSample->SetClampTextureEnabled(true);
+
+				if (UDMMaterialStage* MaskTextureInputStage = MaskInputExpression->GetSubStage())
+				{
+					const TArray<UDMMaterialStageInput*> MaskTextureStageInputs = MaskTextureInputStage->GetInputs();
+
+					for (UDMMaterialStageInput* MaskTextureStageInput : MaskTextureStageInputs)
+					{
+						if (UDMMaterialStageInputValue* MaskTextureInputValue = Cast<UDMMaterialStageInputValue>(MaskTextureStageInput))
+						{
+							if (UDMMaterialValueTexture* MaskTextureValue = Cast<UDMMaterialValueTexture>(MaskTextureInputValue->GetValue()))
+							{
+								MaskTextureValue->SetDefaultValue(UDynamicMaterialEditorSettings::Get()->DefaultOpaqueTexture.LoadSynchronous());
+								MaskTextureValue->ApplyDefaultValue();
+							}
+						}
+					}
+				}
+			}
+		}
+	}
+}
+
 UDMMaterialPropertyOpacity::UDMMaterialPropertyOpacity()
 	: UDMMaterialProperty(
 		EDMMaterialPropertyType(EDMMaterialPropertyType::Opacity),
@@ -29,3 +114,8 @@ UMaterialExpression* UDMMaterialPropertyOpacity::GetDefaultInput(const TSharedRe
 {
 	return CreateConstant(InBuildState, 1.f);
 }
+
+void UDMMaterialPropertyOpacity::OnSlotAdded(UDMMaterialSlot* InSlot)
+{
+	AddDefaultOpacityLayer(EDMMaterialPropertyType::Opacity, InSlot);
+}
diff --git a/Engine/Plugins/Experimental/DynamicMaterial/Source/DynamicMaterialEditor/Private/Components/MaterialProperties/DMMPOpacityMask.cpp b/Engine/Plugins/Experimental/DynamicMaterial/Source/DynamicMaterialEditor/Private/Components/MaterialProperties/DMMPOpacityMask.cpp
--- a/Engine/Plugins/Experimental/DynamicMaterial/Source/DynamicMaterialEditor/Private/Components/MaterialProperties/DMMPOpacityMask.cpp
+++ b/Engine/Plugins/Experimental/DynamicMaterial/Source/DynamicMaterialEditor/Private/Components/MaterialProperties/DMMPOpacityMask.cpp
@@ -1,6 +1,8 @@
 // Copyright Epic Games, Inc. All Rights Reserved.
 
 #include "Components/MaterialProperties/DMMPOpacityMask.h"
+
+#include "Components/MaterialProperties/DMMPOpacity.h"
 #include "Model/DynamicMaterialModelEditorOnlyData.h"
 
 UDMMaterialPropertyOpacityMask::UDMMaterialPropertyOpacityMask()
@@ -20,3 +22,8 @@ UMaterialExpression* UDMMaterialPropertyOpacityMask::GetDefaultInput(
 {
 	return CreateConstant(InBuildState, 1.f);
 }
+
+void UDMMaterialPropertyOpacityMask::OnSlotAdded(UDMMaterialSlot* InSlot)
+{
+	UDMMaterialPropertyOpacity::AddDefaultOpacityLayer(EDMMaterialPropertyType::OpacityMask, InSlot);
+}
diff --git a/Engine/Plugins/Experimental/DynamicMaterial/Source/DynamicMaterialEditor/Private/DynamicMaterialEditorSettings.cpp b/Engine/Plugins/Experimental/DynamicMaterial/Source/DynamicMaterialEditor/Private/DynamicMaterialEditorSettings.cpp
--- a/Engine/Plugins/Experimental/DynamicMaterial/Source/DynamicMaterialEditor/Private/DynamicMaterialEditorSettings.cpp
+++ b/Engine/Plugins/Experimental/DynamicMaterial/Source/DynamicMaterialEditor/Private/DynamicMaterialEditorSettings.cpp
@@ -14,6 +14,53 @@
 
 #define LOCTEXT_NAMESPACE "MaterialDesignerSettings"
 
+bool FDMMaterialChannelListPreset::IsPropertyEnabled(EDMMaterialPropertyType InProperty) const
+{
+	switch (InProperty)
+	{
+		case EDMMaterialPropertyType::BaseColor:
+		case EDMMaterialPropertyType::EmissiveColor:
+			return bRGB;
+
+		case EDMMaterialPropertyType::Opacity:
+		case EDMMaterialPropertyType::OpacityMask:
+			return bOpacity;
+
+		case EDMMaterialPropertyType::Roughness:
+			return bRoughness;
+
+		case EDMMaterialPropertyType::Specular:
+			return bSpecular;
+
+		case EDMMaterialPropertyType::Metallic:
+			return bMetallic;
+
+		case EDMMaterialPropertyType::Normal:
+			return bNormal;
+
+		case EDMMaterialPropertyType::PixelDepthOffset:
+			return bPixelDepthOffset;
+
+		case EDMMaterialPropertyType::WorldPositionOffset:
+			return bWorldPositionOffset;
+
+		case EDMMaterialPropertyType::AmbientOcclusion:
+			return bAmbientOcclusion;
+
+		case EDMMaterialPropertyType::Anisotropy:
+			return bAnisotropy;
+
+		case EDMMaterialPropertyType::Refraction:
+			return bRefraction;
+
+		case EDMMaterialPropertyType::Tangent:
+			return bTangent;
+
+		default:
+			return false;
+	}
+}
+
 UDynamicMaterialEditorSettings::UDynamicMaterialEditorSettings()
 {
 	CategoryName = TEXT("Plugins");
@@ -27,6 +74,42 @@ UDynamicMaterialEditorSettings::UDynamicMaterialEditorSettings()
 	DefaultOpaqueTexture = TSoftObjectPtr<UTexture>(FSoftObjectPath(TEXT("Texture2D'/Engine/EngineResources/WhiteSquareTexture.WhiteSquareTexture'")));
 	DefaultPreviewCanvasTexture = nullptr;
 	DefaultOpacitySlotMask = TSoftObjectPtr<UTexture>(FSoftObjectPath(TEXT("/Script/Engine.Texture2D'/DynamicMaterial/Textures/T_DM_HorizontalGradient.T_DM_HorizontalGradient'")));
+
+	FDMMaterialChannelListPreset Opaque;
+	Opaque.bRGB = true;
+	Opaque.bOpacity = false;
+
+	FDMMaterialChannelListPreset Translucent;
+	Translucent.bRGB = true;
+	Translucent.bOpacity = true;
+
+	FDMMaterialChannelListPreset PBR;
+	PBR.bRGB = true;
+	PBR.bOpacity = true;
+	PBR.bMetallic = true;
+	PBR.bSpecular = true;
+	PBR.bRoughness = true;
+	PBR.bNormal = true;
+	PBR.bAmbientOcclusion = true;
+
+	FDMMaterialChannelListPreset All;
+	All.bRGB = true;
+	All.bOpacity = true;
+	All.bMetallic = true;
+	All.bSpecular = true;
+	All.bRoughness = true;
+	All.bNormal = true;
+	All.bAmbientOcclusion = true;
+	All.bAnisotropy = true;
+	All.bPixelDepthOffset = true;
+	All.bRefraction = true;
+	All.bTangent = true;
+	All.bWorldPositionOffset = true;
+
+	ChannelPresets.Add(TEXT("Opaque"), Opaque);
+	ChannelPresets.Add(TEXT("Translucent"), Translucent);
+	ChannelPresets.Add(TEXT("PBR"), PBR);
+	ChannelPresets.Add(TEXT("All"), All);		
 }
 
 UDynamicMaterialEditorSettings* UDynamicMaterialEditorSettings::Get()
diff --git a/Engine/Plugins/Experimental/DynamicMaterial/Source/DynamicMaterialEditor/Private/DynamicMaterialEditorSettings.h b/Engine/Plugins/Experimental/DynamicMaterial/Source/DynamicMaterialEditor/Private/DynamicMaterialEditorSettings.h
--- a/Engine/Plugins/Experimental/DynamicMaterial/Source/DynamicMaterialEditor/Private/DynamicMaterialEditorSettings.h
+++ b/Engine/Plugins/Experimental/DynamicMaterial/Source/DynamicMaterialEditor/Private/DynamicMaterialEditorSettings.h
@@ -2,6 +2,7 @@
 
 #pragma once
 
+#include "DMDefs.h"
 #include "Engine/DeveloperSettings.h"
 #include "Containers/Array.h"
 #include "Engine/Texture.h"
@@ -11,8 +12,9 @@
 
 class FModifierKeysState;
 class UMaterialFunctionInterface;
-struct FPropertyChangedEvent;
 enum EOrientation : int;
+struct FDMMaterialChannelListPreset;
+struct FPropertyChangedEvent;
 
 USTRUCT(BlueprintType)
 struct FDMMaterialEffectList
@@ -26,6 +28,51 @@ struct FDMMaterialEffectList
 	TArray<TSoftObjectPtr<UMaterialFunctionInterface>> Effects;
 };
 
+USTRUCT(BlueprintType)
+struct FDMMaterialChannelListPreset
+{
+	GENERATED_BODY()
+
+	/** This cannot be disabled. It is here for illustrative purposes only. */
+	UPROPERTY(Config, VisibleAnywhere, BlueprintReadOnly, Category = "Channel")
+	bool bRGB = true;
+
+	UPROPERTY(Config, EditAnywhere, BlueprintReadWrite, Category = "Channel")
+	bool bOpacity = false;
+
+	UPROPERTY(Config, EditAnywhere, BlueprintReadWrite, Category = "Channel")
+	bool bRoughness = false;
+
+	UPROPERTY(Config, EditAnywhere, BlueprintReadWrite, Category = "Channel")
+	bool bSpecular = false;
+
+	UPROPERTY(Config, EditAnywhere, BlueprintReadWrite, Category = "Channel")
+	bool bMetallic = false;
+
+	UPROPERTY(Config, EditAnywhere, BlueprintReadWrite, Category = "Channel")
+	bool bNormal = false;
+
+	UPROPERTY(Config, EditAnywhere, BlueprintReadWrite, Category = "Channel")
+	bool bPixelDepthOffset = false;
+
+	UPROPERTY(Config, EditAnywhere, BlueprintReadWrite, Category = "Channel")
+	bool bWorldPositionOffset = false;
+
+	UPROPERTY(Config, EditAnywhere, BlueprintReadWrite, Category = "Channel")
+	bool bAmbientOcclusion = false;
+
+	UPROPERTY(Config, EditAnywhere, BlueprintReadWrite, Category = "Channel")
+	bool bAnisotropy = false;
+
+	UPROPERTY(Config, EditAnywhere, BlueprintReadWrite, Category = "Channel")
+	bool bRefraction = false;
+
+	UPROPERTY(Config, EditAnywhere, BlueprintReadWrite, Category = "Channel")
+	bool bTangent = false;
+
+	bool IsPropertyEnabled(EDMMaterialPropertyType InProperty) const;
+};
+
 /**
  * Material Designer Settings
  */
@@ -138,6 +185,9 @@ public:
 	UPROPERTY(Config, EditAnywhere, BlueprintReadWrite, Category = "Effects")
 	TArray<FName> CustomEffectsFolders;
 
+	UPROPERTY(Config, EditAnywhere, BlueprintReadWrite, Category = "Channels")
+	TMap<FName, FDMMaterialChannelListPreset> ChannelPresets;
+
 	FOnFinishedChangingProperties OnSettingsChanged;
 
 	//~ Begin UObject
diff --git a/Engine/Plugins/Experimental/DynamicMaterial/Source/DynamicMaterialEditor/Private/Menus/DMToolBarMenus.cpp b/Engine/Plugins/Experimental/DynamicMaterial/Source/DynamicMaterialEditor/Private/Menus/DMToolBarMenus.cpp
--- a/Engine/Plugins/Experimental/DynamicMaterial/Source/DynamicMaterialEditor/Private/Menus/DMToolBarMenus.cpp
+++ b/Engine/Plugins/Experimental/DynamicMaterial/Source/DynamicMaterialEditor/Private/Menus/DMToolBarMenus.cpp
@@ -36,6 +36,72 @@ namespace UE::DynamicMaterialEditor::Private
 	static FName ToolBarMaterialExportSectionName = TEXT("MaterialExport");
 	static FName ToolBarMaterialDesignerSettingsSectionName = TEXT("MaterialDesignerSettings");
 
+	void AddToolBarMaterialInstanceChannelListMenu_Execute(const FToolMenuContext& InContext, FName InPresetName)
+	{
+		if (UDMMenuContext* Context = InContext.FindContext<UDMMenuContext>())
+		{
+			if (UDynamicMaterialModel* MaterialModel = Context->GetModel())
+			{
+				if (UDynamicMaterialModelEditorOnlyData* EditorOnlyData = UDynamicMaterialModelEditorOnlyData::Get(MaterialModel))
+				{
+					TSharedPtr<SDMEditor> EditorWidget;
+
+					if (TSharedPtr<SDMSlot> SlotWidget = Context->GetSlotWidget().Pin())
+					{
+						EditorWidget = SlotWidget->GetEditorWidget();
+
+					}
+
+					EditorOnlyData->SetChannelListPreset(InPresetName);
+
+					if (EditorWidget.IsValid())
+					{
+						EditorWidget->RefreshSlotPickerList();
+					}
+				}
+			}
+		}
+	}
+
+	ECheckBoxState AddToolBarMaterialInstanceChannelListMenu_GetCheckState(const FToolMenuContext& InContext, FName InPresetName)
+	{
+		if (UDMMenuContext* Context = InContext.FindContext<UDMMenuContext>())
+		{
+			if (UDynamicMaterialModel* MaterialModel = Context->GetModel())
+			{
+				if (UDynamicMaterialModelEditorOnlyData* EditorOnlyData = UDynamicMaterialModelEditorOnlyData::Get(MaterialModel))
+				{
+					return InPresetName == EditorOnlyData->GetChannelListPreset()
+						? ECheckBoxState::Checked
+						: ECheckBoxState::Unchecked;
+				}
+			}
+		}
+
+		return ECheckBoxState::Undetermined;
+	}
+
+	void AddToolBarMaterialInstanceChannelListMenu(UToolMenu* InMenu)
+	{
+		FToolMenuSection& NewSection = InMenu->AddSection("MaterialType", LOCTEXT("MaterialType", "Material Type"));
+
+		for (const TPair<FName, FDMMaterialChannelListPreset>& Preset : GetDefault<UDynamicMaterialEditorSettings>()->ChannelPresets)
+		{
+			FToolUIAction MaterialChannelListType;
+			MaterialChannelListType.ExecuteAction = FToolMenuExecuteAction::CreateStatic(&AddToolBarMaterialInstanceChannelListMenu_Execute, Preset.Key);
+			MaterialChannelListType.GetActionCheckState = FToolMenuGetActionCheckState::CreateStatic(&AddToolBarMaterialInstanceChannelListMenu_GetCheckState, Preset.Key);
+
+			NewSection.AddMenuEntry(
+				Preset.Key,
+				FText::FromName(Preset.Key),
+				FText::GetEmpty(),
+				TAttribute<FSlateIcon>(),
+				FToolUIActionChoice(MaterialChannelListType),
+				EUserInterfaceActionType::RadioButton
+			);
+		}
+	}
+
 	void AddToolBarMaterialInstanceDomainMenu_Execute(const FToolMenuContext& InContext, EMaterialDomain InDomain)
 	{
 		if (UDMMenuContext* Context = InContext.FindContext<UDMMenuContext>())
@@ -44,7 +110,20 @@ namespace UE::DynamicMaterialEditor::Private
 			{
 				if (UDynamicMaterialModelEditorOnlyData* EditorOnlyData = UDynamicMaterialModelEditorOnlyData::Get(MaterialModel))
 				{
+					TSharedPtr<SDMEditor> EditorWidget;
+
+					if (TSharedPtr<SDMSlot> SlotWidget = Context->GetSlotWidget().Pin())
+					{
+						EditorWidget = SlotWidget->GetEditorWidget();
+
+					}
+
 					EditorOnlyData->SetDomain(InDomain);
+
+					if (EditorWidget.IsValid())
+					{
+						EditorWidget->RefreshSlotPickerList();
+					}
 				}
 			}
 		}
@@ -99,7 +178,7 @@ namespace UE::DynamicMaterialEditor::Private
 				DomainEnum->GetDisplayNameTextByValue(Domain),
 				FText::GetEmpty(),
 				TAttribute<FSlateIcon>(),
-				FToolUIActionChoice(	DomainAction),
+				FToolUIActionChoice(DomainAction),
 				EUserInterfaceActionType::RadioButton
 			);
 		}
@@ -113,7 +192,20 @@ namespace UE::DynamicMaterialEditor::Private
 			{
 				if (UDynamicMaterialModelEditorOnlyData* EditorOnlyData = UDynamicMaterialModelEditorOnlyData::Get(MaterialModel))
 				{
+					TSharedPtr<SDMEditor> EditorWidget;
+
+					if (TSharedPtr<SDMSlot> SlotWidget = Context->GetSlotWidget().Pin())
+					{
+						EditorWidget = SlotWidget->GetEditorWidget();
+
+					}
+
 					EditorOnlyData->SetBlendMode(InBlendMode);
+
+					if (EditorWidget.IsValid())
+					{
+						EditorWidget->RefreshSlotPickerList();
+					}
 				}
 			}
 		}
@@ -182,6 +274,14 @@ namespace UE::DynamicMaterialEditor::Private
 			{
 				if (UDynamicMaterialModelEditorOnlyData* EditorOnlyData = UDynamicMaterialModelEditorOnlyData::Get(MaterialModel))
 				{
+					TSharedPtr<SDMEditor> EditorWidget;
+
+					if (TSharedPtr<SDMSlot> SlotWidget = Context->GetSlotWidget().Pin())
+					{
+						EditorWidget = SlotWidget->GetEditorWidget();
+
+					}
+
 					switch (EditorOnlyData->GetShadingModel())
 					{
 						case EDMMaterialShadingModel::DefaultLit:
@@ -192,6 +292,11 @@ namespace UE::DynamicMaterialEditor::Private
 							EditorOnlyData->SetShadingModel(EDMMaterialShadingModel::DefaultLit);
 							break;
 					}
+
+					if (EditorWidget.IsValid())
+					{
+						EditorWidget->RefreshSlotPickerList();
+					}
 				}
 			}
 		}
@@ -241,7 +346,20 @@ namespace UE::DynamicMaterialEditor::Private
 			{
 				if (UDynamicMaterialModelEditorOnlyData* EditorOnlyData = UDynamicMaterialModelEditorOnlyData::Get(MaterialModel))
 				{
+					TSharedPtr<SDMEditor> EditorWidget;
+
+					if (TSharedPtr<SDMSlot> SlotWidget = Context->GetSlotWidget().Pin())
+					{
+						EditorWidget = SlotWidget->GetEditorWidget();
+
+					}
+
 					EditorOnlyData->SetPixelAnimationFlag(!EditorOnlyData->IsPixelAnimationFlagSet());
+
+					if (EditorWidget.IsValid())
+					{
+						EditorWidget->RefreshSlotPickerList();
+					}					
 				}
 			}
 		}
@@ -286,7 +404,20 @@ namespace UE::DynamicMaterialEditor::Private
 			{
 				if (UDynamicMaterialModelEditorOnlyData* EditorOnlyData = UDynamicMaterialModelEditorOnlyData::Get(MaterialModel))
 				{
+					TSharedPtr<SDMEditor> EditorWidget;
+
+					if (TSharedPtr<SDMSlot> SlotWidget = Context->GetSlotWidget().Pin())
+					{
+						EditorWidget = SlotWidget->GetEditorWidget();
+
+					}
+
 					EditorOnlyData->SetTwoSidedFlag(!EditorOnlyData->IsTwoSidedFlagSet());
+
+					if (EditorWidget.IsValid())
+					{
+						EditorWidget->RefreshSlotPickerList();
+					}					
 				}
 			}
 		}
@@ -332,6 +463,13 @@ namespace UE::DynamicMaterialEditor::Private
 
 		FToolMenuSection& NewSection = InMenu->AddSection(ToolBarMaterialInstanceSectionName, LOCTEXT("MaterialInstanceSection", "Material Instance"));
 
+		NewSection.AddSubMenu(
+			"ChannelList",
+			LOCTEXT("MaterialInstanceMaterialType", "Material Type"),
+			LOCTEXT("MaterialInstanceMaterialType_ToolTip", "Set the material type (channel list) for the active material designer instance."),
+			FNewToolMenuDelegate::CreateStatic(&AddToolBarMaterialInstanceChannelListMenu)
+		);
+
 		NewSection.AddSubMenu(
 			"Domain",
 			LOCTEXT("MaterialInstanceDomain", "Material Domain"),
@@ -435,7 +573,7 @@ namespace UE::DynamicMaterialEditor::Private
 
 			if (!GeneratedMaterial)
 			{
-				UE::DynamicMaterialEditor::Private::LogError(TEXT("Failed to find a generated material to export."));
+				UE::DynamicMaterialEditor::Private::LogError(TEXT("Failed to find a generated material to export."), true, MaterialModel);
 				return;
 			}
 
@@ -863,7 +1001,7 @@ TSharedRef<SWidget> FDMToolBarMenus::MakeEditorLayoutMenu(const TSharedPtr<SDMEd
 	TSharedPtr<SDMSlot> SlotWidget;
 	if (InEditorWidget.IsValid())
 	{
-		TArray<TSharedPtr<SDMSlot>> SlotWidgets = InEditorWidget->GetSlotWidgets();
+		const TArray<TSharedRef<SDMSlot>>& SlotWidgets = InEditorWidget->GetSlotWidgets();
 
 		if (SlotWidgets.IsEmpty() == false)
 		{
diff --git a/Engine/Plugins/Experimental/DynamicMaterial/Source/DynamicMaterialEditor/Private/Model/DynamicMateriaModelFactory.cpp b/Engine/Plugins/Experimental/DynamicMaterial/Source/DynamicMaterialEditor/Private/Model/DynamicMateriaModelFactory.cpp
--- a/Engine/Plugins/Experimental/DynamicMaterial/Source/DynamicMaterialEditor/Private/Model/DynamicMateriaModelFactory.cpp
+++ b/Engine/Plugins/Experimental/DynamicMaterial/Source/DynamicMaterialEditor/Private/Model/DynamicMateriaModelFactory.cpp
@@ -39,7 +39,7 @@ UObject* UDynamicMaterialModelFactory::FactoryCreateNew(UClass* Class, UObject*
 	ModelEditorOnlyData->MaterialModel = NewModel;
 
 	const FDMInitializationGuard InitGuard;
-	ModelEditorOnlyData->AddSlot();
+	ModelEditorOnlyData->Initialize();
 
 	FDynamicMaterialEditorModule::Get().OnMaterialModelCreated(NewModel);
 
diff --git a/Engine/Plugins/Experimental/DynamicMaterial/Source/DynamicMaterialEditor/Private/Model/DynamicMaterialModelEditorOnlyData.cpp b/Engine/Plugins/Experimental/DynamicMaterial/Source/DynamicMaterialEditor/Private/Model/DynamicMaterialModelEditorOnlyData.cpp
--- a/Engine/Plugins/Experimental/DynamicMaterial/Source/DynamicMaterialEditor/Private/Model/DynamicMaterialModelEditorOnlyData.cpp
+++ b/Engine/Plugins/Experimental/DynamicMaterial/Source/DynamicMaterialEditor/Private/Model/DynamicMaterialModelEditorOnlyData.cpp
@@ -4,13 +4,10 @@
 #include "AssetRegistry/AssetRegistryModule.h"
 #include "AssetToolsModule.h"
 #include "Components/DMMaterialLayer.h"
-#include "Components/DMMaterialParameter.h"
 #include "Components/DMMaterialProperty.h"
 #include "Components/DMMaterialSlot.h"
 #include "Components/DMMaterialStage.h"
-#include "Components/DMMaterialStageBlend.h"
 #include "Components/DMMaterialStageThroughputLayerBlend.h"
-#include "Components/DMMaterialSubStage.h"
 #include "Components/DMMaterialValue.h"
 #include "Components/DMTextureUV.h"
 #include "Components/MaterialProperties/DMMPAmbientOcclusion.h"
@@ -28,11 +25,7 @@
 #include "Components/MaterialProperties/DMMPTangent.h"
 #include "Components/MaterialProperties/DMMPWorldPositionOffset.h"
 #include "Components/MaterialStageBlends/DMMSBNormal.h"
-#include "Components/MaterialStageExpressions/DMMSETextureSample.h"
-#include "Components/MaterialStageInputs/DMMSIExpression.h"
-#include "Components/MaterialStageInputs/DMMSIValue.h"
 #include "Components/MaterialValues/DMMaterialValueFloat1.h"
-#include "Components/MaterialValues/DMMaterialValueTexture.h"
 #include "CoreGlobals.h"
 #include "DMComponentPath.h"
 #include "DMDefs.h"
@@ -46,7 +39,6 @@
 #include "IAssetTools.h"
 #include "Material/DynamicMaterialInstance.h"
 #include "Materials/Material.h"
-#include "Materials/MaterialExpressionAppendVector.h"
 #include "Materials/MaterialExpressionMaterialFunctionCall.h"
 #include "Materials/MaterialExpressionMax.h"
 #include "Materials/MaterialExpressionMultiply.h"
@@ -59,10 +51,46 @@
 
 #define LOCTEXT_NAMESPACE "MaterialDesignerModel"
 
-const FString UDynamicMaterialModelEditorOnlyData::SlotsPathToken       = FString(TEXT("Slots"));
-const FString UDynamicMaterialModelEditorOnlyData::RGBSlotPathToken     = FString(TEXT("RGBSlot"));
-const FString UDynamicMaterialModelEditorOnlyData::OpacitySlotPathToken = FString(TEXT("OpacitySlot"));
-const FString UDynamicMaterialModelEditorOnlyData::PropertiesPathToken  = FString(TEXT("Properties"));
+const FString UDynamicMaterialModelEditorOnlyData::SlotsPathToken               = FString(TEXT("Slots"));
+const FString UDynamicMaterialModelEditorOnlyData::RGBSlotPathToken             = FString(TEXT("RGBSlot"));
+const FString UDynamicMaterialModelEditorOnlyData::OpacitySlotPathToken         = FString(TEXT("OpacitySlot"));
+const FString UDynamicMaterialModelEditorOnlyData::RoughnessPathToken           = FString(TEXT("RoughnessSlot"));
+const FString UDynamicMaterialModelEditorOnlyData::SpecularPathToken            = FString(TEXT("SpecularSlot"));
+const FString UDynamicMaterialModelEditorOnlyData::MetallicPathToken            = FString(TEXT("MetallicSlot"));
+const FString UDynamicMaterialModelEditorOnlyData::NormalPathToken              = FString(TEXT("NormalSlot"));
+const FString UDynamicMaterialModelEditorOnlyData::PixelDepthOffsetPathToken    = FString(TEXT("PixelDepthOffsetSlot"));
+const FString UDynamicMaterialModelEditorOnlyData::WorldPositionOffsetPathToken = FString(TEXT("WorldPositionOffsetSlot"));
+const FString UDynamicMaterialModelEditorOnlyData::AmbientOcclusionPathToken    = FString(TEXT("AmbientOcclusionSlot"));
+const FString UDynamicMaterialModelEditorOnlyData::AnisotropyPathToken          = FString(TEXT("AnisotropySlot"));
+const FString UDynamicMaterialModelEditorOnlyData::RefractionPathToken          = FString(TEXT("RefractionSlot"));
+const FString UDynamicMaterialModelEditorOnlyData::TangentPathToken             = FString(TEXT("TangentSlot"));
+const FString UDynamicMaterialModelEditorOnlyData::Custom1PathToken             = FString(TEXT("Custom1Slot"));
+const FString UDynamicMaterialModelEditorOnlyData::Custom2PathToken             = FString(TEXT("Custom2Slot"));
+const FString UDynamicMaterialModelEditorOnlyData::Custom3PathToken             = FString(TEXT("Custom3Slot"));
+const FString UDynamicMaterialModelEditorOnlyData::Custom4PathToken             = FString(TEXT("Custom4Slot"));
+const FString UDynamicMaterialModelEditorOnlyData::PropertiesPathToken          = FString(TEXT("Properties"));
+
+namespace UE::DynamicMaterialEditor::Private
+{
+	const TMap<FString, EDMMaterialPropertyType> TokenToPropertyMap = {
+		{UDynamicMaterialModelEditorOnlyData::RGBSlotPathToken,             EDMMaterialPropertyType::BaseColor},
+		{UDynamicMaterialModelEditorOnlyData::OpacitySlotPathToken,         EDMMaterialPropertyType::Opacity},
+		{UDynamicMaterialModelEditorOnlyData::RoughnessPathToken,           EDMMaterialPropertyType::Roughness},
+		{UDynamicMaterialModelEditorOnlyData::SpecularPathToken,            EDMMaterialPropertyType::Specular},
+		{UDynamicMaterialModelEditorOnlyData::MetallicPathToken,            EDMMaterialPropertyType::Metallic},
+		{UDynamicMaterialModelEditorOnlyData::NormalPathToken,              EDMMaterialPropertyType::Normal},
+		{UDynamicMaterialModelEditorOnlyData::PixelDepthOffsetPathToken,    EDMMaterialPropertyType::PixelDepthOffset},
+		{UDynamicMaterialModelEditorOnlyData::WorldPositionOffsetPathToken, EDMMaterialPropertyType::WorldPositionOffset},
+		{UDynamicMaterialModelEditorOnlyData::AmbientOcclusionPathToken,    EDMMaterialPropertyType::AmbientOcclusion},
+		{UDynamicMaterialModelEditorOnlyData::AnisotropyPathToken,          EDMMaterialPropertyType::Anisotropy},
+		{UDynamicMaterialModelEditorOnlyData::RefractionPathToken,          EDMMaterialPropertyType::Refraction},
+		{UDynamicMaterialModelEditorOnlyData::TangentPathToken,             EDMMaterialPropertyType::Tangent},
+		{UDynamicMaterialModelEditorOnlyData::Custom1PathToken,             EDMMaterialPropertyType::Custom1},
+		{UDynamicMaterialModelEditorOnlyData::Custom2PathToken,             EDMMaterialPropertyType::Custom2},
+		{UDynamicMaterialModelEditorOnlyData::Custom3PathToken,             EDMMaterialPropertyType::Custom3},
+		{UDynamicMaterialModelEditorOnlyData::Custom4PathToken,             EDMMaterialPropertyType::Custom4}
+	};
+}
 
 const TArray<EMaterialDomain> UDynamicMaterialModelEditorOnlyData::SupportedDomains =
 {
@@ -112,8 +140,9 @@ UDynamicMaterialModelEditorOnlyData::UDynamicMaterialModelEditorOnlyData()
 	BlendMode = EBlendMode::BLEND_Translucent;
 	ShadingModel = EDMMaterialShadingModel::Unlit;
 	bPixelAnimationFlag = false;
-	bTwoSidedFlag = false;
+	bTwoSidedFlag = true;
 	bUseWizard = true;
+	ChannelListPreset = NAME_None;
 
 	Properties.Emplace(EDMMaterialPropertyType::BaseColor,           CreateDefaultSubobject<UDMMaterialPropertyBaseColor>(          "MaterialProperty_BaseColor"));
 	Properties.Emplace(EDMMaterialPropertyType::EmissiveColor,       CreateDefaultSubobject<UDMMaterialPropertyEmissiveColor>(      "MaterialProperty_EmissiveColor"));
@@ -141,6 +170,20 @@ UDynamicMaterialModelEditorOnlyData::UDynamicMaterialModelEditorOnlyData()
 	}
 }
 
+void UDynamicMaterialModelEditorOnlyData::Initialize()
+{
+	if (Slots.IsEmpty())
+	{
+		return;
+	}
+
+	// Will choose appropriate type between BaseColor and EmissiveColor
+	AddSlotForMaterialProperty(EDMMaterialPropertyType::BaseColor);
+
+	// Will choose appropriate type between Opacity and OpacityMask
+	AddSlotForMaterialProperty(EDMMaterialPropertyType::Opacity);
+}
+
 UMaterial* UDynamicMaterialModelEditorOnlyData::GetGeneratedMaterial() const
 {
 	return IsValid(MaterialModel) ? MaterialModel->DynamicMaterial : nullptr;
@@ -156,6 +199,7 @@ void UDynamicMaterialModelEditorOnlyData::ResetData()
 	BlendMode = EBlendMode::BLEND_Translucent;
 	ShadingModel = EDMMaterialShadingModel::Unlit;
 	Slots.Empty();
+	PropertySlotMap.Empty();
 	Expressions.Empty();
 }
 
@@ -285,18 +329,20 @@ void UDynamicMaterialModelEditorOnlyData::BuildMaterial(bool bInDirtyAssets)
 		UMaterialExpression* LastPropertyExpression = nullptr;
 		UDMMaterialSlot* Slot = GetSlotForMaterialProperty(Pair.Key);
 
-		if (Slot && !Slot->GetLayers().IsEmpty())
+		if (!Slot || Slot->GetLayers().IsEmpty())
 		{
-			Slot->GenerateExpressions(BuildState);
+			continue;
+		}
 
-			if (BuildState->GetSlotExpressions(Slot).IsEmpty())
-			{
-				continue;
-			}
+		Slot->GenerateExpressions(BuildState);
 
-			LastPropertyExpression = BuildState->GetLastSlotPropertyExpression(Slot, Pair.Key);
+		if (BuildState->GetSlotExpressions(Slot).IsEmpty())
+		{
+			continue;
 		}
 
+		LastPropertyExpression = BuildState->GetLastSlotPropertyExpression(Slot, Pair.Key);
+
 		if (!LastPropertyExpression)
 		{
 			continue;
@@ -634,6 +680,35 @@ void UDynamicMaterialModelEditorOnlyData::OnWizardComplete()
 	bUseWizard = false;
 }
 
+void UDynamicMaterialModelEditorOnlyData::SetChannelListPreset(FName InPresetName)
+{
+	ChannelListPreset = InPresetName;
+
+	if (const FDMMaterialChannelListPreset* Preset = GetDefault<UDynamicMaterialEditorSettings>()->ChannelPresets.Find(InPresetName))
+	{
+		for (uint8 PropertyIndex = static_cast<uint8>(EDMMaterialPropertyType::None) + 1;
+			PropertyIndex < static_cast<uint8>(EDMMaterialPropertyType::Any);
+			++PropertyIndex)
+		{
+			const EDMMaterialPropertyType Property = static_cast<EDMMaterialPropertyType>(PropertyIndex);
+
+			if (Property == EDMMaterialPropertyType::EmissiveColor || Property == EDMMaterialPropertyType::OpacityMask)
+			{
+				continue;
+			}
+
+			if (Preset->IsPropertyEnabled(Property))
+			{
+				AddSlotForMaterialProperty(Property);
+			}
+			else
+			{
+				RemoveSlotForMaterialProperty(Property);
+			}
+		}
+	}
+}
+
 UDMMaterialComponent* UDynamicMaterialModelEditorOnlyData::GetSubComponentByPath(FDMComponentPath& InPath,
 	const FDMComponentPathSegment& InPathSegment) const
 {
@@ -652,28 +727,6 @@ UDMMaterialComponent* UDynamicMaterialModelEditorOnlyData::GetSubComponentByPath
 		return nullptr;
 	}
 
-	// @TODO This will probably need to change when other slot types are opened up.
-	if (InPathSegment.GetToken() == RGBSlotPathToken)
-	{
-		if (Slots.IsValidIndex(0))
-		{
-			return Slots[0]->GetComponentByPath(InPath);
-		}
-
-		return nullptr;
-	}
-
-	// @TODO This will probably need to change when other slot types are opened up.
-	if (InPathSegment.GetToken() == OpacitySlotPathToken)
-	{
-		if (Slots.IsValidIndex(1))
-		{
-			return Slots[1]->GetComponentByPath(InPath);
-		}
-
-		return nullptr;
-	}
-
 	if (InPathSegment.GetToken() == PropertiesPathToken)
 	{
 		FString PropertyStr;
@@ -695,6 +748,17 @@ UDMMaterialComponent* UDynamicMaterialModelEditorOnlyData::GetSubComponentByPath
 		}
 	}
 
+	// Channels
+	using namespace UE::DynamicMaterialEditor::Private;
+
+	if (const EDMMaterialPropertyType* PropertyPtr = TokenToPropertyMap.Find(FString(InPathSegment.GetToken())))
+	{
+		if (const TObjectPtr<UDMMaterialSlot>* SlotPtr = PropertySlotMap.Find(*PropertyPtr))
+		{
+			return (*SlotPtr)->GetComponentByPath(InPath);
+		}
+	}
+
 	return nullptr;
 }
 
@@ -708,6 +772,50 @@ void UDynamicMaterialModelEditorOnlyData::DoBuild_Implementation(bool bInDirtyAs
 	BuildMaterial(bInDirtyAssets);
 }
 
+void UDynamicMaterialModelEditorOnlyData::SwapSlotMaterialProperty(EDMMaterialPropertyType InPropertyFrom, EDMMaterialPropertyType InPropertyTo)
+{
+	UDMMaterialSlot* FromSlot = GetSlotForMaterialProperty(InPropertyFrom);
+	
+	if (!FromSlot)
+	{
+		return;
+	}
+
+	if (UDMMaterialSlot* ToSlot = GetSlotForMaterialProperty(InPropertyTo))
+	{
+		if (ToSlot == FromSlot)
+		{
+			return;
+		}
+
+		RemoveSlotForMaterialProperty(InPropertyFrom);
+	}
+
+	FromSlot->ChangeMaterialProperty(InPropertyFrom, InPropertyTo);
+}
+
+void UDynamicMaterialModelEditorOnlyData::EnsureSwapSlotMaterialProperty(EDMMaterialPropertyType InPropertyFrom, EDMMaterialPropertyType InPropertyTo)
+{
+	if (UDMMaterialSlot* ToSlot = GetSlotForMaterialProperty(InPropertyTo))
+	{
+		if (UDMMaterialSlot* FromSlot = GetSlotForMaterialProperty(InPropertyFrom))
+		{
+			if (ToSlot != FromSlot)
+			{
+				RemoveSlotForMaterialProperty(InPropertyFrom);
+			}
+		}
+	}
+	else if (GetSlotForMaterialProperty(InPropertyFrom))
+	{
+		SwapSlotMaterialProperty(InPropertyFrom, InPropertyTo);
+	}
+	else
+	{
+		AddSlotForMaterialProperty(InPropertyTo);
+	}
+}
+
 void UDynamicMaterialModelEditorOnlyData::SetDomain(TEnumAsByte<EMaterialDomain> InDomain)
 {
 	if (Domain == InDomain)
@@ -719,28 +827,61 @@ void UDynamicMaterialModelEditorOnlyData::SetDomain(TEnumAsByte<EMaterialDomain>
 
 	if (Domain == EMaterialDomain::MD_PostProcess)
 	{
-		// Post process domain only supports a single slot - emissive color. 1 is the opacity slot.
-		if (Slots.IsValidIndex(1))
+		const FDMUpdateGuard Guard;
+
+		// Post process only supports emissive.
+		for (uint8 PropertyIndex = static_cast<uint8>(EDMMaterialPropertyType::None) + 1;
+			PropertyIndex < static_cast<uint8>(EDMMaterialPropertyType::Any);
+			++PropertyIndex)
 		{
-			RemoveSlot(1);
+			const EDMMaterialPropertyType Property = static_cast<EDMMaterialPropertyType>(PropertyIndex);
+
+			switch (Property)
+			{
+				case EDMMaterialPropertyType::BaseColor:
+				case EDMMaterialPropertyType::EmissiveColor:
+					RemoveSlotForMaterialProperty(Property);
+					continue;
+
+				default:
+					// Do nothing
+					break;
+			}			
 		}
 
+		EnsureSwapSlotMaterialProperty(EDMMaterialPropertyType::BaseColor, EDMMaterialPropertyType::EmissiveColor);
+
 		SetShadingModel(EDMMaterialShadingModel::Unlit);
 		SetBlendMode(EBlendMode::BLEND_Opaque);
-
-		// Setting it to unlit will set these to base color. Need to reset them to emissive.
-		if (Slots.IsValidIndex(0))
+	}
+	else if (Domain == EMaterialDomain::MD_DeferredDecal)
+	{
+		// Post process only supports basic types.
+		for (uint8 PropertyIndex = static_cast<uint8>(EDMMaterialPropertyType::None) + 1;
+			PropertyIndex < static_cast<uint8>(EDMMaterialPropertyType::Any);
+			++PropertyIndex)
 		{
-			for (UDMMaterialLayerObject* Layer : Slots[0]->GetLayers())
+			const EDMMaterialPropertyType Property = static_cast<EDMMaterialPropertyType>(PropertyIndex);
+
+			switch (Property)
 			{
-				if (GUndo)
-				{
-					Layer->Modify();
-				}
+				case EDMMaterialPropertyType::BaseColor:
+				case EDMMaterialPropertyType::EmissiveColor:
+				case EDMMaterialPropertyType::Opacity:
+				case EDMMaterialPropertyType::OpacityMask:
+					RemoveSlotForMaterialProperty(Property);
+					continue;
 
-				Layer->SetMaterialProperty(EDMMaterialPropertyType::EmissiveColor);
-			}
+				default:
+					// Do nothing
+					break;
+			}			
 		}
+
+		EnsureSwapSlotMaterialProperty(EDMMaterialPropertyType::EmissiveColor, EDMMaterialPropertyType::BaseColor);
+
+		SetShadingModel(EDMMaterialShadingModel::DefaultLit);
+		SetBlendMode(EBlendMode::BLEND_Translucent);
 	}
 
 	RequestMaterialBuild();
@@ -755,44 +896,23 @@ void UDynamicMaterialModelEditorOnlyData::SetBlendMode(TEnumAsByte<EBlendMode> I
 
 	BlendMode = InBlendMode;
 
-	if (Slots.IsValidIndex(1))
+	switch (InBlendMode)
 	{
-		// 1 is the opacity slot
-		if (InBlendMode == EBlendMode::BLEND_Opaque)
-		{
-			RemoveSlot(1);
-		}
-		else
-		{
-			for (UDMMaterialLayerObject* Layer : Slots[0]->GetLayers())
-			{
-				if (GUndo)
-				{
-					Layer->Modify();
-				}
-
-				switch (InBlendMode)
-				{
-					case EBlendMode::BLEND_Masked:
-						Layer->SetMaterialProperty(EDMMaterialPropertyType::OpacityMask);
-						break;
-
-					default:
-						Layer->SetMaterialProperty(EDMMaterialPropertyType::Opacity);
-						break;
-				}
-			}
-		}
-	}
-
-	switch (BlendMode)
-	{
-		case EBlendMode::BLEND_Masked:
 		case EBlendMode::BLEND_Opaque:
+			SetPixelAnimationFlag(false);
+			RemoveSlotForMaterialProperty(EDMMaterialPropertyType::Opacity);
+			RemoveSlotForMaterialProperty(EDMMaterialPropertyType::OpacityMask);
 			break;
 
-		default:
+		case EBlendMode::BLEND_Masked:
 			SetPixelAnimationFlag(false);
+			EnsureSwapSlotMaterialProperty(EDMMaterialPropertyType::Opacity, EDMMaterialPropertyType::OpacityMask);
+			break;
+
+		case EBlendMode::BLEND_Translucent:
+		case EBlendMode::BLEND_Additive:
+		case EBlendMode::BLEND_Modulate:
+			EnsureSwapSlotMaterialProperty(EDMMaterialPropertyType::OpacityMask, EDMMaterialPropertyType::Opacity);
 			break;
 	}
 
@@ -808,28 +928,27 @@ void UDynamicMaterialModelEditorOnlyData::SetShadingModel(EDMMaterialShadingMode
 
 	ShadingModel = InShadingModel;
 
-	if (Slots.IsValidIndex(0))
+	EDMMaterialPropertyType FromProperty;
+	EDMMaterialPropertyType ToProperty;
+
+	switch (ShadingModel)
 	{
-		for (UDMMaterialLayerObject* Layer : Slots[0]->GetLayers())
-		{
-			if (GUndo)
-			{
-				Layer->Modify();
-			}
+		case EDMMaterialShadingModel::Unlit:
+			FromProperty = EDMMaterialPropertyType::BaseColor;
+			ToProperty = EDMMaterialPropertyType::EmissiveColor;
+			break;
 
-			switch (ShadingModel)
-			{
-				case EDMMaterialShadingModel::Unlit:
-					Layer->SetMaterialProperty(EDMMaterialPropertyType::EmissiveColor);
-					break;
+		case EDMMaterialShadingModel::DefaultLit:
+			FromProperty = EDMMaterialPropertyType::EmissiveColor;
+			ToProperty = EDMMaterialPropertyType::BaseColor;
+			break;
 
-				default:
-					Layer->SetMaterialProperty(EDMMaterialPropertyType::BaseColor);
-					break;
-			}
-		}
+		default:
+			return;
 	}
 
+	EnsureSwapSlotMaterialProperty(FromProperty, ToProperty);
+
 	RequestMaterialBuild();
 }
 
@@ -893,6 +1012,11 @@ void UDynamicMaterialModelEditorOnlyData::SetTwoSidedFlag(bool bInFlagValue)
 	RequestMaterialBuild();
 }
 
+FName UDynamicMaterialModelEditorOnlyData::GetChannelListPreset() const
+{
+	return ChannelListPreset;
+}
+
 void UDynamicMaterialModelEditorOnlyData::OpenMaterialEditor() const
 {
 	if (!IsValid(MaterialModel) || !IsValid(MaterialModel->DynamicMaterial))
@@ -931,157 +1055,125 @@ UDMMaterialProperty* UDynamicMaterialModelEditorOnlyData::GetMaterialProperty(ED
 
 UDMMaterialSlot* UDynamicMaterialModelEditorOnlyData::GetSlot(int32 Index) const
 {
-	if (!Slots.IsValidIndex(Index))
+	if (Slots.IsValidIndex(Index))
 	{
-		return nullptr;
+		return Slots[Index];		
 	}
 
-	return Slots[Index];
+	return nullptr;
+}
+
+UDMMaterialSlot* UDynamicMaterialModelEditorOnlyData::GetSlotForMaterialProperty(EDMMaterialPropertyType InType) const
+{
+	if (const TObjectPtr<UDMMaterialSlot>* SlotPtr = PropertySlotMap.Find(InType))
+	{
+		return *SlotPtr;
+	}
+
+	return nullptr;
 }
 
 UDMMaterialSlot* UDynamicMaterialModelEditorOnlyData::AddSlot()
 {
-	EDMMaterialPropertyType NewSlotProperty = EDMMaterialPropertyType::None;
+	constexpr int64 Start = static_cast<int64>(EDMMaterialPropertyType::None) + 1;
+	constexpr int64 End = static_cast<int64>(EDMMaterialPropertyType::Any);
 
-	if (Slots.IsEmpty())
+	for (int64 PropertyIndex = Start; PropertyIndex < End; ++PropertyIndex)
 	{
-		const FDMUpdateGuard Guard;
+		const EDMMaterialPropertyType Property = static_cast<EDMMaterialPropertyType>(PropertyIndex);
 
-		switch (ShadingModel)
+		if (GetSlotForMaterialProperty(Property))
 		{
-			case EDMMaterialShadingModel::DefaultLit:
-				NewSlotProperty = EDMMaterialPropertyType::BaseColor;
-				break;
+			continue;
+		}
 
-			case EDMMaterialShadingModel::Unlit:
-				NewSlotProperty = EDMMaterialPropertyType::EmissiveColor;
-				break;
+		return AddSlotForMaterialProperty(Property);
+	}
 
-			default:
-				checkNoEntry()
+	return nullptr;
+}
+
+UDMMaterialSlot* UDynamicMaterialModelEditorOnlyData::AddSlotForMaterialProperty(EDMMaterialPropertyType InType)
+{
+	if (InType == EDMMaterialPropertyType::BaseColor || InType == EDMMaterialPropertyType::EmissiveColor)
+	{
+		if (Domain == EMaterialDomain::MD_DeferredDecal)
+		{
+			InType = EDMMaterialPropertyType::EmissiveColor;
+		}
+		else
+		{
+			switch (ShadingModel)
+			{
+				case EDMMaterialShadingModel::DefaultLit:
+					InType = EDMMaterialPropertyType::BaseColor;
+					break;
+
+				case EDMMaterialShadingModel::Unlit:
+					InType = EDMMaterialPropertyType::EmissiveColor;
+					break;
+
+				default:
+					checkNoEntry();
 					break;
+			}
 		}
 	}
 	else if (Domain == EMaterialDomain::MD_PostProcess)
 	{
 		return nullptr;
 	}
-	else
+	else if (InType == EDMMaterialPropertyType::Opacity || InType == EDMMaterialPropertyType::OpacityMask)
 	{
-		for (int64 FirstAvailableEnum = static_cast<int64>(EDMMaterialPropertyType::EmissiveColor) + 1; FirstAvailableEnum < static_cast<int64>(EDMMaterialPropertyType::Any); ++FirstAvailableEnum)
+		switch (BlendMode)
 		{
-			EDMMaterialPropertyType FirstAvailableEnumActual = static_cast<EDMMaterialPropertyType>(FirstAvailableEnum);
+			case EBlendMode::BLEND_Translucent:
+			case EBlendMode::BLEND_Additive:
+			case EBlendMode::BLEND_Modulate:
+				InType = EDMMaterialPropertyType::Opacity;
+				break;
 
-			if (GetSlotForMaterialProperty(FirstAvailableEnumActual) == nullptr)
-			{
-				NewSlotProperty = FirstAvailableEnumActual;
+			case EBlendMode::BLEND_Masked:
+				InType = EDMMaterialPropertyType::OpacityMask;
 				break;
-			}
-		}
 
-		if (NewSlotProperty == EDMMaterialPropertyType::None)
-		{
-			return nullptr;
+			case EBlendMode::BLEND_Opaque:
+				return nullptr;
+
+			default:
+				checkNoEntry();
+				break;
 		}
 	}
 
+	if (Domain == EMaterialDomain::MD_DeferredDecal && InType != EDMMaterialPropertyType::EmissiveColor)
+	{
+		return nullptr;
+	}
+
+	if (UDMMaterialSlot* ExistingSlot = GetSlotForMaterialProperty(InType))
+	{
+		return ExistingSlot;
+	}
+
 	UDMMaterialSlot* NewSlot = NewObject<UDMMaterialSlot>(this, NAME_None, RF_Transactional);
 	check(NewSlot);
 
+	AssignMaterialPropertyToSlot(InType, NewSlot);
+
 	NewSlot->SetIndex(Slots.Num());
 	Slots.Add(NewSlot);
 	NewSlot->SetComponentState(EDMComponentLifetimeState::Added);
 
 	NewSlot->GetOnConnectorsUpdateDelegate().AddUObject(this, &UDynamicMaterialModelEditorOnlyData::OnSlotConnectorsUpdated);
-	NewSlot->Update(EDMUpdateType::Structure);
-
-	UDMMaterialStage* DefaultStage = UDMMaterialStageBlend::CreateStage(UDMMaterialStageBlendNormal::StaticClass());
-	check(DefaultStage);
-
-	DefaultStage->SetBeingEdited(true);
-
-	UDMMaterialStage* MaskStage = UDMMaterialStageThroughputLayerBlend::CreateStage();
-	check(MaskStage);
-
-	NewSlot->AddLayerWithMask(NewSlotProperty, DefaultStage, MaskStage);
-	NewSlot->SetEditingLayers(true);
-
-	AssignMaterialPropertyToSlot(NewSlotProperty, NewSlot);
-
-	UDMMaterialStageInputExpression* BaseInputExpression = UDMMaterialStageInputExpression::ChangeStageInput_Expression(DefaultStage,
-		UDMMaterialStageExpressionTextureSample::StaticClass(), UDMMaterialStageBlendNormal::InputB,
-		FDMMaterialStageConnectorChannel::WHOLE_CHANNEL, 0,
-		FDMMaterialStageConnectorChannel::WHOLE_CHANNEL);
 
-	switch (NewSlotProperty)
+	if (UDMMaterialProperty* Property = GetMaterialProperty(InType))
 	{
-		case EDMMaterialPropertyType::Opacity:
-		case EDMMaterialPropertyType::OpacityMask:
-			if (UDMMaterialStageExpressionTextureSample* BaseInputTextureSample = Cast<UDMMaterialStageExpressionTextureSample>(BaseInputExpression->GetMaterialStageExpression()))
-			{
-				BaseInputTextureSample->SetClampTextureEnabled(true);
-
-				if (UDMMaterialStage* BaseTextureInputStage = BaseInputExpression->GetSubStage())
-				{
-					const TArray<UDMMaterialStageInput*> BaseTextureStageInputs = BaseTextureInputStage->GetInputs();
-
-					for (UDMMaterialStageInput* BaseTextureStageInput : BaseTextureStageInputs)
-					{
-						if (UDMMaterialStageInputValue* BaseTextureInputValue = Cast<UDMMaterialStageInputValue>(BaseTextureStageInput))
-						{
-							if (UDMMaterialValueTexture* BaseTextureValue = Cast<UDMMaterialValueTexture>(BaseTextureInputValue->GetValue()))
-							{
-								BaseTextureValue->SetDefaultValue(UDynamicMaterialEditorSettings::Get()->DefaultOpacitySlotMask.LoadSynchronous());
-								BaseTextureValue->ApplyDefaultValue();
-								break;
-							}
-						}
-					}
-				}
-			}
-
-			{
-				const TArray<UDMMaterialStageInput*> MaskStageInputs = MaskStage->GetInputs();
-
-				for (UDMMaterialStageInput* MaskStageInput : MaskStageInputs)
-				{
-					if (UDMMaterialStageInputExpression* MaskInputExpression = Cast<UDMMaterialStageInputExpression>(MaskStageInput))
-					{
-						if (UDMMaterialStageExpressionTextureSample* MaskInputTextureSample = Cast<UDMMaterialStageExpressionTextureSample>(MaskInputExpression->GetMaterialStageExpression()))
-						{
-							MaskInputTextureSample->SetClampTextureEnabled(true);
-
-							if (UDMMaterialStage* MaskTextureInputStage = MaskInputExpression->GetSubStage())
-							{
-								const TArray<UDMMaterialStageInput*> MaskTextureStageInputs = MaskTextureInputStage->GetInputs();
-
-								for (UDMMaterialStageInput* MaskTextureStageInput : MaskTextureStageInputs)
-								{
-									if (UDMMaterialStageInputValue* MaskTextureInputValue = Cast<UDMMaterialStageInputValue>(MaskTextureStageInput))
-									{
-										if (UDMMaterialValueTexture* MaskTextureValue = Cast<UDMMaterialValueTexture>(MaskTextureInputValue->GetValue()))
-										{
-											MaskTextureValue->SetDefaultValue(UDynamicMaterialEditorSettings::Get()->DefaultOpaqueTexture.LoadSynchronous());
-											MaskTextureValue->ApplyDefaultValue();
-											break;
-										}
-									}
-								}
-							}
-						}
-
-						break;
-					}
-				}
-			}
-
-			break;
-
-		default:
-			// Do nothing
-			break;
+		Property->OnSlotAdded(NewSlot);
 	}
 
+	NewSlot->Update(EDMUpdateType::Structure);
+
 	OnSlotListUpdateDelegate.Broadcast(MaterialModel);
 
 	RequestMaterialBuild();
@@ -1089,19 +1181,21 @@ UDMMaterialSlot* UDynamicMaterialModelEditorOnlyData::AddSlot()
 	return NewSlot;
 }
 
-void UDynamicMaterialModelEditorOnlyData::RemoveSlot(int32 Index)
+UDMMaterialSlot* UDynamicMaterialModelEditorOnlyData::RemoveSlot(int32 Index)
 {
 	UDMMaterialSlot* Slot = GetSlot(Index);
-	check(Slot);
+	
+	if (!Slot)
+	{
+		return nullptr;
+	}
 	
 	if (GUndo)
 	{
 		Slot->Modify(/* Always mark dirty */ false);
 	}
 
-	TArray<EDMMaterialPropertyType> SlotProperties = GetMaterialPropertiesForSlot(Slot);
-
-	for (EDMMaterialPropertyType MaterialProperty : SlotProperties)
+	for (EDMMaterialPropertyType MaterialProperty : GetMaterialPropertiesForSlot(Slot))
 	{
 		UnassignMaterialProperty(MaterialProperty);
 	}
@@ -1113,25 +1207,41 @@ void UDynamicMaterialModelEditorOnlyData::RemoveSlot(int32 Index)
 		PropertySlotMap.Remove(*Key);
 	}
 
-	Slots.RemoveAt(Index);
+	const int32 SlotIndex = Slots.IndexOfByKey(Slot);
+
+	if (SlotIndex != INDEX_NONE)
+	{
+		Slots.RemoveAtSwap(SlotIndex);
+		
+		if (Slots.IsValidIndex(SlotIndex))
+		{
+			if (GUndo)
+			{
+				Slots[SlotIndex]->Modify(/* Always mark dirty */ false);
+			}
+
+			Slots[SlotIndex]->SetIndex(SlotIndex);
+		}
+	}
+
 	Slot->GetOnConnectorsUpdateDelegate().RemoveAll(this);
 	Slot->SetComponentState(EDMComponentLifetimeState::Removed);
 
 	RequestMaterialBuild();
 
 	OnSlotListUpdateDelegate.Broadcast(MaterialModel);
+	
+	return Slot;
 }
 
-UDMMaterialSlot* UDynamicMaterialModelEditorOnlyData::GetSlotForMaterialProperty(EDMMaterialPropertyType Property) const
+UDMMaterialSlot* UDynamicMaterialModelEditorOnlyData::RemoveSlotForMaterialProperty(EDMMaterialPropertyType InType)
 {
-	TObjectPtr<UDMMaterialSlot> const* SlotPtr = PropertySlotMap.Find(Property);
-
-	if (!SlotPtr)
+	if (const TObjectPtr<UDMMaterialSlot>* SlotPtr = PropertySlotMap.Find(InType))
 	{
-		return nullptr;
+		return RemoveSlot(Slots.IndexOfByKey(*SlotPtr));
 	}
 
-	return *SlotPtr;
+	return nullptr;
 }
 
 TArray<EDMMaterialPropertyType> UDynamicMaterialModelEditorOnlyData::GetMaterialPropertiesForSlot(UDMMaterialSlot* Slot) const
@@ -1169,7 +1279,11 @@ void UDynamicMaterialModelEditorOnlyData::AssignMaterialPropertyToSlot(EDMMateri
 void UDynamicMaterialModelEditorOnlyData::UnassignMaterialProperty(EDMMaterialPropertyType Property)
 {
 	TObjectPtr<UDMMaterialSlot>* SlotPtr = PropertySlotMap.Find(Property);
-	check(SlotPtr);
+	
+	if (!SlotPtr)
+	{
+		return;
+	}
 
 	PropertySlotMap.Remove(Property);
 	(*SlotPtr)->GetOnPropertiesUpdateDelegate().Broadcast(*SlotPtr);
diff --git a/Engine/Plugins/Experimental/DynamicMaterial/Source/DynamicMaterialEditor/Private/Slate/Layers/SDMSlotLayerView.cpp b/Engine/Plugins/Experimental/DynamicMaterial/Source/DynamicMaterialEditor/Private/Slate/Layers/SDMSlotLayerView.cpp
--- a/Engine/Plugins/Experimental/DynamicMaterial/Source/DynamicMaterialEditor/Private/Slate/Layers/SDMSlotLayerView.cpp
+++ b/Engine/Plugins/Experimental/DynamicMaterial/Source/DynamicMaterialEditor/Private/Slate/Layers/SDMSlotLayerView.cpp
@@ -366,7 +366,6 @@ void SDMSlotLayerView::OnLayerItemSelectionChanged(TSharedPtr<FDMMaterialLayerRe
 				if (UDMMaterialStage* Stage = BaseWidget->GetStage())
 				{
 					Stage->Modify();
-					Stage->SetBeingEdited(false);
 
 					if (InSelectInfo == ESelectInfo::OnMouseClick && SlotWidget.IsValid())
 					{
@@ -380,8 +379,6 @@ void SDMSlotLayerView::OnLayerItemSelectionChanged(TSharedPtr<FDMMaterialLayerRe
 			{
 				if (UDMMaterialStage* Stage = MaskWidget->GetStage())
 				{
-					Stage->SetBeingEdited(false);
-
 					if (InSelectInfo == ESelectInfo::OnMouseClick && SlotWidget.IsValid())
 					{
 						SlotWidget->InvalidateSlotSettingsRowWidget();
@@ -841,7 +838,6 @@ void SDMSlotLayerView::Tick(const FGeometry& InAllottedGeometry, const double In
 			if (UDMMaterialStage* SelectedStage = SelectedStageWeak.Get())
 			{
 				SelectedStage->Modify();
-				SelectedStage->SetBeingEdited(true);
 			}
 		}
 
diff --git a/Engine/Plugins/Experimental/DynamicMaterial/Source/DynamicMaterialEditor/Private/Slate/SDMComponentEdit.cpp b/Engine/Plugins/Experimental/DynamicMaterial/Source/DynamicMaterialEditor/Private/Slate/SDMComponentEdit.cpp
--- a/Engine/Plugins/Experimental/DynamicMaterial/Source/DynamicMaterialEditor/Private/Slate/SDMComponentEdit.cpp
+++ b/Engine/Plugins/Experimental/DynamicMaterial/Source/DynamicMaterialEditor/Private/Slate/SDMComponentEdit.cpp
@@ -232,6 +232,11 @@ SDMComponentEdit::~SDMComponentEdit()
 
 	SDMEditor::ClearPropertyHandles(this);
 
+	if (UDMMaterialStage* Stage = Cast<UDMMaterialStage>(ComponentWeak.Get()))
+	{
+		Stage->SetBeingEdited(false);
+	}
+
 	if (UDMMaterialStage* Stage = Cast<UDMMaterialStage>(ComponentWeak.Get()))
 	{
 		Stage->GetOnUpdate().RemoveAll(this);
@@ -259,6 +264,11 @@ void SDMComponentEdit::Construct(const FArguments& InArgs, UDMMaterialComponent*
 			}
 		}
 
+		if (UDMMaterialStage* Stage = Cast<UDMMaterialStage>(InComponent))
+		{
+			Stage->SetBeingEdited(true);
+		}
+
 		ChildSlot
 			.HAlign(HAlign_Fill)
 			.VAlign(VAlign_Top)
diff --git a/Engine/Plugins/Experimental/DynamicMaterial/Source/DynamicMaterialEditor/Private/Slate/SDMEditor.cpp b/Engine/Plugins/Experimental/DynamicMaterial/Source/DynamicMaterialEditor/Private/Slate/SDMEditor.cpp
--- a/Engine/Plugins/Experimental/DynamicMaterial/Source/DynamicMaterialEditor/Private/Slate/SDMEditor.cpp
+++ b/Engine/Plugins/Experimental/DynamicMaterial/Source/DynamicMaterialEditor/Private/Slate/SDMEditor.cpp
@@ -3,6 +3,7 @@
 #include "SDMEditor.h"
 #include "AssetThumbnail.h"
 #include "Components/DMMaterialLayer.h"
+#include "Components/DMMaterialProperty.h"
 #include "Components/DMMaterialSlot.h"
 #include "Components/DMMaterialValue.h"
 #include "Components/MaterialStageExpressions/DMMSETextureSample.h"
@@ -42,6 +43,7 @@
 #include "Widgets/Layout/SBox.h"
 #include "Widgets/Layout/SExpandableArea.h"
 #include "Widgets/Layout/SScrollBox.h"
+#include "Widgets/Layout/SWrapBox.h"
 #include "Widgets/SBoxPanel.h"
 #include "Widgets/SNullWidget.h"
 #include "Widgets/Text/STextBlock.h"
@@ -576,6 +578,7 @@ TSharedRef<SWidget> SDMEditor::CreateMainLayout()
 		.AutoHeight()
 		.HAlign(HAlign_Fill)
 		.VAlign(VAlign_Top)
+		.Padding(5.f, 0.f)
 		[
 			SAssignNew(SlotPickerContainer, SBox)
 			[
@@ -677,7 +680,10 @@ TSharedRef<SWidget> SDMEditor::CreateParametersArea()
 
 TSharedRef<SWidget> SDMEditor::CreateSlotPickerWidget()
 {
-	TSharedRef<SHorizontalBox> SlotSelector = SNew(SHorizontalBox);
+	TSharedRef<SWrapBox> SlotSelector = SNew(SWrapBox)
+		.UseAllottedSize(true)
+		.InnerSlotPadding(FVector2D(5, 5))
+		.Orientation(EOrientation::Orient_Horizontal);
 
 	UDynamicMaterialModelEditorOnlyData* ModelEditorOnlyData = UDynamicMaterialModelEditorOnlyData::Get(MaterialModelWeak);
 
@@ -686,83 +692,77 @@ TSharedRef<SWidget> SDMEditor::CreateSlotPickerWidget()
 		return SlotSelector;
 	}
 
-	const TArray<UDMMaterialSlot*>& Slots = ModelEditorOnlyData->GetSlots();
+	const FDMMaterialChannelListPreset* Preset = GetDefault<UDynamicMaterialEditorSettings>()->ChannelPresets.Find(ModelEditorOnlyData->GetChannelListPreset());
+	bool bHasBaseColorSlot = !!ModelEditorOnlyData->GetSlotForMaterialProperty(EDMMaterialPropertyType::BaseColor);
 
-	if (Slots.IsEmpty())
-	{
-		return SlotSelector;
-	}
+	UEnum* MaterialPropertyEnum = StaticEnum<EDMMaterialPropertyType>();
 
-	if (Slots.IsValidIndex(0))
-	{
-		SlotSelector->AddSlot()
-			.Padding(3.f, 0.f, 3.f, 0.f)
-			[
-				SNew(SCheckBox)
-				.Style(FAppStyle::Get(), "DetailsView.SectionButton")
-				.HAlign(EHorizontalAlignment::HAlign_Center)
-				.IsChecked(this, &SDMEditor::GetRGBSlotCheckState_HasSlot)
-				.OnCheckStateChanged(this, &SDMEditor::OnRGBSlotCheckStateChanged_HasSlot)
-				.Content()
-				[
-					SNew(STextBlock)
-					.Text(LOCTEXT("RGBSlot", "RGB"))
-				]
-			];		
-	}
-	else
+	for (uint8 PropertyIndex = static_cast<uint8>(EDMMaterialPropertyType::None) + 1;
+		PropertyIndex < static_cast<uint8>(EDMMaterialPropertyType::Any);
+		++PropertyIndex)
 	{
-		SlotSelector->AddSlot()
-			.Padding(3.f, 0.f, 3.f, 0.f)
-			[
-				SNew(SCheckBox)
-				.Style(FAppStyle::Get(), "DetailsView.SectionButton")
-				.HAlign(EHorizontalAlignment::HAlign_Center)
-				.IsChecked(this, &SDMEditor::GetRGBSlotCheckState_NoSlot)
-				.OnCheckStateChanged(this, &SDMEditor::OnRGBSlotCheckStateChanged_NoSlot)
-				.Content()
-				[
-					SNew(STextBlock)
-					.Text(LOCTEXT("RGBSlot", "RGB"))
-				]
-			];
-	}
+		EDMMaterialPropertyType Property = static_cast<EDMMaterialPropertyType>(PropertyIndex);
 
-	if (Slots.IsValidIndex(1))
-	{
-		SlotSelector->AddSlot()
-			.Padding(0.f, 0.f, 3.f, 0.f)
-			[
-				SNew(SCheckBox)
-				.Style(FAppStyle::Get(), "DetailsView.SectionButton")
-				.HAlign(EHorizontalAlignment::HAlign_Center)
-				.IsEnabled(this, &SDMEditor::GetOpacityButtonEnabled_HasSlot)
-				.IsChecked(this, &SDMEditor::GetOpacitySlotCheckState_HasSlot)
-				.OnCheckStateChanged(this, &SDMEditor::OnOpacitySlotCheckStateChanged_HasSlot)
-				.Content()
-				[
-					SNew(STextBlock)
-					.Text(LOCTEXT("OpacitySlot", "Opacity"))
-				]
-			];		
-	}
-	else
-	{
-		SlotSelector->AddSlot()
-			.Padding(0.f, 0.f, 3.f, 0.f)
-			[
-				SNew(SCheckBox)
-				.Style(FAppStyle::Get(), "DetailsView.SectionButton")
-				.HAlign(EHorizontalAlignment::HAlign_Center)
-				.IsEnabled(this, &SDMEditor::GetOpacityButtonEnabled_HasSlot)
-				.IsChecked(this, &SDMEditor::GetOpacitySlotCheckState_NoSlot)
-				.OnCheckStateChanged(this, &SDMEditor::OnOpacitySlotCheckStateChanged_NoSlot)
-				.Content()
+		// Filter RGB modes
+		if (bHasBaseColorSlot)
+		{
+			if (Property == EDMMaterialPropertyType::EmissiveColor)
+			{
+				continue;
+			}
+		}
+		else
+		{
+			if (Property == EDMMaterialPropertyType::BaseColor)
+			{
+				continue;
+			}
+		}
+
+		// Filter opacity modes
+		if (Property == EDMMaterialPropertyType::OpacityMask)
+		{
+			if (ModelEditorOnlyData->GetBlendMode() != BLEND_Masked)
+			{
+				continue;
+			}
+		}
+		else
+		{
+			if (ModelEditorOnlyData->GetBlendMode() == BLEND_Masked)
+			{
+				continue;
+			}
+		}
+
+		const bool bIsPropertyEnabled = Preset
+			? Preset->IsPropertyEnabled(Property)
+			: !!ModelEditorOnlyData->GetSlotForMaterialProperty(Property);
+
+		if (bIsPropertyEnabled)
+		{
+			const int32 EnumIndex = MaterialPropertyEnum->GetIndexByValue(static_cast<int64>(Property));
+
+			constexpr const TCHAR* ShortNameName = TEXT("ShortName");
+			const FString ShortName = MaterialPropertyEnum->GetMetaData(ShortNameName, EnumIndex);
+
+			SlotSelector->AddSlot()
 				[
-					SNew(STextBlock)
-					.Text(LOCTEXT("AddOpacity", "Opacity"))
-				]
-			];
+					SNew(SCheckBox)
+					.Style(FAppStyle::Get(), "DetailsView.SectionButton")
+					.HAlign(EHorizontalAlignment::HAlign_Center)
+					.IsEnabled(this, &SDMEditor::IsPropertyValidForModel, Property)
+					.IsChecked(this, &SDMEditor::GetSlotCheckState, Property)
+					.OnCheckStateChanged(this, &SDMEditor::OnSlotCheckStateChanged, Property)
+					.Padding(FVector2D(5.f, 3.f))
+					.ToolTipText(this, &SDMEditor::GetToolTipForProperty, Property)
+					.Content()
+					[
+						SNew(STextBlock)
+						.Text(!ShortName.IsEmpty() ? FText::FromString(ShortName) : MaterialPropertyEnum->GetDisplayNameTextByValue(static_cast<int64>(Property)))
+					]
+				];
+		}
 	}
 
 	return SlotSelector;
@@ -772,9 +772,9 @@ TSharedPtr<SDMSlot> SDMEditor::GetSlotWidget(UDMMaterialSlot* Slot) const
 {
 	if (ensure(IsValid(Slot)))
 	{
-		for (const TSharedPtr<SDMSlot>& SlotWidget : SlotWidgets)
+		for (const TSharedRef<SDMSlot>& SlotWidget : SlotWidgets)
 		{
-			if (SlotWidget.IsValid() && SlotWidget->GetSlot() == Slot)
+			if (SlotWidget->GetSlot() == Slot)
 			{
 				return SlotWidget;
 			}
@@ -788,6 +788,7 @@ void SDMEditor::RefreshGlobalOpacitySlider()
 {
 	if (GlobalOpacityContainer.IsValid())
 	{
+		GlobalOpacityContainer->SetContent(SNullWidget::NullWidget);
 		GlobalOpacityContainer->SetContent(CreateGlobalOpacityWidget());
 	}
 }
@@ -804,6 +805,7 @@ void SDMEditor::RefreshSlotPickerList()
 {
 	if (SlotPickerContainer.IsValid())
 	{
+		SlotPickerContainer->SetContent(SNullWidget::NullWidget);
 		SlotPickerContainer->SetContent(CreateSlotPickerWidget());
 	}
 }
@@ -812,6 +814,7 @@ void SDMEditor::RefreshSlotsList()
 {
 	if (SlotsContainer.IsValid())
 	{
+		SlotsContainer->SetContent(SNullWidget::NullWidget);
 		SlotsContainer->SetContent(CreateSlotsWidget());
 	}
 }
@@ -885,10 +888,10 @@ void SDMEditor::Tick(const FGeometry& AllottedGeometry, const double InCurrentTi
 		else
 		{
 			bool bHasInvalidSlotWidget = false;
-			for (const TWeakPtr<SDMSlot> SlotWeak : SlotWidgets)
+
+			for (const TSharedRef<SDMSlot>& Slot : SlotWidgets)
 			{
-				TSharedPtr<SDMSlot> Slot = SlotWeak.Pin();
-				if (Slot.IsValid() && !Slot->CheckValidity())
+				if (!Slot->CheckValidity())
 				{
 					bHasInvalidSlotWidget = true;
 					break;
@@ -902,25 +905,6 @@ void SDMEditor::Tick(const FGeometry& AllottedGeometry, const double InCurrentTi
 	}
 }
 
-FReply SDMEditor::OnAddSlotButtonClicked()
-{
-	if (!MaterialModelWeak.IsValid())
-	{
-		return FReply::Unhandled();
-	}
-
-	UDynamicMaterialModelEditorOnlyData* ModelEditorOnlyData = UDynamicMaterialModelEditorOnlyData::Get(MaterialModelWeak);
-
-	if (!ModelEditorOnlyData)
-	{
-		return FReply::Unhandled();
-	}
-
-	ModelEditorOnlyData->AddSlot();
-
-	return FReply::Handled();
-}
-
 void SDMEditor::OnMaterialBuilt(UDynamicMaterialModel* InMaterialModel)
 {
 }
@@ -945,94 +929,66 @@ bool SDMEditor::IsGlobalOpacityEnabled() const
 	return false;
 }
 
-ECheckBoxState SDMEditor::GetRGBSlotCheckState_HasSlot() const
+bool SDMEditor::IsPropertyValidForModel(EDMMaterialPropertyType InProperty) const
 {
-	return GetActiveSlotIndex() == 0 ? ECheckBoxState::Checked : ECheckBoxState::Unchecked;
-}
-
-void SDMEditor::OnRGBSlotCheckStateChanged_HasSlot(ECheckBoxState InCheckState)
-{
-	SetActiveSlotIndex(0);
-}
-
-ECheckBoxState SDMEditor::GetRGBSlotCheckState_NoSlot() const
-{
-	return ECheckBoxState::Unchecked;
-}
-
-void SDMEditor::OnRGBSlotCheckStateChanged_NoSlot(ECheckBoxState InCheckState)
-{
-	if (UDynamicMaterialModelEditorOnlyData* ModelEditorOnlyData = UDynamicMaterialModelEditorOnlyData::Get(MaterialModelWeak))
+	if (UDynamicMaterialModel* MaterialModel = MaterialModelWeak.Get())
 	{
-		if (ModelEditorOnlyData->AddSlot())
+		if (UDynamicMaterialModelEditorOnlyData* EditorOnlyData = UDynamicMaterialModelEditorOnlyData::Get(MaterialModel))
 		{
-			RefreshSlotPickerList();
-			SetActiveSlotIndex(0);
+			if (UDMMaterialProperty* Property = EditorOnlyData->GetMaterialProperty(InProperty))
+			{
+				return Property->IsValidForModel(*EditorOnlyData);
+			}
 		}
 	}
-}
-
-bool SDMEditor::GetOpacityButtonEnabled_NoSlot() const
-{
-	if (UDynamicMaterialModelEditorOnlyData* ModelEditorOnlyData = UDynamicMaterialModelEditorOnlyData::Get(MaterialModelWeak))
-	{
-		return ModelEditorOnlyData->GetDomain() != EMaterialDomain::MD_PostProcess
-			&& ModelEditorOnlyData->GetBlendMode() != EBlendMode::BLEND_Opaque;
-	}
 
 	return false;
 }
 
-bool SDMEditor::GetOpacityButtonEnabled_HasSlot() const
+ECheckBoxState SDMEditor::GetSlotCheckState(EDMMaterialPropertyType InProperty)  const
 {
-	if (UDynamicMaterialModelEditorOnlyData* ModelEditorOnlyData = UDynamicMaterialModelEditorOnlyData::Get(MaterialModelWeak))
+	if (SlotWidgets.IsValidIndex(ActiveSlotIndex))
 	{
-		return ModelEditorOnlyData->GetBlendMode() != EBlendMode::BLEND_Opaque;
+		if (UDMMaterialSlot* Slot = SlotWidgets[ActiveSlotIndex]->GetSlot())
+		{
+			if (UDynamicMaterialModelEditorOnlyData* EditorOnlyData = Slot->GetMaterialModelEditorOnlyData())
+			{
+				return EditorOnlyData->GetSlotForMaterialProperty(InProperty) == Slot
+					? ECheckBoxState::Checked
+					: ECheckBoxState::Unchecked;
+			}
+		}
 	}
 
-	return false;
-}
-
-ECheckBoxState SDMEditor::GetOpacitySlotCheckState_HasSlot() const
-{
-	return GetActiveSlotIndex() == 1 ? ECheckBoxState::Checked : ECheckBoxState::Unchecked;
+	return ECheckBoxState::Undetermined;
 }
 
-void SDMEditor::OnOpacitySlotCheckStateChanged_HasSlot(ECheckBoxState InCheckState)
+FText SDMEditor::GetToolTipForProperty(EDMMaterialPropertyType InProperty) const
 {
-	SetActiveSlotIndex(1);
-}
+	if (!IsPropertyValidForModel(InProperty))
+	{
+		return LOCTEXT("SlotNotCompatible", "This slot is not compatible with the material settings.");
+	}
 
-ECheckBoxState SDMEditor::GetOpacitySlotCheckState_NoSlot() const
-{
-	return ECheckBoxState::Unchecked;
+	return FText::GetEmpty();
 }
 
-void SDMEditor::OnOpacitySlotCheckStateChanged_NoSlot(ECheckBoxState InCheckState)
+void SDMEditor::OnSlotCheckStateChanged(ECheckBoxState InCheckState, EDMMaterialPropertyType InProperty)
 {
-	if (UDynamicMaterialModelEditorOnlyData* ModelEditorOnlyData = UDynamicMaterialModelEditorOnlyData::Get(MaterialModelWeak))
+	if (InCheckState == ECheckBoxState::Checked)
 	{
-		if (UDMMaterialSlot* NewSlot = ModelEditorOnlyData->AddSlot())
+		for (int32 SlotIndex = 0; SlotIndex < SlotWidgets.Num(); ++SlotIndex)
 		{
-			if (UDMMaterialLayerObject* Layer = NewSlot->GetLayer(0))
+			if (UDMMaterialSlot* Slot = SlotWidgets[SlotIndex]->GetSlot())
 			{
-				FDMScopedUITransaction Transaction(LOCTEXT("AddOpacitySlot", "Material Designer Add Opacity Slot"));
-				Layer->Modify();
-
-				switch (ModelEditorOnlyData->GetBlendMode())
+				if (UDynamicMaterialModelEditorOnlyData* EditorOnlyData = Slot->GetMaterialModelEditorOnlyData())
 				{
-					case EBlendMode::BLEND_Masked:
-						Layer->SetMaterialProperty(EDMMaterialPropertyType::OpacityMask);
-						break;
-
-					default:
-						Layer->SetMaterialProperty(EDMMaterialPropertyType::Opacity);
-						break;
+					if (EditorOnlyData->GetSlotForMaterialProperty(InProperty) == Slot)
+					{
+						SetActiveSlotIndex(SlotIndex);
+					}
 				}
 			}
-
-			RefreshSlotPickerList();
-			SetActiveSlotIndex(1);
 		}
 	}
 }
@@ -1298,6 +1254,8 @@ TSharedRef<SWidget> SDMEditor::CreateSlotsWidget()
 		return CreateEmptySlotsContent();
 	}
 
+	SlotWidgets.Reserve(Slots.Num());
+
 	for (int32 SlotIdx = 0; SlotIdx < Slots.Num(); ++SlotIdx)
 	{
 		UDMMaterialSlot* Slot = Slots[SlotIdx];
@@ -1310,158 +1268,17 @@ TSharedRef<SWidget> SDMEditor::CreateSlotsWidget()
 		SlotWidgets.Add(SlotWidget);
 	}
 
-	return SlotWidgets.IsValidIndex(ActiveSlotIndex)
-		? SlotWidgets[ActiveSlotIndex].ToSharedRef()
-		: SNullWidget::NullWidget;
-}
-
-EMaterialDomain SDMEditor::GetSelectedDomain() const
-{
-	if (UDynamicMaterialModelEditorOnlyData* ModelEditorOnlyData = UDynamicMaterialModelEditorOnlyData::Get(MaterialModelWeak))
-	{
-		return ModelEditorOnlyData->GetDomain();
-	}
-
-	return EMaterialDomain::MD_Surface;
-}
-
-void SDMEditor::OnDomainChanged(const EMaterialDomain InDomain)
-{
-	if (UDynamicMaterialModelEditorOnlyData* ModelEditorOnlyData = UDynamicMaterialModelEditorOnlyData::Get(MaterialModelWeak))
-	{
-		FDMScopedUITransaction Transaction(LOCTEXT("ChangeDomain", "Material Designer Change Domain"));
-		ModelEditorOnlyData->Modify();
-		ModelEditorOnlyData->SetDomain(InDomain);
-
-		RefreshSlotPickerList();
-
-		if (InDomain == EMaterialDomain::MD_PostProcess)
-		{
-			SetActiveSlotIndex(0);
-		}
-	}
-}
-
-bool SDMEditor::CanChangeDomain() const
-{
-	return true;
-}
-
-EBlendMode SDMEditor::GetSelectedBlendMode() const
-{
-	if (UDynamicMaterialModelEditorOnlyData* ModelEditorOnlyData = UDynamicMaterialModelEditorOnlyData::Get(MaterialModelWeak))
-	{
-		return ModelEditorOnlyData->GetBlendMode();
-	}
-
-	return EBlendMode::BLEND_Opaque;
-}
-
-void SDMEditor::OnBlendModeChanged(const EBlendMode InBlendMode)
-{
-	if (UDynamicMaterialModelEditorOnlyData* ModelEditorOnlyData = UDynamicMaterialModelEditorOnlyData::Get(MaterialModelWeak))
-	{
-		FDMScopedUITransaction Transaction(LOCTEXT("ChangeBlendMode", "Material Designer Change Blend Mode"));
-		ModelEditorOnlyData->Modify();
-		ModelEditorOnlyData->SetBlendMode(InBlendMode);
-
-		RefreshSlotPickerList();
-
-		if (InBlendMode == EBlendMode::BLEND_Opaque)
-		{
-			SetActiveSlotIndex(0);
-		}
-	}
-}
-
-bool SDMEditor::CanChangeBlendType() const
-{
-	if (UDynamicMaterialModelEditorOnlyData* ModelEditorOnlyData = UDynamicMaterialModelEditorOnlyData::Get(MaterialModelWeak))
+	if (!SlotWidgets.IsValidIndex(ActiveSlotIndex))
 	{
-		switch (ModelEditorOnlyData->GetDomain())
+		if (SlotWidgets.IsEmpty())
 		{
-			case EMaterialDomain::MD_Surface:
-			case EMaterialDomain::MD_DeferredDecal:
-				return true;
+			return SNullWidget::NullWidget;
 		}
-	}
-
-	return false;
-}
 
-ECheckBoxState SDMEditor::IsMaterialUnlit() const
-{
-	if (UDynamicMaterialModelEditorOnlyData* ModelEditorOnlyData = UDynamicMaterialModelEditorOnlyData::Get(MaterialModelWeak))
-	{
-		return ModelEditorOnlyData->GetShadingModel() == EDMMaterialShadingModel::Unlit
-			? ECheckBoxState::Checked : ECheckBoxState::Unchecked;
+		ActiveSlotIndex = 0;
 	}
 
-	return ECheckBoxState::Unchecked;
-}
-
-void SDMEditor::OnMaterialUnlitChanged(const ECheckBoxState InNewCheckState)
-{
-	if (UDynamicMaterialModelEditorOnlyData* ModelEditorOnlyData = UDynamicMaterialModelEditorOnlyData::Get(MaterialModelWeak))
-	{
-		FDMScopedUITransaction Transaction(LOCTEXT("ChangeShadingModel", "Material Designer Change Shading Model"));
-		ModelEditorOnlyData->Modify();
-		ModelEditorOnlyData->SetShadingModel(InNewCheckState == ECheckBoxState::Checked
-			? EDMMaterialShadingModel::Unlit : EDMMaterialShadingModel::DefaultLit);
-	}
-}
-
-bool SDMEditor::CanMaterialBeAnimated() const
-{
-	if (UDynamicMaterialModelEditorOnlyData* ModelEditorOnlyData = UDynamicMaterialModelEditorOnlyData::Get(MaterialModelWeak))
-	{
-		switch (ModelEditorOnlyData->GetBlendMode())
-		{
-			case EBlendMode::BLEND_Masked:
-			case EBlendMode::BLEND_Opaque:
-				return true;
-
-			default:
-				return false;
-		}
-	}
-
-	return false;
-}
-
-bool SDMEditor::CanChangeMaterialShadingModel() const
-{
-	if (UDynamicMaterialModelEditorOnlyData* ModelEditorOnlyData = UDynamicMaterialModelEditorOnlyData::Get(MaterialModelWeak))
-	{
-		switch (ModelEditorOnlyData->GetDomain())
-		{
-			case EMaterialDomain::MD_Surface:
-				return true;
-		}
-	}
-
-	return false;
-}
-
-ECheckBoxState SDMEditor::IsMaterialAnimated() const
-{
-	if (UDynamicMaterialModelEditorOnlyData* ModelEditorOnlyData = UDynamicMaterialModelEditorOnlyData::Get(MaterialModelWeak))
-	{
-		return ModelEditorOnlyData->IsPixelAnimationFlagSet()
-			? ECheckBoxState::Checked : ECheckBoxState::Unchecked;
-	}
-
-	return ECheckBoxState::Unchecked;
-}
-
-void SDMEditor::OnMaterialAnimatedChanged(const ECheckBoxState InNewCheckState)
-{
-	if (UDynamicMaterialModelEditorOnlyData* ModelEditorOnlyData = UDynamicMaterialModelEditorOnlyData::Get(MaterialModelWeak))
-	{
-		FDMScopedUITransaction Transaction(LOCTEXT("ToggleAnimated", "Material Designer Toggle Animated"));
-		ModelEditorOnlyData->Modify();
-		ModelEditorOnlyData->SetPixelAnimationFlag(InNewCheckState == ECheckBoxState::Checked);
-	}
+	return SlotWidgets[ActiveSlotIndex];
 }
 
 void SDMEditor::OnUndo()
diff --git a/Engine/Plugins/Experimental/DynamicMaterial/Source/DynamicMaterialEditor/Private/Slate/SDMEditor.h b/Engine/Plugins/Experimental/DynamicMaterial/Source/DynamicMaterialEditor/Private/Slate/SDMEditor.h
--- a/Engine/Plugins/Experimental/DynamicMaterial/Source/DynamicMaterialEditor/Private/Slate/SDMEditor.h
+++ b/Engine/Plugins/Experimental/DynamicMaterial/Source/DynamicMaterialEditor/Private/Slate/SDMEditor.h
@@ -58,7 +58,7 @@ public:
 
 	void SetMaterialActor(AActor* InActor);
 
-	const TArray<TSharedPtr<SDMSlot>>& GetSlotWidgets() const { return SlotWidgets; }
+	const TArray<TSharedRef<SDMSlot>>& GetSlotWidgets() const { return SlotWidgets; }
 	TSharedPtr<SDMSlot> GetSlotWidget(UDMMaterialSlot* Slot) const;
 
 	void RefreshGlobalOpacitySlider();
@@ -129,7 +129,7 @@ protected:
 	TWeakObjectPtr<UDynamicMaterialModel> MaterialModelWeak;
 	FDMObjectMaterialProperty ObjectProperty;
 
-	TArray<TSharedPtr<SDMSlot>> SlotWidgets;
+	TArray<TSharedRef<SDMSlot>> SlotWidgets;
 
 	void BindCommands();
 
@@ -142,27 +142,16 @@ protected:
 
 	TSharedRef<SWidget> CreateActorMaterialSlotSelector(const AActor* InActor);
 
-	FReply OnAddSlotButtonClicked();
-
 	void OnMaterialBuilt(UDynamicMaterialModel* InMaterialModel);
 	void OnValuesUpdated(UDynamicMaterialModel* InMaterialModel);
 	void OnSlotsUpdated(UDynamicMaterialModel* InMaterialModel);
 
 	bool IsGlobalOpacityEnabled() const;
 
-	ECheckBoxState GetRGBSlotCheckState_HasSlot() const;
-	void OnRGBSlotCheckStateChanged_HasSlot(ECheckBoxState InCheckState);
-
-	ECheckBoxState GetRGBSlotCheckState_NoSlot() const;
-	void OnRGBSlotCheckStateChanged_NoSlot(ECheckBoxState InCheckState);
-
-	bool GetOpacityButtonEnabled_HasSlot() const;
-	ECheckBoxState GetOpacitySlotCheckState_HasSlot() const;
-	void OnOpacitySlotCheckStateChanged_HasSlot(ECheckBoxState InCheckState);
-
-	bool GetOpacityButtonEnabled_NoSlot() const;
-	ECheckBoxState GetOpacitySlotCheckState_NoSlot() const;
-	void OnOpacitySlotCheckStateChanged_NoSlot(ECheckBoxState InCheckState);
+	bool IsPropertyValidForModel(EDMMaterialPropertyType InProperty) const;
+	ECheckBoxState GetSlotCheckState(EDMMaterialPropertyType InProperty) const;
+	FText GetToolTipForProperty(EDMMaterialPropertyType InProperty) const;
+	void OnSlotCheckStateChanged(ECheckBoxState InCheckState, EDMMaterialPropertyType InProperty);
 
 	FReply OnCreateMaterialButtonClicked(TWeakPtr<FDMObjectMaterialProperty> InMaterialProperty);
 
@@ -171,22 +160,5 @@ protected:
 
 	void OnSettingsChanged(const FPropertyChangedEvent& InPropertyChangedEvent);
 
-	EMaterialDomain GetSelectedDomain() const;
-	void OnDomainChanged(const EMaterialDomain InDomain);
-	bool CanChangeDomain() const;
-
-	EBlendMode GetSelectedBlendMode() const;
-	void OnBlendModeChanged(const EBlendMode InBlendMode);
-	bool CanChangeBlendType() const;
-
-	ECheckBoxState IsMaterialUnlit() const;
-	void OnMaterialUnlitChanged(const ECheckBoxState InNewCheckState);
-
-	bool CanMaterialBeAnimated() const;
-	ECheckBoxState IsMaterialAnimated() const;
-	void OnMaterialAnimatedChanged(const ECheckBoxState InNewCheckState);
-
-	bool CanChangeMaterialShadingModel() const;
-
 	void OnUndo();
 };
diff --git a/Engine/Plugins/Experimental/DynamicMaterial/Source/DynamicMaterialEditor/Private/Slate/SDMMaterialWizard.cpp b/Engine/Plugins/Experimental/DynamicMaterial/Source/DynamicMaterialEditor/Private/Slate/SDMMaterialWizard.cpp
--- a/Engine/Plugins/Experimental/DynamicMaterial/Source/DynamicMaterialEditor/Private/Slate/SDMMaterialWizard.cpp
+++ b/Engine/Plugins/Experimental/DynamicMaterial/Source/DynamicMaterialEditor/Private/Slate/SDMMaterialWizard.cpp
@@ -3,6 +3,7 @@
 #include "Slate/SDMMaterialWizard.h"
 
 #include "DMDefs.h"
+#include "DynamicMaterialEditorSettings.h"
 #include "DynamicMaterialEditorStyle.h"
 #include "Engine/EngineTypes.h"
 #include "MaterialDomain.h"
@@ -10,6 +11,8 @@
 #include "Model/DynamicMaterialModelEditorOnlyData.h"
 #include "SDMEditor.h"
 #include "Styling/SlateTypes.h"
+#include "Styling/StyleColors.h"
+#include "Widgets/Colors/SColorBlock.h"
 #include "Widgets/Input/SButton.h"
 #include "Widgets/Input/SCheckBox.h"
 #include "Widgets/Layout/SBox.h"
@@ -31,6 +34,12 @@ void SDMMaterialWizard::Construct(const FArguments& InArgs, const TSharedRef<SDM
 {
 	EditorWeak = InEditor;
 
+	for (const TPair<FName, FDMMaterialChannelListPreset>& Preset : GetDefault<UDynamicMaterialEditorSettings>()->ChannelPresets)
+	{
+		CurrentPreset = Preset.Key;
+		break;
+	}
+
 	ChildSlot
 		[
 			CreateLayout()
@@ -76,14 +85,15 @@ TSharedRef<SWidget> SDMMaterialWizard::CreateLayout()
 			.Padding(0.0f, SeparationDistance, 0.0f, TitleContentDistance)
 			[
 				SNew(STextBlock)
-				.Text(LOCTEXT("MaterialDomain", "Material Domain"))
+				.Text(LOCTEXT("MaterialType", "Material Type"))
 			]
 
 			+ SVerticalBox::Slot()
 			.AutoHeight()
 			.HAlign(EHorizontalAlignment::HAlign_Fill)
+			.Padding(0.0f, 0.f, 0.0f, TitleContentDistance)
 			[
-				CreateMaterialDomainOptions()
+				CreateChannelPresets()
 			]
 
 			+ SVerticalBox::Slot()
@@ -92,14 +102,31 @@ TSharedRef<SWidget> SDMMaterialWizard::CreateLayout()
 			.Padding(0.0f, SeparationDistance, 0.0f, TitleContentDistance)
 			[
 				SNew(STextBlock)
-				.Text(LOCTEXT("BlendMode", "Blend Mode"))
+				.Text(LOCTEXT("AvailableChannels", "Available Channels"))
 			]
 
 			+ SVerticalBox::Slot()
 			.AutoHeight()
 			.HAlign(EHorizontalAlignment::HAlign_Fill)
 			[
-				CreateBlendModeOptions()
+				SAssignNew(PresetChannelContainer, SBox)
+				.HAlign(EHorizontalAlignment::HAlign_Fill)
+				[
+					CreateChannelList()
+				]
+			]
+
+			+ SVerticalBox::Slot()
+			.AutoHeight()
+			.HAlign(HAlign_Fill)
+			.Padding(0.0f, SeparationDistance, 0.0f, 0.f)
+			[
+				SNew(SBox)
+				.HeightOverride(1.f)
+				[
+					SNew(SColorBlock)
+					.Color(FStyleColors::Foreground.GetSpecifiedColor().ToFColor(false))
+				]
 			]
 
 			+ SVerticalBox::Slot()
@@ -108,14 +135,14 @@ TSharedRef<SWidget> SDMMaterialWizard::CreateLayout()
 			.Padding(0.0f, SeparationDistance, 0.0f, TitleContentDistance)
 			[
 				SNew(STextBlock)
-				.Text(LOCTEXT("ShadingModel", "Shading Model"))
+				.Text(LOCTEXT("MaterialDomain", "Material Domain"))
 			]
 
 			+ SVerticalBox::Slot()
 			.AutoHeight()
 			.HAlign(EHorizontalAlignment::HAlign_Fill)
 			[
-				CreateShadingModelOptions()
+				CreateMaterialDomainOptions()
 			]
 
 			+ SVerticalBox::Slot()
@@ -124,14 +151,14 @@ TSharedRef<SWidget> SDMMaterialWizard::CreateLayout()
 			.Padding(0.0f, SeparationDistance, 0.0f, TitleContentDistance)
 			[
 				SNew(STextBlock)
-				.Text(LOCTEXT("MotionVectors", "Motion Vectors"))
+				.Text(LOCTEXT("BlendMode", "Blend Mode"))
 			]
 
 			+ SVerticalBox::Slot()
 			.AutoHeight()
 			.HAlign(EHorizontalAlignment::HAlign_Fill)
 			[
-				CreateAnimationOptions()
+				CreateBlendModeOptions()
 			]
 
 			+ SVerticalBox::Slot()
@@ -140,14 +167,14 @@ TSharedRef<SWidget> SDMMaterialWizard::CreateLayout()
 			.Padding(0.0f, SeparationDistance, 0.0f, TitleContentDistance)
 			[
 				SNew(STextBlock)
-				.Text(LOCTEXT("Geometry", "Geometry"))
+				.Text(LOCTEXT("ShadingModel", "Shading Model"))
 			]
 
 			+ SVerticalBox::Slot()
 			.AutoHeight()
 			.HAlign(EHorizontalAlignment::HAlign_Fill)
 			[
-				CreateTwoSidedOptions()
+				CreateShadingModelOptions()
 			]
 
 			+ SVerticalBox::Slot()
@@ -156,22 +183,43 @@ TSharedRef<SWidget> SDMMaterialWizard::CreateLayout()
 			.Padding(0.0f, SeparationDistance, 0.0f, TitleContentDistance)
 			[
 				SNew(STextBlock)
-				.Text(LOCTEXT("Channels", "Channels"))
+				.Text(LOCTEXT("MotionVectors", "Motion Vectors"))
 			]
 
 			+ SVerticalBox::Slot()
 			.AutoHeight()
 			.HAlign(EHorizontalAlignment::HAlign_Fill)
-			.Padding(0.0f, 0.f, 0.0f, TitleContentDistance)
 			[
-				CreateChannelPresets()
+				CreateAnimationOptions()
+			]
+
+			+ SVerticalBox::Slot()
+			.AutoHeight()
+			.HAlign(HAlign_Left)
+			.Padding(0.0f, SeparationDistance, 0.0f, TitleContentDistance)
+			[
+				SNew(STextBlock)
+				.Text(LOCTEXT("Geometry", "Geometry"))
 			]
 
 			+ SVerticalBox::Slot()
 			.AutoHeight()
 			.HAlign(EHorizontalAlignment::HAlign_Fill)
 			[
-				CreateChannelList()
+				CreateTwoSidedOptions()
+			]
+
+			+ SVerticalBox::Slot()
+			.AutoHeight()
+			.HAlign(HAlign_Fill)
+			.Padding(0.0f, SeparationDistance, 0.0f, 0.f)
+			[
+				SNew(SBox)
+				.HeightOverride(1.f)
+				[
+					SNew(SColorBlock)
+					.Color(FStyleColors::Foreground.GetSpecifiedColor().ToFColor(false))
+				]
 			]
 
 			+ SVerticalBox::Slot()
@@ -368,25 +416,78 @@ TSharedRef<SWidget> SDMMaterialWizard::CreateChannelPresets()
 {
 	using namespace UE::DynamicMaterialEditor::Private;
 
-	return SNew(SHorizontalBox)
-		+ SHorizontalBox::Slot()
-		.AutoWidth()
-		.Padding(0.f, 0.f, TitleContentDistance, 0.f)
-		[
-			SNew(STextBlock)
-			.Text(LOCTEXT("Presets", "Presets"))
-		]
-		+ SHorizontalBox::Slot()
-		[
-			SNew(STextBlock)
-			.Text(INVTEXT("-"))
-		];
+	TSharedRef<SWrapBox> ChannelPresets = SNew(SWrapBox)
+		.UseAllottedSize(true)
+		.InnerSlotPadding(WrapBoxSlotPadding)
+		.Orientation(EOrientation::Orient_Horizontal);
+
+	for (const TPair<FName, FDMMaterialChannelListPreset>& Preset : GetDefault<UDynamicMaterialEditorSettings>()->ChannelPresets)
+	{
+		ChannelPresets->AddSlot()
+			[
+				SNew(SCheckBox)
+					.Style(FAppStyle::Get(), "DetailsView.SectionButton")
+					.HAlign(EHorizontalAlignment::HAlign_Center)
+					.Padding(ButtonPadding)
+					.IsChecked(this, &SDMMaterialWizard::Preset_GetState, Preset.Key)
+					.OnCheckStateChanged(this, &SDMMaterialWizard::Preset_OnChange, Preset.Key)
+					[
+						SNew(STextBlock)
+						.Text(FText::FromName(Preset.Key))
+					]
+			];
+	}
+
+	return ChannelPresets;
 }
 
 TSharedRef<SWidget> SDMMaterialWizard::CreateChannelList()
 {
-	return SNew(STextBlock)
-		.Text(INVTEXT("Full Channel List"));
+	using namespace UE::DynamicMaterialEditor::Private;
+
+	UEnum* MaterialPropertyEnum = StaticEnum<EDMMaterialPropertyType>();
+
+	TSharedRef<SWrapBox> ChannelPresets = SNew(SWrapBox)
+		.UseAllottedSize(true)
+		.InnerSlotPadding(WrapBoxSlotPadding)
+		.Orientation(EOrientation::Orient_Horizontal);
+
+	if (const FDMMaterialChannelListPreset* Preset = GetDefault<UDynamicMaterialEditorSettings>()->ChannelPresets.Find(CurrentPreset))
+	{
+		for (uint8 PropertyIndex = static_cast<uint8>(EDMMaterialPropertyType::None) + 1;
+			PropertyIndex < static_cast<uint8>(EDMMaterialPropertyType::Any);
+			++PropertyIndex)
+		{
+			const EDMMaterialPropertyType Property = static_cast<EDMMaterialPropertyType>(PropertyIndex);
+
+			if (Property == EDMMaterialPropertyType::EmissiveColor || Property == EDMMaterialPropertyType::OpacityMask)
+			{
+				continue;
+			}
+				
+			if (Preset->IsPropertyEnabled(Property))
+			{
+				constexpr const TCHAR* ShortNameName = TEXT("ShortName");
+				const FString ShortName = MaterialPropertyEnum->GetMetaData(ShortNameName, MaterialPropertyEnum->GetIndexByValue(static_cast<int64>(PropertyIndex)));
+
+				ChannelPresets->AddSlot()
+					[
+						SNew(SCheckBox)
+						.Style(FAppStyle::Get(), "DetailsView.SectionButton")
+						.HAlign(EHorizontalAlignment::HAlign_Center)
+						.Padding(ButtonPadding)
+						.IsChecked(false)
+						.IsEnabled(false)
+						[
+							SNew(STextBlock)
+							.Text(!ShortName.IsEmpty() ? FText::FromString(ShortName) : MaterialPropertyEnum->GetDisplayNameTextByValue(static_cast<int64>(Property)))
+						]
+					];
+			}
+		}
+	}
+
+	return ChannelPresets;
 }
 
 TSharedRef<SWidget> SDMMaterialWizard::CreateAcceptButton()
@@ -602,6 +703,26 @@ void SDMMaterialWizard::TwoSided_OnChange(ECheckBoxState InState, bool bInValue)
 	}
 }
 
+ECheckBoxState SDMMaterialWizard::Preset_GetState(FName InPresetName) const
+{
+	return CurrentPreset == InPresetName
+		? ECheckBoxState::Checked
+		: ECheckBoxState::Unchecked;
+}
+
+void SDMMaterialWizard::Preset_OnChange(ECheckBoxState InState, FName InPresetName)
+{
+	if (InState == ECheckBoxState::Checked)
+	{
+		CurrentPreset = InPresetName;
+
+		if (PresetChannelContainer.IsValid())
+		{
+			PresetChannelContainer->SetContent(CreateChannelList());
+		}
+	}
+}
+
 FReply SDMMaterialWizard::Accept_OnClick()
 {
 	if (TSharedPtr<SDMEditor> Editor = GetEditor())
@@ -610,6 +731,7 @@ FReply SDMMaterialWizard::Accept_OnClick()
 		{
 			if (UDynamicMaterialModelEditorOnlyData* EditorOnlyData = UDynamicMaterialModelEditorOnlyData::Get(MaterialModel))
 			{
+				EditorOnlyData->SetChannelListPreset(CurrentPreset);
 				EditorOnlyData->OnWizardComplete();
 				Editor->SetMaterialModel(MaterialModel); // Refresh display
 			}
diff --git a/Engine/Plugins/Experimental/DynamicMaterial/Source/DynamicMaterialEditor/Private/Slate/SDMMaterialWizard.h b/Engine/Plugins/Experimental/DynamicMaterial/Source/DynamicMaterialEditor/Private/Slate/SDMMaterialWizard.h
--- a/Engine/Plugins/Experimental/DynamicMaterial/Source/DynamicMaterialEditor/Private/Slate/SDMMaterialWizard.h
+++ b/Engine/Plugins/Experimental/DynamicMaterial/Source/DynamicMaterialEditor/Private/Slate/SDMMaterialWizard.h
@@ -5,6 +5,7 @@
 #include "Templates/SharedPointer.h"
 #include "Widgets/SCompoundWidget.h"
 
+class SBox;
 class SDMEditor;
 class SWidget;
 class UDynamicMaterialModel;
@@ -27,6 +28,8 @@ public:
 
 protected:
 	TWeakPtr<SDMEditor> EditorWeak;
+	FName CurrentPreset;
+	TSharedPtr<SBox> PresetChannelContainer;
 
 	TSharedRef<SWidget> CreateLayout();
 	TSharedRef<SWidget> CreateMaterialDomainOptions();
@@ -58,5 +61,8 @@ protected:
 	ECheckBoxState TwoSided_GetState(bool bInValue) const;
 	void TwoSided_OnChange(ECheckBoxState InState, bool bInValue);
 
+	ECheckBoxState Preset_GetState(FName InPresetName) const;
+	void Preset_OnChange(ECheckBoxState InState, FName InPresetName);
+
 	FReply Accept_OnClick();
 };
diff --git a/Engine/Plugins/Experimental/DynamicMaterial/Source/DynamicMaterialEditor/Private/Slate/SDMSlot.cpp b/Engine/Plugins/Experimental/DynamicMaterial/Source/DynamicMaterialEditor/Private/Slate/SDMSlot.cpp
--- a/Engine/Plugins/Experimental/DynamicMaterial/Source/DynamicMaterialEditor/Private/Slate/SDMSlot.cpp
+++ b/Engine/Plugins/Experimental/DynamicMaterial/Source/DynamicMaterialEditor/Private/Slate/SDMSlot.cpp
@@ -73,6 +73,7 @@ SDMSlot::~SDMSlot()
 	{
 		Slot->GetOnPropertiesUpdateDelegate().RemoveAll(this);
 		Slot->GetOnLayersUpdateDelegate().RemoveAll(this);
+		Slot->SetEditingLayers(false);
 	}
 
 	SDMEditor::ClearPropertyHandles(this);
@@ -398,6 +399,9 @@ void SDMSlot::RefreshMainWidget()
 		.OnLayerStageSelected(this, &SDMSlot::OnLayerStageSelected)
 		.PreviewSize(LayerPreviewSize);
 
+	ChildSlot.AttachWidget(SNullWidget::NullWidget);
+	SplitterContainer.Reset();
+
 	SplitterContainer =
 		SNew(SSplitter)
 		.Style(FAppStyle::Get(), "DetailsView.Splitter")
@@ -418,44 +422,44 @@ void SDMSlot::RefreshMainWidget()
 			]
 		]
 
-	+ SSplitter::Slot()
-		.Expose(LayerViewSplitterSlot)
-		.Resizable(true)
-		.SizeRule(SSplitter::ESizeRule::FractionOfParent)
-		.MinSize(50)
-		.Value(LayerViewSplitterLocation)
-		[
-			SNew(SBorder)
-			.Padding(2.0f)
-			.BorderImage(FDynamicMaterialEditorStyle::GetBrush("LayerView.Background"))
-			.HAlign(HAlign_Fill)
-			.VAlign(VAlign_Fill)
+		+ SSplitter::Slot()
+			.Expose(LayerViewSplitterSlot)
+			.Resizable(true)
+			.SizeRule(SSplitter::ESizeRule::FractionOfParent)
+			.MinSize(50)
+			.Value(LayerViewSplitterLocation)
 			[
-				LayerView.ToSharedRef()
+				SNew(SBorder)
+				.Padding(2.0f)
+				.BorderImage(FDynamicMaterialEditorStyle::GetBrush("LayerView.Background"))
+				.HAlign(HAlign_Fill)
+				.VAlign(VAlign_Fill)
+				[
+					LayerView.ToSharedRef()
+				]
 			]
-		]
-
-	+ SSplitter::Slot()
-		.Resizable(false)
-		.SizeRule(SSplitter::ESizeRule::SizeToContent)
-		[
-			CreateLayerButtonsRowWidget()
-		]
 
-	+ SSplitter::Slot()
-		.Expose(ExtraSpaceSplitterSlot)
-		.Resizable(true)
-		.SizeRule(SSplitter::ESizeRule::FractionOfParent)
-		.MinSize(50)
-		.Value(ExtraSpaceSplitterLocation)
-		[
-			SAssignNew(ComponentEditContainer, SScrollBox)
-			+ SScrollBox::Slot()
-			.AutoSize()
+		+ SSplitter::Slot()
+			.Resizable(false)
+			.SizeRule(SSplitter::ESizeRule::SizeToContent)
 			[
-				CreateComponentEditWidget()
+				CreateLayerButtonsRowWidget()
 			]
-		];
+
+		+ SSplitter::Slot()
+			.Expose(ExtraSpaceSplitterSlot)
+			.Resizable(true)
+			.SizeRule(SSplitter::ESizeRule::FractionOfParent)
+			.MinSize(50)
+			.Value(ExtraSpaceSplitterLocation)
+			[
+				SAssignNew(ComponentEditContainer, SScrollBox)
+				+ SScrollBox::Slot()
+				.AutoSize()
+				[
+					CreateComponentEditWidget()
+				]
+			];
 
 	TSharedRef<SVerticalBox> Child = SNew(SVerticalBox);
 
@@ -894,6 +898,7 @@ void SDMSlot::RefreshHeaderPropertyListWidget()
 
 	if (HeaderPropertyListWidget.IsValid())
 	{
+		HeaderPropertyListWidget->SetContent(SNullWidget::NullWidget);
 		HeaderPropertyListWidget->SetContent(CreateHeaderPropertyListWidget());
 	}
 }
@@ -904,6 +909,7 @@ void SDMSlot::RefreshSlotSettingsRowWidget()
 
 	if (SlotSettingsRowContainer.IsValid())
 	{
+		SlotSettingsRowContainer->SetContent(SNullWidget::NullWidget);
 		SlotSettingsRowContainer->SetContent(CreateSlotSettingsRow());
 	}
 }
@@ -1111,11 +1117,6 @@ bool SDMSlot::CanRemoveLayerByIndex(const int32 InLayerIndex) const
 		return false;
 	}
 
-	if (Slot->GetLayers().Num() > 1)
-	{
-		return true;
-	}
-
 	const UDMMaterialLayerObject* LayerToRemove = Slot->GetLayer(InLayerIndex);
 
 	if (!LayerToRemove)
@@ -1123,26 +1124,7 @@ bool SDMSlot::CanRemoveLayerByIndex(const int32 InLayerIndex) const
 		return false;
 	}
 
-	if (UDynamicMaterialModelEditorOnlyData* ModelEditorOnlyData = Slot->GetMaterialModelEditorOnlyData())
-	{
-		const TArray<EDMMaterialPropertyType> Properties = ModelEditorOnlyData->GetMaterialPropertiesForSlot(Slot);
-
-		for (EDMMaterialPropertyType Property : Properties)
-		{
-			switch (Property)
-			{
-				case EDMMaterialPropertyType::BaseColor:
-				case EDMMaterialPropertyType::EmissiveColor:
-					return Slot->GetLayers().Num() > 1;
-
-				default:
-					// Nothing
-					break;
-			}
-		}
-	}
-
-	return true;
+	return Slot->CanRemoveLayer(LayerToRemove);
 }
 
 void SDMSlot::RemoveLayerByIndex(const int32 InLayerIndex, const bool bInSelectNextMaskIfMask)
@@ -1277,19 +1259,15 @@ UDMMaterialLayerObject* SDMSlot::AddNewLayer(UDMMaterialStage* InNewBaseStage, U
 
 		if (Slot->GetLayers().IsEmpty())
 		{
-			switch (ModelEditorOnlyData->GetShadingModel())
-			{
-				case EDMMaterialShadingModel::Unlit:
-					MaterialProperty = EDMMaterialPropertyType::EmissiveColor;
-					break;
-
-				case EDMMaterialShadingModel::DefaultLit:
-					MaterialProperty = EDMMaterialPropertyType::BaseColor;
-					break;
+			TArray<EDMMaterialPropertyType> SlotProperties = ModelEditorOnlyData->GetMaterialPropertiesForSlot(Slot);
 
-				default:
-					checkNoEntry();
-					break;
+			if (!SlotProperties.IsEmpty())
+			{
+				MaterialProperty = SlotProperties[0];
+			}
+			else
+			{
+				check("Cannot find material property.");
 			}
 		}
 		else
diff --git a/Engine/Plugins/Experimental/DynamicMaterial/Source/DynamicMaterialEditor/Public/Components/DMMaterialProperty.h b/Engine/Plugins/Experimental/DynamicMaterial/Source/DynamicMaterialEditor/Public/Components/DMMaterialProperty.h
--- a/Engine/Plugins/Experimental/DynamicMaterial/Source/DynamicMaterialEditor/Public/Components/DMMaterialProperty.h
+++ b/Engine/Plugins/Experimental/DynamicMaterial/Source/DynamicMaterialEditor/Public/Components/DMMaterialProperty.h
@@ -10,6 +10,7 @@
 
 class FText;
 class UDMMaterialSlot;
+class UDMMaterialStageInputExpression;
 class UDynamicMaterialModel;
 class UDynamicMaterialModelEditorOnlyData;
 class UMaterialExpression;
@@ -58,6 +59,8 @@ public:
 
 	virtual UMaterialExpression* GetDefaultInput(const TSharedRef<FDMMaterialBuildState>& InBuildState) const;
 
+	virtual void OnSlotAdded(UDMMaterialSlot* InSlot) {}
+
 	//~ Begin UDMMaterialComponent
 	virtual void Update(EDMUpdateType InUpdateType) override;
 	virtual bool IsRootComponent() const override { return true; }
diff --git a/Engine/Plugins/Experimental/DynamicMaterial/Source/DynamicMaterialEditor/Public/Components/DMMaterialSlot.h b/Engine/Plugins/Experimental/DynamicMaterial/Source/DynamicMaterialEditor/Public/Components/DMMaterialSlot.h
--- a/Engine/Plugins/Experimental/DynamicMaterial/Source/DynamicMaterialEditor/Public/Components/DMMaterialSlot.h
+++ b/Engine/Plugins/Experimental/DynamicMaterial/Source/DynamicMaterialEditor/Public/Components/DMMaterialSlot.h
@@ -73,6 +73,9 @@ public:
 	UFUNCTION(BlueprintCallable, Category = "Material Designer")
 	bool PasteLayer(UDMMaterialLayerObject* InLayer);
 
+	UFUNCTION(BlueprintCallable, Category = "Material Designer")
+	bool CanRemoveLayer(const UDMMaterialLayerObject* InLayer) const;
+
 	UFUNCTION(BlueprintCallable, Category = "Material Designer")
 	bool RemoveLayer(UDMMaterialLayerObject* InLayer);
 
@@ -115,8 +118,11 @@ public:
 	UMaterial* GetPreviewMaterial(EDMMaterialLayerStage InLayerStage);
 
 	UFUNCTION(BlueprintCallable, Category = "Material Designer")
-	bool SetLayerMaterialPropertyAndReplaceOthers(UDMMaterialLayerObject* InLayer, EDMMaterialPropertyType InMaterialProperty, 
-		EDMMaterialPropertyType InReplaceWithProperty);
+	bool SetLayerMaterialPropertyAndReplaceOthers(UDMMaterialLayerObject* InLayer, EDMMaterialPropertyType InPropertyFrom,
+		EDMMaterialPropertyType InPropertyTo);
+
+	UFUNCTION(BlueprintCallable, Category = "Material Designer")
+	bool ChangeMaterialProperty(EDMMaterialPropertyType InPropertyFrom, EDMMaterialPropertyType InPropertyTo);
 
 	//~ Begin UDMMaterialComponent
 	virtual void Update(EDMUpdateType InUpdateType) override;
diff --git a/Engine/Plugins/Experimental/DynamicMaterial/Source/DynamicMaterialEditor/Public/Components/MaterialProperties/DMMPBaseColor.h b/Engine/Plugins/Experimental/DynamicMaterial/Source/DynamicMaterialEditor/Public/Components/MaterialProperties/DMMPBaseColor.h
--- a/Engine/Plugins/Experimental/DynamicMaterial/Source/DynamicMaterialEditor/Public/Components/MaterialProperties/DMMPBaseColor.h
+++ b/Engine/Plugins/Experimental/DynamicMaterial/Source/DynamicMaterialEditor/Public/Components/MaterialProperties/DMMPBaseColor.h
@@ -13,10 +13,13 @@ class DYNAMICMATERIALEDITOR_API UDMMaterialPropertyBaseColor : public UDMMateria
 	GENERATED_BODY()
 
 public:
+	static void AddDefaultRGBLayer(EDMMaterialPropertyType InMaterialProperty, UDMMaterialSlot* InSlot);
+
 	UDMMaterialPropertyBaseColor();
 
 	//~ Begin UDMMaterialProperty
 	virtual bool IsValidForModel(UDynamicMaterialModelEditorOnlyData& InModelEditorOnlyData) const override;
 	virtual UMaterialExpression* GetDefaultInput(const TSharedRef<FDMMaterialBuildState>& InBuildState) const override;
+	virtual void OnSlotAdded(UDMMaterialSlot* InSlot) override;
 	//~ End UDMMaterialProperty
 };
diff --git a/Engine/Plugins/Experimental/DynamicMaterial/Source/DynamicMaterialEditor/Public/Components/MaterialProperties/DMMPEmissiveColor.h b/Engine/Plugins/Experimental/DynamicMaterial/Source/DynamicMaterialEditor/Public/Components/MaterialProperties/DMMPEmissiveColor.h
--- a/Engine/Plugins/Experimental/DynamicMaterial/Source/DynamicMaterialEditor/Public/Components/MaterialProperties/DMMPEmissiveColor.h
+++ b/Engine/Plugins/Experimental/DynamicMaterial/Source/DynamicMaterialEditor/Public/Components/MaterialProperties/DMMPEmissiveColor.h
@@ -18,5 +18,6 @@ public:
 	//~ Begin UDMMaterialProperty
 	virtual bool IsValidForModel(UDynamicMaterialModelEditorOnlyData& InModelEditorOnlyData) const override;
 	virtual UMaterialExpression* GetDefaultInput(const TSharedRef<FDMMaterialBuildState>& InBuildState) const override;
+	virtual void OnSlotAdded(UDMMaterialSlot* InSlot) override;
 	//~ End UDMMaterialProperty
 };
diff --git a/Engine/Plugins/Experimental/DynamicMaterial/Source/DynamicMaterialEditor/Public/Components/MaterialProperties/DMMPOpacity.h b/Engine/Plugins/Experimental/DynamicMaterial/Source/DynamicMaterialEditor/Public/Components/MaterialProperties/DMMPOpacity.h
--- a/Engine/Plugins/Experimental/DynamicMaterial/Source/DynamicMaterialEditor/Public/Components/MaterialProperties/DMMPOpacity.h
+++ b/Engine/Plugins/Experimental/DynamicMaterial/Source/DynamicMaterialEditor/Public/Components/MaterialProperties/DMMPOpacity.h
@@ -13,10 +13,13 @@ class DYNAMICMATERIALEDITOR_API UDMMaterialPropertyOpacity : public UDMMaterialP
 	GENERATED_BODY()
 
 public:
+	static void AddDefaultOpacityLayer(EDMMaterialPropertyType InMaterialProperty, UDMMaterialSlot* InSlot);
+
 	UDMMaterialPropertyOpacity();
 
 	//~ Begin UDMMaterialProperty
 	virtual bool IsValidForModel(UDynamicMaterialModelEditorOnlyData& InModelEditorOnlyData) const override;
 	virtual UMaterialExpression* GetDefaultInput(const TSharedRef<FDMMaterialBuildState>& InBuildState) const override;
+	virtual void OnSlotAdded(UDMMaterialSlot* InSlot) override;
 	//~ End UDMMaterialProperty
 };
diff --git a/Engine/Plugins/Experimental/DynamicMaterial/Source/DynamicMaterialEditor/Public/Components/MaterialProperties/DMMPOpacityMask.h b/Engine/Plugins/Experimental/DynamicMaterial/Source/DynamicMaterialEditor/Public/Components/MaterialProperties/DMMPOpacityMask.h
--- a/Engine/Plugins/Experimental/DynamicMaterial/Source/DynamicMaterialEditor/Public/Components/MaterialProperties/DMMPOpacityMask.h
+++ b/Engine/Plugins/Experimental/DynamicMaterial/Source/DynamicMaterialEditor/Public/Components/MaterialProperties/DMMPOpacityMask.h
@@ -18,5 +18,6 @@ public:
 	//~ Begin UDMMaterialProperty
 	virtual bool IsValidForModel(UDynamicMaterialModelEditorOnlyData& InModelEditorOnlyData) const override;
 	virtual UMaterialExpression* GetDefaultInput(const TSharedRef<FDMMaterialBuildState>& InBuildState) const override;
+	virtual void OnSlotAdded(UDMMaterialSlot* InSlot) override;
 	//~ End UDMMaterialProperty
 };
diff --git a/Engine/Plugins/Experimental/DynamicMaterial/Source/DynamicMaterialEditor/Public/Model/DynamicMaterialModelEditorOnlyData.h b/Engine/Plugins/Experimental/DynamicMaterial/Source/DynamicMaterialEditor/Public/Model/DynamicMaterialModelEditorOnlyData.h
--- a/Engine/Plugins/Experimental/DynamicMaterial/Source/DynamicMaterialEditor/Public/Model/DynamicMaterialModelEditorOnlyData.h
+++ b/Engine/Plugins/Experimental/DynamicMaterial/Source/DynamicMaterialEditor/Public/Model/DynamicMaterialModelEditorOnlyData.h
@@ -21,9 +21,10 @@ class UDMTextureUV;
 class UDynamicMaterialModel;
 class UDynamicMaterialModelEditorOnlyData;
 class UMaterialExpression;
-struct FDMComponentPathSegment;
 struct FDMComponentPath;
+struct FDMComponentPathSegment;
 struct FDMMaterialBuildState;
+struct FDMMaterialChannelListPreset;
 
 DECLARE_MULTICAST_DELEGATE_OneParam(FDMOnMaterialBuilt, UDynamicMaterialModel*);
 DECLARE_MULTICAST_DELEGATE_OneParam(FDMOnValueListUpdated, UDynamicMaterialModel*);
@@ -47,8 +48,22 @@ class DYNAMICMATERIALEDITOR_API UDynamicMaterialModelEditorOnlyData : public UOb
 
 public:
 	static const FString SlotsPathToken;
-	static const FString RGBSlotPathToken; // @TODO This will probably need to change when other slot types are opened up.
-	static const FString OpacitySlotPathToken; // @TODO This will probably need to change when other slot types are opened up.
+	static const FString RGBSlotPathToken;
+	static const FString OpacitySlotPathToken;
+	static const FString RoughnessPathToken;
+	static const FString SpecularPathToken;
+	static const FString MetallicPathToken;
+	static const FString NormalPathToken;
+	static const FString PixelDepthOffsetPathToken;
+	static const FString WorldPositionOffsetPathToken;
+	static const FString AmbientOcclusionPathToken;
+	static const FString AnisotropyPathToken;
+	static const FString RefractionPathToken;
+	static const FString TangentPathToken;
+	static const FString Custom1PathToken;
+	static const FString Custom2PathToken;
+	static const FString Custom3PathToken;
+	static const FString Custom4PathToken;
 	static const FString PropertiesPathToken;
 
 	static const TArray<EMaterialDomain> SupportedDomains;
@@ -61,6 +76,8 @@ public:
 
 	UDynamicMaterialModelEditorOnlyData();
 
+	void Initialize();
+
 	UFUNCTION(BlueprintPure, Category = "Material Designer")
 	UDynamicMaterialModel* GetMaterialModel() const { return MaterialModel; }
 
@@ -100,6 +117,12 @@ public:
 	UFUNCTION(BlueprintCallable, Category = "Material Designer")
 	void SetTwoSidedFlag(bool bInFlagValue);
 
+	UFUNCTION(BlueprintCallable, Category = "Material Designer")
+	FName GetChannelListPreset() const;
+
+	UFUNCTION(BlueprintCallable, Category = "Material Designer")
+	void SetChannelListPreset(FName InPresetName);
+
 	UFUNCTION(BlueprintCallable, Category = "Material Designer")
 	void OpenMaterialEditor() const;
 
@@ -112,17 +135,26 @@ public:
 	UFUNCTION(BlueprintPure, Category = "Material Designer")
 	const TArray<UDMMaterialSlot*>& GetSlots() const { return Slots; }
 
+	/** Gets slot by index. Highly recommended to use GetSlotForMaterialProperty(PropertyType). */
 	UFUNCTION(BlueprintPure, Category = "Material Designer")
 	UDMMaterialSlot* GetSlot(int32 Index) const;
 
+	UFUNCTION(BlueprintPure, Category = "Material Designer")
+	UDMMaterialSlot* GetSlotForMaterialProperty(EDMMaterialPropertyType InType) const;
+
+	/** Adds the next available slot. Highly recommended to use AddSlotForMaterialProperty(PropertyType). */
 	UFUNCTION(BlueprintCallable, Category = "Material Designer")
 	UDMMaterialSlot* AddSlot();
 
 	UFUNCTION(BlueprintCallable, Category = "Material Designer")
-	void RemoveSlot(int32 Index);
+	UDMMaterialSlot* AddSlotForMaterialProperty(EDMMaterialPropertyType InType);
 
-	UFUNCTION(BlueprintPure, Category = "Material Designer")
-	UDMMaterialSlot* GetSlotForMaterialProperty(EDMMaterialPropertyType Property) const;
+	/** Removes the next slot by index. Highly recommended to use RemoveSlotForMaterialProperty(PropertyType). */
+	UFUNCTION(BlueprintCallable, Category = "Material Designer")
+	UDMMaterialSlot* RemoveSlot(int32 Index);
+
+	UFUNCTION(BlueprintCallable, Category = "Material Designer")
+	UDMMaterialSlot* RemoveSlotForMaterialProperty(EDMMaterialPropertyType InType);
 
 	UFUNCTION(BlueprintPure, Category = "Material Designer")
 	TArray<EDMMaterialPropertyType> GetMaterialPropertiesForSlot(UDMMaterialSlot* Slot) const;
@@ -199,6 +231,9 @@ protected:
 	UPROPERTY(VisibleInstanceOnly, BlueprintReadOnly, Category = "Material Designer")
 	bool bTwoSidedFlag;
 
+	UPROPERTY(VisibleInstanceOnly, BlueprintReadOnly, Category = "Material Designer")
+	FName ChannelListPreset;
+
 	UPROPERTY(VisibleInstanceOnly, BlueprintReadOnly, Instanced, Category = "Material Designer")
 	TMap<EDMMaterialPropertyType, TObjectPtr<UDMMaterialProperty>> Properties;
 
@@ -238,6 +273,15 @@ protected:
 
 	void OnSlotConnectorsUpdated(UDMMaterialSlot* Slot);
 
+	/** Swaps the material properties from one slot to another, unless both slots exist and/or are the same. */
+	void SwapSlotMaterialProperty(EDMMaterialPropertyType InPropertyFrom, EDMMaterialPropertyType InPropertyTo);
+
+	/** 
+	 * Swaps the material properties from one slot to another, unless both slots exist and/or are the same. 
+	 * Ensuring that the To Property exists.
+	 */
+	void EnsureSwapSlotMaterialProperty(EDMMaterialPropertyType InPropertyFrom, EDMMaterialPropertyType InPropertyTo);
+
 	//~ Begin IDynamicMaterialModelEditorOnlyDataInterface
 	virtual void ReinitComponents() override;
 	virtual void ResetData() override;
