Material Designer
- Slight redesign of the material wizard.

diff --git a/Engine/Plugins/Experimental/DynamicMaterial/Source/DynamicMaterialEditor/Private/Slate/SDMMaterialWizard.cpp b/Engine/Plugins/Experimental/DynamicMaterial/Source/DynamicMaterialEditor/Private/Slate/SDMMaterialWizard.cpp
--- a/Engine/Plugins/Experimental/DynamicMaterial/Source/DynamicMaterialEditor/Private/Slate/SDMMaterialWizard.cpp
+++ b/Engine/Plugins/Experimental/DynamicMaterial/Source/DynamicMaterialEditor/Private/Slate/SDMMaterialWizard.cpp
@@ -12,10 +12,10 @@
 #include "SDMEditor.h"
 #include "Styling/SlateTypes.h"
 #include "Styling/StyleColors.h"
-#include "Widgets/Colors/SColorBlock.h"
 #include "Widgets/Input/SButton.h"
 #include "Widgets/Input/SCheckBox.h"
 #include "Widgets/Layout/SBox.h"
+#include "Widgets/Layout/SExpandableArea.h"
 #include "Widgets/Layout/SWrapBox.h"
 #include "Widgets/SBoxPanel.h"
 #include "Widgets/Text/STextBlock.h"
@@ -121,104 +121,101 @@ TSharedRef<SWidget> SDMMaterialWizard::CreateLayout()
 			.HAlign(HAlign_Fill)
 			.Padding(0.0f, SeparationDistance, 0.0f, 0.f)
 			[
-				SNew(SBox)
-				.HeightOverride(1.f)
+				SNew(SExpandableArea)
+				.BorderImage(FAppStyle::GetBrush(TEXT("Menu.Background")))
+				.BodyBorderImage(FAppStyle::GetBrush("DetailsView.CategoryTop"))
+				.HeaderContent()
 				[
-					SNew(SColorBlock)
-					.Color(FStyleColors::Foreground.GetSpecifiedColor().ToFColor(false))
+					SNew(STextBlock)
+					.Text(LOCTEXT("AdvancedSettings", "Advanced Settings"))
 				]
-			]
-
-			+ SVerticalBox::Slot()
-			.AutoHeight()
-			.HAlign(HAlign_Left)
-			.Padding(0.0f, SeparationDistance, 0.0f, TitleContentDistance)
-			[
-				SNew(STextBlock)
-				.Text(LOCTEXT("MaterialDomain", "Material Domain"))
-			]
+				.BodyContent()
+				[
+					SNew(SBox)
+					.Padding(FMargin(SeparationDistance, SeparationDistance, SeparationDistance, SeparationDistance))
+					[
+						SNew(SVerticalBox)
 
-			+ SVerticalBox::Slot()
-			.AutoHeight()
-			.HAlign(EHorizontalAlignment::HAlign_Fill)
-			[
-				CreateMaterialDomainOptions()
-			]
+						+ SVerticalBox::Slot()
+						.AutoHeight()
+						.HAlign(HAlign_Left)
+						.Padding(0.0f, 0.f, 0.0f, TitleContentDistance)
+						[
+							SNew(STextBlock)
+							.Text(LOCTEXT("MaterialDomain", "Material Domain"))
+						]
 
-			+ SVerticalBox::Slot()
-			.AutoHeight()
-			.HAlign(HAlign_Left)
-			.Padding(0.0f, SeparationDistance, 0.0f, TitleContentDistance)
-			[
-				SNew(STextBlock)
-				.Text(LOCTEXT("BlendMode", "Blend Mode"))
-			]
+						+ SVerticalBox::Slot()
+						.AutoHeight()
+						.HAlign(EHorizontalAlignment::HAlign_Fill)
+						[
+							CreateMaterialDomainOptions()
+						]
 
-			+ SVerticalBox::Slot()
-			.AutoHeight()
-			.HAlign(EHorizontalAlignment::HAlign_Fill)
-			[
-				CreateBlendModeOptions()
-			]
+						+ SVerticalBox::Slot()
+						.AutoHeight()
+						.HAlign(HAlign_Left)
+						.Padding(0.0f, SeparationDistance, 0.0f, TitleContentDistance)
+						[
+							SNew(STextBlock)
+							.Text(LOCTEXT("BlendMode", "Blend Mode"))
+						]
 
-			+ SVerticalBox::Slot()
-			.AutoHeight()
-			.HAlign(HAlign_Left)
-			.Padding(0.0f, SeparationDistance, 0.0f, TitleContentDistance)
-			[
-				SNew(STextBlock)
-				.Text(LOCTEXT("ShadingModel", "Shading Model"))
-			]
+						+ SVerticalBox::Slot()
+						.AutoHeight()
+						.HAlign(EHorizontalAlignment::HAlign_Fill)
+						[
+							CreateBlendModeOptions()
+						]
 
-			+ SVerticalBox::Slot()
-			.AutoHeight()
-			.HAlign(EHorizontalAlignment::HAlign_Fill)
-			[
-				CreateShadingModelOptions()
-			]
+						+ SVerticalBox::Slot()
+						.AutoHeight()
+						.HAlign(HAlign_Left)
+						.Padding(0.0f, SeparationDistance, 0.0f, TitleContentDistance)
+						[
+							SNew(STextBlock)
+							.Text(LOCTEXT("ShadingModel", "Shading Model"))
+						]
 
-			+ SVerticalBox::Slot()
-			.AutoHeight()
-			.HAlign(HAlign_Left)
-			.Padding(0.0f, SeparationDistance, 0.0f, TitleContentDistance)
-			[
-				SNew(STextBlock)
-				.Text(LOCTEXT("MotionVectors", "Motion Vectors"))
-			]
+						+ SVerticalBox::Slot()
+						.AutoHeight()
+						.HAlign(EHorizontalAlignment::HAlign_Fill)
+						[
+							CreateShadingModelOptions()
+						]
 
-			+ SVerticalBox::Slot()
-			.AutoHeight()
-			.HAlign(EHorizontalAlignment::HAlign_Fill)
-			[
-				CreateAnimationOptions()
-			]
+						+ SVerticalBox::Slot()
+						.AutoHeight()
+						.HAlign(HAlign_Left)
+						.Padding(0.0f, SeparationDistance, 0.0f, TitleContentDistance)
+						[
+							SNew(STextBlock)
+							.Text(LOCTEXT("MotionVectors", "Motion Vectors"))
+						]
 
-			+ SVerticalBox::Slot()
-			.AutoHeight()
-			.HAlign(HAlign_Left)
-			.Padding(0.0f, SeparationDistance, 0.0f, TitleContentDistance)
-			[
-				SNew(STextBlock)
-				.Text(LOCTEXT("Geometry", "Geometry"))
-			]
+						+ SVerticalBox::Slot()
+						.AutoHeight()
+						.HAlign(EHorizontalAlignment::HAlign_Fill)
+						[
+							CreateAnimationOptions()
+						]
 
-			+ SVerticalBox::Slot()
-			.AutoHeight()
-			.HAlign(EHorizontalAlignment::HAlign_Fill)
-			[
-				CreateTwoSidedOptions()
-			]
+						+ SVerticalBox::Slot()
+						.AutoHeight()
+						.HAlign(HAlign_Left)
+						.Padding(0.0f, SeparationDistance, 0.0f, TitleContentDistance)
+						[
+							SNew(STextBlock)
+							.Text(LOCTEXT("Geometry", "Geometry"))
+						]
 
-			+ SVerticalBox::Slot()
-			.AutoHeight()
-			.HAlign(HAlign_Fill)
-			.Padding(0.0f, SeparationDistance, 0.0f, 0.f)
-			[
-				SNew(SBox)
-				.HeightOverride(1.f)
-				[
-					SNew(SColorBlock)
-					.Color(FStyleColors::Foreground.GetSpecifiedColor().ToFColor(false))
+						+ SVerticalBox::Slot()
+						.AutoHeight()
+						.HAlign(EHorizontalAlignment::HAlign_Fill)
+						[
+							CreateTwoSidedOptions()
+						]
+					]
 				]
 			]
 
@@ -300,30 +297,40 @@ TSharedRef<SWidget> SDMMaterialWizard::CreateShadingModelOptions()
 {
 	using namespace UE::DynamicMaterialEditor::Private;
 
-	UEnum* MaterialShadingModelEnum = StaticEnum<EMaterialShadingModel>();
-
 	TSharedRef<SWrapBox> MaterialShadingModelOptions = SNew(SWrapBox)
 		.UseAllottedSize(true)
 		.InnerSlotPadding(WrapBoxSlotPadding)
 		.Orientation(EOrientation::Orient_Horizontal);
 
-	for (EDMMaterialShadingModel MaterialShadingModel : {EDMMaterialShadingModel::Unlit, EDMMaterialShadingModel::DefaultLit})
-	{
-		MaterialShadingModelOptions->AddSlot()
+	MaterialShadingModelOptions->AddSlot()
+		[
+			SNew(SCheckBox)
+			.Style(FAppStyle::Get(), "DetailsView.SectionButton")
+			.HAlign(EHorizontalAlignment::HAlign_Center)
+			.Padding(ButtonPadding)
+			.IsEnabled(this, &SDMMaterialWizard::Unlit_IsEnabled)
+			.IsChecked(this, &SDMMaterialWizard::Unlit_GetState, EDMMaterialShadingModel::Unlit)
+			.OnCheckStateChanged(this, &SDMMaterialWizard::Unlit_OnChange, EDMMaterialShadingModel::Unlit)
 			[
-				SNew(SCheckBox)
-				.Style(FAppStyle::Get(), "DetailsView.SectionButton")
-				.HAlign(EHorizontalAlignment::HAlign_Center)
-				.Padding(ButtonPadding)
-				.IsEnabled(this, &SDMMaterialWizard::Unlit_IsEnabled)
-				.IsChecked(this, &SDMMaterialWizard::Unlit_GetState, MaterialShadingModel)
-				.OnCheckStateChanged(this, &SDMMaterialWizard::Unlit_OnChange, MaterialShadingModel)
-				[
-					SNew(STextBlock)
-					.Text(MaterialShadingModelEnum->GetDisplayNameTextByValue(static_cast<int64>(MaterialShadingModel)))
-				]
-			];
-	}
+				SNew(STextBlock)
+				.Text(LOCTEXT("Emissive", "Emissive"))
+			]
+		];
+
+	MaterialShadingModelOptions->AddSlot()
+		[
+			SNew(SCheckBox)
+			.Style(FAppStyle::Get(), "DetailsView.SectionButton")
+			.HAlign(EHorizontalAlignment::HAlign_Center)
+			.Padding(ButtonPadding)
+			.IsEnabled(this, &SDMMaterialWizard::Unlit_IsEnabled)
+			.IsChecked(this, &SDMMaterialWizard::Unlit_GetState, EDMMaterialShadingModel::DefaultLit)
+			.OnCheckStateChanged(this, &SDMMaterialWizard::Unlit_OnChange, EDMMaterialShadingModel::DefaultLit)
+			[
+				SNew(STextBlock)
+				.Text(LOCTEXT("DefaultLit", "Default Lit"))
+			]
+		];
 
 	return MaterialShadingModelOptions;
 }
@@ -344,11 +351,11 @@ TSharedRef<SWidget> SDMMaterialWizard::CreateAnimationOptions()
 				.HAlign(EHorizontalAlignment::HAlign_Center)
 				.Padding(ButtonPadding)
 				.IsEnabled(this, &SDMMaterialWizard::Animated_IsEnabled)
-				.IsChecked(this, &SDMMaterialWizard::Animated_GetState, true)
-				.OnCheckStateChanged(this, &SDMMaterialWizard::Animated_OnChange, true)
+				.IsChecked(this, &SDMMaterialWizard::Animated_GetState, false)
+				.OnCheckStateChanged(this, &SDMMaterialWizard::Animated_OnChange, false)
 				[
 					SNew(STextBlock)
-					.Text(LOCTEXT("Animated", "Animated"))
+					.Text(LOCTEXT("Static", "Static"))
 				]
 		];
 
@@ -359,11 +366,11 @@ TSharedRef<SWidget> SDMMaterialWizard::CreateAnimationOptions()
 				.HAlign(EHorizontalAlignment::HAlign_Center)
 				.Padding(ButtonPadding)
 				.IsEnabled(this, &SDMMaterialWizard::Animated_IsEnabled)
-				.IsChecked(this, &SDMMaterialWizard::Animated_GetState, false)
-				.OnCheckStateChanged(this, &SDMMaterialWizard::Animated_OnChange, false)
+				.IsChecked(this, &SDMMaterialWizard::Animated_GetState, true)
+				.OnCheckStateChanged(this, &SDMMaterialWizard::Animated_OnChange, true)
 				[
 					SNew(STextBlock)
-					.Text(LOCTEXT("Static", "Static"))
+					.Text(LOCTEXT("Animated", "Animated"))
 				]
 		];
 
@@ -386,11 +393,11 @@ TSharedRef<SWidget> SDMMaterialWizard::CreateTwoSidedOptions()
 				.HAlign(EHorizontalAlignment::HAlign_Center)
 				.Padding(ButtonPadding)
 				.IsEnabled(this, &SDMMaterialWizard::TwoSided_IsEnabled)
-				.IsChecked(this, &SDMMaterialWizard::TwoSided_GetState, true)
-				.OnCheckStateChanged(this, &SDMMaterialWizard::TwoSided_OnChange, true)
+				.IsChecked(this, &SDMMaterialWizard::TwoSided_GetState, false)
+				.OnCheckStateChanged(this, &SDMMaterialWizard::TwoSided_OnChange, false)
 				[
 					SNew(STextBlock)
-					.Text(LOCTEXT("TwoSided", "Two Sided"))
+					.Text(LOCTEXT("OneSided", "One Sided"))
 				]
 		];
 
@@ -401,11 +408,11 @@ TSharedRef<SWidget> SDMMaterialWizard::CreateTwoSidedOptions()
 				.HAlign(EHorizontalAlignment::HAlign_Center)
 				.Padding(ButtonPadding)
 				.IsEnabled(this, &SDMMaterialWizard::TwoSided_IsEnabled)
-				.IsChecked(this, &SDMMaterialWizard::TwoSided_GetState, false)
-				.OnCheckStateChanged(this, &SDMMaterialWizard::TwoSided_OnChange, false)
+				.IsChecked(this, &SDMMaterialWizard::TwoSided_GetState, true)
+				.OnCheckStateChanged(this, &SDMMaterialWizard::TwoSided_OnChange, true)
 				[
 					SNew(STextBlock)
-					.Text(LOCTEXT("OneSided", "One Sided"))
+					.Text(LOCTEXT("TwoSided", "Two Sided"))
 				]
 		];
 
