Material Designer
- Fix for content-browser based material creation

diff --git a/Engine/Plugins/Experimental/DynamicMaterial/Source/DynamicMaterialEditor/Private/Model/DynamicMaterialModelEditorOnlyData.cpp b/Engine/Plugins/Experimental/DynamicMaterial/Source/DynamicMaterialEditor/Private/Model/DynamicMaterialModelEditorOnlyData.cpp
--- a/Engine/Plugins/Experimental/DynamicMaterial/Source/DynamicMaterialEditor/Private/Model/DynamicMaterialModelEditorOnlyData.cpp
+++ b/Engine/Plugins/Experimental/DynamicMaterial/Source/DynamicMaterialEditor/Private/Model/DynamicMaterialModelEditorOnlyData.cpp
@@ -172,16 +172,13 @@ UDynamicMaterialModelEditorOnlyData::UDynamicMaterialModelEditorOnlyData()
 
 void UDynamicMaterialModelEditorOnlyData::Initialize()
 {
-	if (Slots.IsEmpty())
+	if (!Slots.IsEmpty())
 	{
 		return;
 	}
 
 	// Will choose appropriate type between BaseColor and EmissiveColor
 	AddSlotForMaterialProperty(EDMMaterialPropertyType::BaseColor);
-
-	// Will choose appropriate type between Opacity and OpacityMask
-	AddSlotForMaterialProperty(EDMMaterialPropertyType::Opacity);
 }
 
 UMaterial* UDynamicMaterialModelEditorOnlyData::GetGeneratedMaterial() const
