Material Designer
- Change to how the wizard works. It will now popup based on whether a channel list preset has been selected or not.

diff --git a/Engine/Plugins/Experimental/DynamicMaterial/Source/DynamicMaterialEditor/Private/Model/DynamicMaterialModelEditorOnlyData.cpp b/Engine/Plugins/Experimental/DynamicMaterial/Source/DynamicMaterialEditor/Private/Model/DynamicMaterialModelEditorOnlyData.cpp
--- a/Engine/Plugins/Experimental/DynamicMaterial/Source/DynamicMaterialEditor/Private/Model/DynamicMaterialModelEditorOnlyData.cpp
+++ b/Engine/Plugins/Experimental/DynamicMaterial/Source/DynamicMaterialEditor/Private/Model/DynamicMaterialModelEditorOnlyData.cpp
@@ -141,7 +141,6 @@ UDynamicMaterialModelEditorOnlyData::UDynamicMaterialModelEditorOnlyData()
 	ShadingModel = EDMMaterialShadingModel::Unlit;
 	bPixelAnimationFlag = false;
 	bTwoSidedFlag = true;
-	bUseWizard = true;
 	ChannelListPreset = NAME_None;
 
 	Properties.Emplace(EDMMaterialPropertyType::BaseColor,           CreateDefaultSubobject<UDMMaterialPropertyBaseColor>(          "MaterialProperty_BaseColor"));
@@ -669,12 +668,11 @@ TSharedRef<FDMMaterialBuildState> UDynamicMaterialModelEditorOnlyData::CreateBui
 
 bool UDynamicMaterialModelEditorOnlyData::NeedsWizard() const
 {
-	return bUseWizard;
+	return ChannelListPreset == NAME_None;
 }
 
 void UDynamicMaterialModelEditorOnlyData::OnWizardComplete()
 {
-	bUseWizard = false;
 }
 
 void UDynamicMaterialModelEditorOnlyData::SetChannelListPreset(FName InPresetName)
@@ -1312,8 +1310,6 @@ void UDynamicMaterialModelEditorOnlyData::PostLoad()
 	SetFlags(RF_Transactional);
 
 	ReinitComponents();
-
-	bUseWizard = false;
 }
 
 void UDynamicMaterialModelEditorOnlyData::PostEditUndo()
diff --git a/Engine/Plugins/Experimental/DynamicMaterial/Source/DynamicMaterialEditor/Public/Model/DynamicMaterialModelEditorOnlyData.h b/Engine/Plugins/Experimental/DynamicMaterial/Source/DynamicMaterialEditor/Public/Model/DynamicMaterialModelEditorOnlyData.h
--- a/Engine/Plugins/Experimental/DynamicMaterial/Source/DynamicMaterialEditor/Public/Model/DynamicMaterialModelEditorOnlyData.h
+++ b/Engine/Plugins/Experimental/DynamicMaterial/Source/DynamicMaterialEditor/Public/Model/DynamicMaterialModelEditorOnlyData.h
@@ -249,13 +249,6 @@ protected:
 	UPROPERTY(VisibleInstanceOnly, BlueprintReadOnly, Category = "Material Designer")
 	bool bCreateMaterialPackage;
 
-	/**
-	 * This is just a quick, transient flag that marks this model as needing to
-	 * open the wizard when it is opened by the material designer. It is initially
-	 * set to true, but set to false on first opening and on post load.
-	 */
-	bool bUseWizard;
-
 	FDMOnMaterialBuilt OnMaterialBuiltDelegate;
 	FDMOnValueListUpdated OnValueListUpdateDelegate;
 	FDMOnValueUpdated OnValueUpdateDelegate;
