Material Designer
- Material widget now display channel names as just text, not disabled buttons.
- Changed the order of the channel list to reflect the order of the channels in the editor only data.

diff --git a/Engine/Plugins/Experimental/DynamicMaterial/Source/DynamicMaterialEditor/Private/Slate/SDMMaterialWizard.cpp b/Engine/Plugins/Experimental/DynamicMaterial/Source/DynamicMaterialEditor/Private/Slate/SDMMaterialWizard.cpp
--- a/Engine/Plugins/Experimental/DynamicMaterial/Source/DynamicMaterialEditor/Private/Slate/SDMMaterialWizard.cpp
+++ b/Engine/Plugins/Experimental/DynamicMaterial/Source/DynamicMaterialEditor/Private/Slate/SDMMaterialWizard.cpp
@@ -27,6 +27,7 @@ namespace UE::DynamicMaterialEditor::Private
 	constexpr float SeparationDistance = 20.f;
 	constexpr float TitleContentDistance = 5.f;
 	static const FMargin ButtonPadding = FMargin(10.f, 5.f);
+	static const FMargin TextPadding = FMargin(5.f, 2.f);
 	static const FVector2D WrapBoxSlotPadding = FVector2D(5, 5);
 }
 
@@ -452,6 +453,13 @@ TSharedRef<SWidget> SDMMaterialWizard::CreateChannelList()
 {
 	using namespace UE::DynamicMaterialEditor::Private;
 
+	UDynamicMaterialModelEditorOnlyData* ModelEditorOnlyData = UDynamicMaterialModelEditorOnlyData::Get(GetMaterialModel());
+
+	if (!ModelEditorOnlyData)
+	{
+		return SNullWidget::NullWidget;
+	}
+
 	UEnum* MaterialPropertyEnum = StaticEnum<EDMMaterialPropertyType>();
 
 	TSharedRef<SWrapBox> ChannelPresets = SNew(SWrapBox)
@@ -459,36 +467,26 @@ TSharedRef<SWidget> SDMMaterialWizard::CreateChannelList()
 		.InnerSlotPadding(WrapBoxSlotPadding)
 		.Orientation(EOrientation::Orient_Horizontal);
 
+
 	if (const FDMMaterialChannelListPreset* Preset = GetDefault<UDynamicMaterialEditorSettings>()->ChannelPresets.Find(CurrentPreset))
 	{
-		for (uint8 PropertyIndex = static_cast<uint8>(EDMMaterialPropertyType::None) + 1;
-			PropertyIndex < static_cast<uint8>(EDMMaterialPropertyType::Any);
-			++PropertyIndex)
+		for (const TPair<EDMMaterialPropertyType, UDMMaterialProperty*>& Property : ModelEditorOnlyData->GetMaterialProperties())
 		{
-			const EDMMaterialPropertyType Property = static_cast<EDMMaterialPropertyType>(PropertyIndex);
-
-			if (Property == EDMMaterialPropertyType::EmissiveColor || Property == EDMMaterialPropertyType::OpacityMask)
+			if (Property.Key == EDMMaterialPropertyType::EmissiveColor || Property.Key == EDMMaterialPropertyType::OpacityMask)
 			{
 				continue;
 			}
 				
-			if (Preset->IsPropertyEnabled(Property))
+			if (Preset->IsPropertyEnabled(Property.Key))
 			{
 				constexpr const TCHAR* ShortNameName = TEXT("ShortName");
-				const FString ShortName = MaterialPropertyEnum->GetMetaData(ShortNameName, MaterialPropertyEnum->GetIndexByValue(static_cast<int64>(PropertyIndex)));
+				const FString ShortName = MaterialPropertyEnum->GetMetaData(ShortNameName, MaterialPropertyEnum->GetIndexByValue(static_cast<int64>(Property.Key)));
 
 				ChannelPresets->AddSlot()
+					.Padding(TextPadding)
 					[
-						SNew(SCheckBox)
-						.Style(FAppStyle::Get(), "DetailsView.SectionButton")
-						.HAlign(EHorizontalAlignment::HAlign_Center)
-						.Padding(ButtonPadding)
-						.IsChecked(false)
-						.IsEnabled(false)
-						[
-							SNew(STextBlock)
-							.Text(!ShortName.IsEmpty() ? FText::FromString(ShortName) : MaterialPropertyEnum->GetDisplayNameTextByValue(static_cast<int64>(Property)))
-						]
+						SNew(STextBlock)
+						.Text(!ShortName.IsEmpty() ? FText::FromString(ShortName) : MaterialPropertyEnum->GetDisplayNameTextByValue(static_cast<int64>(Property.Key)))
 					];
 			}
 		}
