Material Designer
- Removed advanced material settings from wizard.
- Added good defaults to the material type defaults.
- Some fixes related to material property selection.
- Changed slot button generation to use the order of the properties in the material.
- Changed material property order (swapped PDO and Refraction)
- Changed base and emissive colour channels to not use the 'RGB' short name.

diff --git a/Engine/Plugins/Experimental/DynamicMaterial/Source/DynamicMaterial/Public/DMDefs.h b/Engine/Plugins/Experimental/DynamicMaterial/Source/DynamicMaterial/Public/DMDefs.h
--- a/Engine/Plugins/Experimental/DynamicMaterial/Source/DynamicMaterial/Public/DMDefs.h
+++ b/Engine/Plugins/Experimental/DynamicMaterial/Source/DynamicMaterial/Public/DMDefs.h
@@ -15,8 +15,8 @@ UENUM(BlueprintType)
 enum class EDMMaterialPropertyType : uint8
 {
 	None = 0,
-	BaseColor UMETA(ShortName = "RGB"),
-	EmissiveColor UMETA(ShortName = "RGB"),
+	BaseColor UMETA(ShortName = "Base Color"),
+	EmissiveColor UMETA(ShortName = "Emissive"),
 	Opacity,
 	OpacityMask UMETA(ShortName = "Mask"),
 	Roughness,
diff --git a/Engine/Plugins/Experimental/DynamicMaterial/Source/DynamicMaterialEditor/Private/Components/MaterialProperties/DMMPEmissiveColor.cpp b/Engine/Plugins/Experimental/DynamicMaterial/Source/DynamicMaterialEditor/Private/Components/MaterialProperties/DMMPEmissiveColor.cpp
--- a/Engine/Plugins/Experimental/DynamicMaterial/Source/DynamicMaterialEditor/Private/Components/MaterialProperties/DMMPEmissiveColor.cpp
+++ b/Engine/Plugins/Experimental/DynamicMaterial/Source/DynamicMaterialEditor/Private/Components/MaterialProperties/DMMPEmissiveColor.cpp
@@ -32,5 +32,5 @@ UMaterialExpression* UDMMaterialPropertyEmissiveColor::GetDefaultInput(
 
 void UDMMaterialPropertyEmissiveColor::OnSlotAdded(UDMMaterialSlot* InSlot)
 {
-	UDMMaterialPropertyBaseColor::AddDefaultRGBLayer(EDMMaterialPropertyType::BaseColor, InSlot);
+	UDMMaterialPropertyBaseColor::AddDefaultRGBLayer(EDMMaterialPropertyType::EmissiveColor, InSlot);
 }
diff --git a/Engine/Plugins/Experimental/DynamicMaterial/Source/DynamicMaterialEditor/Private/DynamicMaterialEditorModule.cpp b/Engine/Plugins/Experimental/DynamicMaterial/Source/DynamicMaterialEditor/Private/DynamicMaterialEditorModule.cpp
--- a/Engine/Plugins/Experimental/DynamicMaterial/Source/DynamicMaterialEditor/Private/DynamicMaterialEditorModule.cpp
+++ b/Engine/Plugins/Experimental/DynamicMaterial/Source/DynamicMaterialEditor/Private/DynamicMaterialEditorModule.cpp
@@ -36,7 +36,7 @@
 #include "Model/DynamicMaterialModelEditorOnlyData.h"
 #include "Modules/ModuleManager.h"
 #include "PropertyEditorModule.h"
-#include "Model/DMModelCreatedCallback.h"
+#include "Model/DMOnWizardCompleteCallback.h"
 #include "Slate/Properties/Editors/SDMPropertyEditBoolValue.h"
 #include "Slate/Properties/Editors/SDMPropertyEditFloat1Value.h"
 #include "Slate/Properties/Editors/SDMPropertyEditFloat2Value.h"
@@ -82,7 +82,7 @@ TMap<UClass*, FDMCreateValueEditWidgetDelegate> FDynamicMaterialEditorModule::Va
 TMap<UClass*, FDMComponentPropertyRowGeneratorDelegate> FDynamicMaterialEditorModule::ComponentPropertyRowGenerators;
 TMap<UClass*, FDMGetObjectMaterialPropertiesDelegate> FDynamicMaterialEditorModule::CustomMaterialPropertyGenerators;
 FDMOnUIValueUpdate FDynamicMaterialEditorModule::OnUIValueUpdate;
-TArray<TSharedRef<IDMMaterialModelCreatedCallback>> FDynamicMaterialEditorModule::OnCreatedCallbacks;
+TArray<TSharedRef<IDMOnWizardCompleteCallback>> FDynamicMaterialEditorModule::OnWizardCompleteCallbacks;
 
 void FDynamicMaterialEditorModule::RegisterValueEditWidgetDelegate(UClass* InClass, FDMCreateValueEditWidgetDelegate ValueEditBodyDelegate)
 {
@@ -187,24 +187,24 @@ void FDynamicMaterialEditorModule::RegisterCustomMaterialPropertyGenerator(UClas
 	CustomMaterialPropertyGenerators.FindOrAdd(InClass) = InGenerator;
 }
 
-void FDynamicMaterialEditorModule::RegisterMaterialModelCreatedCallback(const TSharedRef<IDMMaterialModelCreatedCallback> InCallback)
+void FDynamicMaterialEditorModule::RegisterMaterialModelCreatedCallback(const TSharedRef<IDMOnWizardCompleteCallback> InCallback)
 {
-	OnCreatedCallbacks.Add(InCallback);
+	OnWizardCompleteCallbacks.Add(InCallback);
 
-	OnCreatedCallbacks.StableSort(
-		[](const TSharedRef<IDMMaterialModelCreatedCallback>& InA, const TSharedRef<IDMMaterialModelCreatedCallback>& InB)
+	OnWizardCompleteCallbacks.StableSort(
+		[](const TSharedRef<IDMOnWizardCompleteCallback>& InA, const TSharedRef<IDMOnWizardCompleteCallback>& InB)
 		{
 			return InA.Get() < InB.Get();
 		}
 	);
 }
 
-void FDynamicMaterialEditorModule::UnregisterMaterialModelCreatedCallback(const TSharedRef<IDMMaterialModelCreatedCallback> InCallback)
+void FDynamicMaterialEditorModule::UnregisterMaterialModelCreatedCallback(const TSharedRef<IDMOnWizardCompleteCallback> InCallback)
 {
-	OnCreatedCallbacks.Remove(InCallback);
+	OnWizardCompleteCallbacks.Remove(InCallback);
 }
 
-void FDynamicMaterialEditorModule::OnMaterialModelCreated(UDynamicMaterialModel* InModel)
+void FDynamicMaterialEditorModule::OnWizardComplete(UDynamicMaterialModel* InModel)
 {
 	if (!IsValid(InModel))
 	{
@@ -216,7 +216,7 @@ void FDynamicMaterialEditorModule::OnMaterialModelCreated(UDynamicMaterialModel*
 	UActorComponent* OuterComponent = InModel->GetTypedOuter<UActorComponent>();
 	AActor* OuterActor = InModel->GetTypedOuter<AActor>();
 
-	const FDMMaterialModelCreatedCallbackParams Params = {
+	const FDMOnWizardCompleteCallbackParams Params = {
 		InModel,
 		EditorOnlyData,
 		Outer,
@@ -224,9 +224,9 @@ void FDynamicMaterialEditorModule::OnMaterialModelCreated(UDynamicMaterialModel*
 		OuterActor
 	};
 
-	for (const TSharedRef<IDMMaterialModelCreatedCallback>& OnCreatedCallback : OnCreatedCallbacks)
+	for (const TSharedRef<IDMOnWizardCompleteCallback>& OnWizardCompleteCallback : OnWizardCompleteCallbacks)
 	{
-		OnCreatedCallback->OnModelCreated(Params);
+		OnWizardCompleteCallback->OnModelCreated(Params);
 	}
 }
 
diff --git a/Engine/Plugins/Experimental/DynamicMaterial/Source/DynamicMaterialEditor/Private/DynamicMaterialEditorModule.h b/Engine/Plugins/Experimental/DynamicMaterial/Source/DynamicMaterialEditor/Private/DynamicMaterialEditorModule.h
--- a/Engine/Plugins/Experimental/DynamicMaterial/Source/DynamicMaterialEditor/Private/DynamicMaterialEditorModule.h
+++ b/Engine/Plugins/Experimental/DynamicMaterial/Source/DynamicMaterialEditor/Private/DynamicMaterialEditorModule.h
@@ -90,8 +90,8 @@ public:
 	//~ Begin IDynamicMaterialEditorModule
 	virtual void OpenEditor(UWorld* InWorld) override;
 	virtual void RegisterCustomMaterialPropertyGenerator(UClass* InClass, FDMGetObjectMaterialPropertiesDelegate InGenerator) override;
-	virtual void RegisterMaterialModelCreatedCallback(const TSharedRef<IDMMaterialModelCreatedCallback> InCallback)  override;
-	virtual void UnregisterMaterialModelCreatedCallback(const TSharedRef<IDMMaterialModelCreatedCallback> InCallback) override;
+	virtual void RegisterMaterialModelCreatedCallback(const TSharedRef<IDMOnWizardCompleteCallback> InCallback)  override;
+	virtual void UnregisterMaterialModelCreatedCallback(const TSharedRef<IDMOnWizardCompleteCallback> InCallback) override;
 	//~ End IDynamicMaterialEditorModule
 
 	//~ Begin IModuleInterface
@@ -119,14 +119,14 @@ public:
 
 	const TSharedRef<FUICommandList>& GetCommandList() const { return CommandList; }
 
-	void OnMaterialModelCreated(UDynamicMaterialModel* InModel);
+	void OnWizardComplete(UDynamicMaterialModel* InModel);
 
 protected:
 	static TMap<UClass*, FDMCreateValueEditWidgetDelegate> ValueEditWidgetDelegates;
 	static TMap<UClass*, FDMComponentPropertyRowGeneratorDelegate> ComponentPropertyRowGenerators;
 	static TMap<UClass*, FDMGetObjectMaterialPropertiesDelegate> CustomMaterialPropertyGenerators;
 	static FDMOnUIValueUpdate OnUIValueUpdate;
-	static TArray<TSharedRef<IDMMaterialModelCreatedCallback>> OnCreatedCallbacks;
+	static TArray<TSharedRef<IDMOnWizardCompleteCallback>> OnWizardCompleteCallbacks;
 
 	TSet<FDMBuildRequestEntry> BuildRequestList;
 	TSharedRef<FUICommandList> CommandList;
diff --git a/Engine/Plugins/Experimental/DynamicMaterial/Source/DynamicMaterialEditor/Private/DynamicMaterialEditorSettings.cpp b/Engine/Plugins/Experimental/DynamicMaterial/Source/DynamicMaterialEditor/Private/DynamicMaterialEditorSettings.cpp
--- a/Engine/Plugins/Experimental/DynamicMaterial/Source/DynamicMaterialEditor/Private/DynamicMaterialEditorSettings.cpp
+++ b/Engine/Plugins/Experimental/DynamicMaterial/Source/DynamicMaterialEditor/Private/DynamicMaterialEditorSettings.cpp
@@ -78,10 +78,18 @@ UDynamicMaterialEditorSettings::UDynamicMaterialEditorSettings()
 	FDMMaterialChannelListPreset Opaque;
 	Opaque.bRGB = true;
 	Opaque.bOpacity = false;
+	Opaque.DefaultBlendMode = BLEND_Opaque;
+	Opaque.DefaultShadingModel = EDMMaterialShadingModel::DefaultLit;
+	Opaque.bDefaultAnimated = false;
+	Opaque.bDefaultTwoSided = true;
 
 	FDMMaterialChannelListPreset Translucent;
 	Translucent.bRGB = true;
 	Translucent.bOpacity = true;
+	Translucent.DefaultBlendMode = BLEND_Translucent;
+	Translucent.DefaultShadingModel = EDMMaterialShadingModel::Unlit;
+	Translucent.bDefaultAnimated = false;
+	Translucent.bDefaultTwoSided = true;
 
 	FDMMaterialChannelListPreset PBR;
 	PBR.bRGB = true;
@@ -91,6 +99,10 @@ UDynamicMaterialEditorSettings::UDynamicMaterialEditorSettings()
 	PBR.bRoughness = true;
 	PBR.bNormal = true;
 	PBR.bAmbientOcclusion = true;
+	PBR.DefaultBlendMode = BLEND_Opaque;
+	PBR.DefaultShadingModel = EDMMaterialShadingModel::DefaultLit;
+	PBR.bDefaultAnimated = false;
+	PBR.bDefaultTwoSided = true;
 
 	FDMMaterialChannelListPreset All;
 	All.bRGB = true;
@@ -105,10 +117,14 @@ UDynamicMaterialEditorSettings::UDynamicMaterialEditorSettings()
 	All.bRefraction = true;
 	All.bTangent = true;
 	All.bWorldPositionOffset = true;
+	All.DefaultBlendMode = BLEND_Opaque;
+	All.DefaultShadingModel = EDMMaterialShadingModel::DefaultLit;
+	All.bDefaultAnimated = false;
+	All.bDefaultTwoSided = true;
 
+	ChannelPresets.Add(TEXT("PBR"), PBR);
 	ChannelPresets.Add(TEXT("Opaque"), Opaque);
 	ChannelPresets.Add(TEXT("Translucent"), Translucent);
-	ChannelPresets.Add(TEXT("PBR"), PBR);
 	ChannelPresets.Add(TEXT("All"), All);		
 }
 
@@ -124,12 +140,6 @@ UDynamicMaterialEditorSettings* UDynamicMaterialEditorSettings::Get()
 	return DefaultSettings;
 }
 
-void UDynamicMaterialEditorSettings::PostInitProperties()
-{
-	Super::PostInitProperties();
-
-}
-
 void UDynamicMaterialEditorSettings::PostEditChangeProperty(FPropertyChangedEvent& InPropertyChangedEvent)
 {
 	Super::PostEditChangeProperty(InPropertyChangedEvent);
diff --git a/Engine/Plugins/Experimental/DynamicMaterial/Source/DynamicMaterialEditor/Private/DynamicMaterialEditorSettings.h b/Engine/Plugins/Experimental/DynamicMaterial/Source/DynamicMaterialEditor/Private/DynamicMaterialEditorSettings.h
--- a/Engine/Plugins/Experimental/DynamicMaterial/Source/DynamicMaterialEditor/Private/DynamicMaterialEditorSettings.h
+++ b/Engine/Plugins/Experimental/DynamicMaterial/Source/DynamicMaterialEditor/Private/DynamicMaterialEditorSettings.h
@@ -2,10 +2,11 @@
 
 #pragma once
 
+#include "Containers/Array.h"
 #include "DMDefs.h"
 #include "Engine/DeveloperSettings.h"
-#include "Containers/Array.h"
 #include "Engine/Texture.h"
+#include "MaterialDomain.h"
 #include "PropertyEditorDelegates.h"
 #include "UObject/SoftObjectPtr.h"
 #include "DynamicMaterialEditorSettings.generated.h"
@@ -70,6 +71,18 @@ struct FDMMaterialChannelListPreset
 	UPROPERTY(Config, EditAnywhere, BlueprintReadWrite, Category = "Channel")
 	bool bTangent = false;
 
+	UPROPERTY(Config, EditAnywhere, BlueprintReadWrite, Category = "Channel")
+	TEnumAsByte<EBlendMode> DefaultBlendMode = BLEND_Opaque;
+
+	UPROPERTY(Config, EditAnywhere, BlueprintReadWrite, Category = "Channel")
+	EDMMaterialShadingModel DefaultShadingModel = EDMMaterialShadingModel::Unlit;
+
+	UPROPERTY(Config, EditAnywhere, BlueprintReadWrite, Category = "Channel")
+	bool bDefaultAnimated = false;
+
+	UPROPERTY(Config, EditAnywhere, BlueprintReadWrite, Category = "Channel")
+	bool bDefaultTwoSided = true;
+
 	bool IsPropertyEnabled(EDMMaterialPropertyType InProperty) const;
 };
 
@@ -191,7 +204,6 @@ public:
 	FOnFinishedChangingProperties OnSettingsChanged;
 
 	//~ Begin UObject
-	virtual void PostInitProperties() override;
 	virtual void PostEditChangeProperty(FPropertyChangedEvent& InPropertyChangedEvent) override;
 	//~ End UObject
 
diff --git a/Engine/Plugins/Experimental/DynamicMaterial/Source/DynamicMaterialEditor/Private/Model/DMMaterialModelDefaults.cpp b/Engine/Plugins/Experimental/DynamicMaterial/Source/DynamicMaterialEditor/Private/Model/DMMaterialModelDefaults.cpp
--- a/Engine/Plugins/Experimental/DynamicMaterial/Source/DynamicMaterialEditor/Private/Model/DMMaterialModelDefaults.cpp
+++ b/Engine/Plugins/Experimental/DynamicMaterial/Source/DynamicMaterialEditor/Private/Model/DMMaterialModelDefaults.cpp
@@ -6,14 +6,14 @@
 #include "Components/StaticMeshComponent.h"
 #include "DynamicMaterialEditorModule.h"
 #include "Delegates/IDelegateInstance.h"
-#include "Model/DMModelCreatedCallback.h"
+#include "Model/DMOnWizardCompleteCallback.h"
 #include "Model/DynamicMaterialModel.h"
 #include "Model/DynamicMaterialModelEditorOnlyData.h"
 
 namespace UE::DynamicMaterialEditor::Private
 {
-	TSharedPtr<IDMMaterialModelCreatedCallback> DefaultsCallback_StaticMeshComponent;
-	TSharedPtr<IDMMaterialModelCreatedCallback> DefaultsCallback_DecalComponent;
+	TSharedPtr<IDMOnWizardCompleteCallback> DefaultsCallback_StaticMeshComponent;
+	TSharedPtr<IDMOnWizardCompleteCallback> DefaultsCallback_DecalComponent;
 }
 
 struct FDMMaterialModelCreatedCallback_StaticMeshComponent : FDMMaterialModelCreatedCallbackBase
@@ -26,7 +26,7 @@ struct FDMMaterialModelCreatedCallback_StaticMeshComponent : FDMMaterialModelCre
 	virtual ~FDMMaterialModelCreatedCallback_StaticMeshComponent() override = default;
 
 	//~ Begin IDMMaterialModelCreatedCallback
-	virtual void OnModelCreated(const FDMMaterialModelCreatedCallbackParams& InParams) override
+	virtual void OnModelCreated(const FDMOnWizardCompleteCallbackParams& InParams) override
 	{
 		if (!::IsValid(InParams.EditorOnlyData))
 		{
@@ -67,7 +67,7 @@ struct FDMMaterialModelCreatedCallback_DecalComponent : FDMMaterialModelCreatedC
 	virtual ~FDMMaterialModelCreatedCallback_DecalComponent() override = default;
 
 	//~ Begin IDMMaterialModelCreatedCallback
-	virtual void OnModelCreated(const FDMMaterialModelCreatedCallbackParams& InParams) override
+	virtual void OnModelCreated(const FDMOnWizardCompleteCallbackParams& InParams) override
 	{
 		if (!::IsValid(InParams.EditorOnlyData))
 		{
diff --git a/Engine/Plugins/Experimental/DynamicMaterial/Source/DynamicMaterialEditor/Private/Model/DMModelCreatedCallback.cpp b/Engine/Plugins/Experimental/DynamicMaterial/Source/DynamicMaterialEditor/Private/Model/DMModelCreatedCallback.cpp
deleted file mode 100644
diff --git a/Engine/Plugins/Experimental/DynamicMaterial/Source/DynamicMaterialEditor/Private/Model/DMOnWizardCompleteCallback.cpp b/Engine/Plugins/Experimental/DynamicMaterial/Source/DynamicMaterialEditor/Private/Model/DMOnWizardCompleteCallback.cpp
new file mode 100644
--- /dev/null
+++ b/Engine/Plugins/Experimental/DynamicMaterial/Source/DynamicMaterialEditor/Private/Model/DMOnWizardCompleteCallback.cpp
@@ -0,0 +1,24 @@
+// Copyright Epic Games, Inc. All Rights Reserved.
+
+#include "Model/DMOnWizardCompleteCallback.h"
+
+FDMMaterialModelCreatedCallbackBase::FDMMaterialModelCreatedCallbackBase(uint32 InPriority)
+	: Priority(InPriority)
+{
+}
+
+uint32 FDMMaterialModelCreatedCallbackBase::GetPriority() const
+{
+	return Priority;
+}
+
+FDMMaterialModelCreatedCallbackDelegate::FDMMaterialModelCreatedCallbackDelegate(uint32 InPriority, const FOnModelCreated& InOnModelCreatedDelegate)
+	: FDMMaterialModelCreatedCallbackBase(InPriority)
+	, OnModelCreatedDelegate(InOnModelCreatedDelegate)
+{
+}
+
+void FDMMaterialModelCreatedCallbackDelegate::OnModelCreated(const FDMOnWizardCompleteCallbackParams& InParams)
+{
+	OnModelCreatedDelegate.ExecuteIfBound(InParams);
+}
diff --git a/Engine/Plugins/Experimental/DynamicMaterial/Source/DynamicMaterialEditor/Private/Model/DynamicMateriaModelFactory.cpp b/Engine/Plugins/Experimental/DynamicMaterial/Source/DynamicMaterialEditor/Private/Model/DynamicMateriaModelFactory.cpp
--- a/Engine/Plugins/Experimental/DynamicMaterial/Source/DynamicMaterialEditor/Private/Model/DynamicMateriaModelFactory.cpp
+++ b/Engine/Plugins/Experimental/DynamicMaterial/Source/DynamicMaterialEditor/Private/Model/DynamicMateriaModelFactory.cpp
@@ -41,8 +41,6 @@ UObject* UDynamicMaterialModelFactory::FactoryCreateNew(UClass* Class, UObject*
 	const FDMInitializationGuard InitGuard;
 	ModelEditorOnlyData->Initialize();
 
-	FDynamicMaterialEditorModule::Get().OnMaterialModelCreated(NewModel);
-
 	return NewModel;
 }
 
diff --git a/Engine/Plugins/Experimental/DynamicMaterial/Source/DynamicMaterialEditor/Private/Model/DynamicMaterialModelEditorOnlyData.cpp b/Engine/Plugins/Experimental/DynamicMaterial/Source/DynamicMaterialEditor/Private/Model/DynamicMaterialModelEditorOnlyData.cpp
--- a/Engine/Plugins/Experimental/DynamicMaterial/Source/DynamicMaterialEditor/Private/Model/DynamicMaterialModelEditorOnlyData.cpp
+++ b/Engine/Plugins/Experimental/DynamicMaterial/Source/DynamicMaterialEditor/Private/Model/DynamicMaterialModelEditorOnlyData.cpp
@@ -154,8 +154,8 @@ UDynamicMaterialModelEditorOnlyData::UDynamicMaterialModelEditorOnlyData()
 	Properties.Emplace(EDMMaterialPropertyType::Normal,              CreateDefaultSubobject<UDMMaterialPropertyNormal>(             "MaterialProperty_Normal"));
 	Properties.Emplace(EDMMaterialPropertyType::Tangent,             CreateDefaultSubobject<UDMMaterialPropertyTangent>(            "MaterialProperty_Tangent"));
 	Properties.Emplace(EDMMaterialPropertyType::WorldPositionOffset, CreateDefaultSubobject<UDMMaterialPropertyWorldPositionOffset>("MaterialProperty_WorldPositionOffset"));
-	Properties.Emplace(EDMMaterialPropertyType::Refraction,          CreateDefaultSubobject<UDMMaterialPropertyRefraction>(         "MaterialProperty_Refraction"));
 	Properties.Emplace(EDMMaterialPropertyType::AmbientOcclusion,    CreateDefaultSubobject<UDMMaterialPropertyAmbientOcclusion>(   "MaterialProperty_AmbientOcclusion"));
+	Properties.Emplace(EDMMaterialPropertyType::Refraction,          CreateDefaultSubobject<UDMMaterialPropertyRefraction>(         "MaterialProperty_Refraction"));
 	Properties.Emplace(EDMMaterialPropertyType::PixelDepthOffset,    CreateDefaultSubobject<UDMMaterialPropertyPixelDepthOffset>(   "MaterialProperty_PixelDepthOffset"));
 
 	Properties.Emplace(EDMMaterialPropertyType::Custom1, UDMMaterialProperty::CreateCustomMaterialPropertyDefaultSubobject(this, EDMMaterialPropertyType::Custom1, "MaterialProperty_Custom1"));
@@ -673,35 +673,48 @@ bool UDynamicMaterialModelEditorOnlyData::NeedsWizard() const
 
 void UDynamicMaterialModelEditorOnlyData::OnWizardComplete()
 {
+	if (UDynamicMaterialModel* MaterialModelLocal = MaterialModel.Get())
+	{
+		FDynamicMaterialEditorModule::Get().OnWizardComplete(MaterialModelLocal);
+	}
 }
 
 void UDynamicMaterialModelEditorOnlyData::SetChannelListPreset(FName InPresetName)
 {
 	ChannelListPreset = InPresetName;
 
-	if (const FDMMaterialChannelListPreset* Preset = GetDefault<UDynamicMaterialEditorSettings>()->ChannelPresets.Find(InPresetName))
+	const FDMMaterialChannelListPreset* Preset = GetDefault<UDynamicMaterialEditorSettings>()->ChannelPresets.Find(InPresetName);
+
+	if (!Preset)
 	{
-		for (uint8 PropertyIndex = static_cast<uint8>(EDMMaterialPropertyType::None) + 1;
-			PropertyIndex < static_cast<uint8>(EDMMaterialPropertyType::Any);
-			++PropertyIndex)
-		{
-			const EDMMaterialPropertyType Property = static_cast<EDMMaterialPropertyType>(PropertyIndex);
+		return;
+	}
 
-			if (Property == EDMMaterialPropertyType::EmissiveColor || Property == EDMMaterialPropertyType::OpacityMask)
-			{
-				continue;
-			}
+	for (uint8 PropertyIndex = static_cast<uint8>(EDMMaterialPropertyType::None) + 1;
+		PropertyIndex < static_cast<uint8>(EDMMaterialPropertyType::Any);
+		++PropertyIndex)
+	{
+		const EDMMaterialPropertyType Property = static_cast<EDMMaterialPropertyType>(PropertyIndex);
 
-			if (Preset->IsPropertyEnabled(Property))
-			{
-				AddSlotForMaterialProperty(Property);
-			}
-			else
-			{
-				RemoveSlotForMaterialProperty(Property);
-			}
+		if (Property == EDMMaterialPropertyType::EmissiveColor || Property == EDMMaterialPropertyType::OpacityMask)
+		{
+			continue;
+		}
+
+		if (Preset->IsPropertyEnabled(Property))
+		{
+			AddSlotForMaterialProperty(Property);
+		}
+		else
+		{
+			RemoveSlotForMaterialProperty(Property);
 		}
 	}
+
+	SetBlendMode(Preset->DefaultBlendMode);
+	SetShadingModel(Preset->DefaultShadingModel);
+	SetPixelAnimationFlag(Preset->bDefaultAnimated);
+	SetTwoSidedFlag(Preset->bDefaultTwoSided);
 }
 
 UDMMaterialComponent* UDynamicMaterialModelEditorOnlyData::GetSubComponentByPath(FDMComponentPath& InPath,
@@ -1070,10 +1083,10 @@ UDMMaterialSlot* UDynamicMaterialModelEditorOnlyData::GetSlotForMaterialProperty
 
 UDMMaterialSlot* UDynamicMaterialModelEditorOnlyData::AddSlot()
 {
-	constexpr int64 Start = static_cast<int64>(EDMMaterialPropertyType::None) + 1;
-	constexpr int64 End = static_cast<int64>(EDMMaterialPropertyType::Any);
+	constexpr uint8 Start = static_cast<uint8>(EDMMaterialPropertyType::None) + 1;
+	constexpr uint8 End = static_cast<uint8>(EDMMaterialPropertyType::Any);
 
-	for (int64 PropertyIndex = Start; PropertyIndex < End; ++PropertyIndex)
+	for (uint8 PropertyIndex = Start; PropertyIndex < End; ++PropertyIndex)
 	{
 		const EDMMaterialPropertyType Property = static_cast<EDMMaterialPropertyType>(PropertyIndex);
 
diff --git a/Engine/Plugins/Experimental/DynamicMaterial/Source/DynamicMaterialEditor/Private/Slate/SDMEditor.cpp b/Engine/Plugins/Experimental/DynamicMaterial/Source/DynamicMaterialEditor/Private/Slate/SDMEditor.cpp
--- a/Engine/Plugins/Experimental/DynamicMaterial/Source/DynamicMaterialEditor/Private/Slate/SDMEditor.cpp
+++ b/Engine/Plugins/Experimental/DynamicMaterial/Source/DynamicMaterialEditor/Private/Slate/SDMEditor.cpp
@@ -292,7 +292,6 @@ void SDMEditor::SetMaterialModel(UDynamicMaterialModel* InMaterialModel)
 	if (EditorOnlyData && EditorOnlyData->NeedsWizard())
 	{
 		Container->SetContent(SNew(SDMMaterialWizard, SharedThis(this)));
-		EditorOnlyData->OnWizardComplete();
 	}
 	else
 	{
@@ -697,51 +696,32 @@ TSharedRef<SWidget> SDMEditor::CreateSlotPickerWidget()
 
 	UEnum* MaterialPropertyEnum = StaticEnum<EDMMaterialPropertyType>();
 
-	for (uint8 PropertyIndex = static_cast<uint8>(EDMMaterialPropertyType::None) + 1;
-		PropertyIndex < static_cast<uint8>(EDMMaterialPropertyType::Any);
-		++PropertyIndex)
+	for (const TPair<EDMMaterialPropertyType, UDMMaterialProperty*>& Property : ModelEditorOnlyData->GetMaterialProperties())
 	{
-		EDMMaterialPropertyType Property = static_cast<EDMMaterialPropertyType>(PropertyIndex);
-
-		// Filter RGB modes
-		if (bHasBaseColorSlot)
-		{
-			if (Property == EDMMaterialPropertyType::EmissiveColor)
-			{
-				continue;
-			}
-		}
-		else
+		// Always create base colour, not emissive - Will be sorted out by the button itself.
+		if (Property.Key == EDMMaterialPropertyType::EmissiveColor)
 		{
-			if (Property == EDMMaterialPropertyType::BaseColor)
-			{
-				continue;
-			}
+			continue;
 		}
 
-		// Filter opacity modes
-		if (Property == EDMMaterialPropertyType::OpacityMask)
+		// Always create opacity, not opacity mask - Will be sorted out by the button itself.
+		if (Property.Key == EDMMaterialPropertyType::OpacityMask)
 		{
-			if (ModelEditorOnlyData->GetBlendMode() != BLEND_Masked)
-			{
-				continue;
-			}
+			continue;
 		}
-		else
+
+		if (Property.Key == EDMMaterialPropertyType::Opacity && ModelEditorOnlyData->GetBlendMode() == BLEND_Opaque)
 		{
-			if (ModelEditorOnlyData->GetBlendMode() == BLEND_Masked)
-			{
-				continue;
-			}
+			continue;
 		}
 
 		const bool bIsPropertyEnabled = Preset
-			? Preset->IsPropertyEnabled(Property)
-			: !!ModelEditorOnlyData->GetSlotForMaterialProperty(Property);
+			? Preset->IsPropertyEnabled(Property.Key)
+			: !!ModelEditorOnlyData->GetSlotForMaterialProperty(Property.Key);
 
 		if (bIsPropertyEnabled)
 		{
-			const int32 EnumIndex = MaterialPropertyEnum->GetIndexByValue(static_cast<int64>(Property));
+			const int32 EnumIndex = MaterialPropertyEnum->GetIndexByValue(static_cast<int64>(Property.Key));
 
 			constexpr const TCHAR* ShortNameName = TEXT("ShortName");
 			const FString ShortName = MaterialPropertyEnum->GetMetaData(ShortNameName, EnumIndex);
@@ -751,15 +731,15 @@ TSharedRef<SWidget> SDMEditor::CreateSlotPickerWidget()
 					SNew(SCheckBox)
 					.Style(FAppStyle::Get(), "DetailsView.SectionButton")
 					.HAlign(EHorizontalAlignment::HAlign_Center)
-					.IsEnabled(this, &SDMEditor::IsPropertyValidForModel, Property)
-					.IsChecked(this, &SDMEditor::GetSlotCheckState, Property)
-					.OnCheckStateChanged(this, &SDMEditor::OnSlotCheckStateChanged, Property)
+					.IsEnabled(this, &SDMEditor::IsPropertyValidForModel, Property.Key)
+					.IsChecked(this, &SDMEditor::GetSlotCheckState, Property.Key)
+					.OnCheckStateChanged(this, &SDMEditor::OnSlotCheckStateChanged, Property.Key)
 					.Padding(FVector2D(5.f, 3.f))
-					.ToolTipText(this, &SDMEditor::GetToolTipForProperty, Property)
+					.ToolTipText(this, &SDMEditor::GetToolTipForProperty, Property.Key)
 					.Content()
 					[
 						SNew(STextBlock)
-						.Text(!ShortName.IsEmpty() ? FText::FromString(ShortName) : MaterialPropertyEnum->GetDisplayNameTextByValue(static_cast<int64>(Property)))
+						.Text(!ShortName.IsEmpty() ? FText::FromString(ShortName) : MaterialPropertyEnum->GetDisplayNameTextByValue(static_cast<int64>(Property.Key)))
 					]
 				];
 		}
@@ -937,7 +917,26 @@ bool SDMEditor::IsPropertyValidForModel(EDMMaterialPropertyType InProperty) cons
 		{
 			if (UDMMaterialProperty* Property = EditorOnlyData->GetMaterialProperty(InProperty))
 			{
-				return Property->IsValidForModel(*EditorOnlyData);
+				if (Property->IsValidForModel(*EditorOnlyData))
+				{
+					return true;
+				}
+			}
+
+			if (InProperty == EDMMaterialPropertyType::BaseColor)
+			{
+				if (UDMMaterialProperty* Property = EditorOnlyData->GetMaterialProperty(EDMMaterialPropertyType::EmissiveColor))
+				{
+					return Property->IsValidForModel(*EditorOnlyData);
+				}
+			}
+
+			if (InProperty == EDMMaterialPropertyType::Opacity)
+			{
+				if (UDMMaterialProperty* Property = EditorOnlyData->GetMaterialProperty(EDMMaterialPropertyType::OpacityMask))
+				{
+					return Property->IsValidForModel(*EditorOnlyData);
+				}
 			}
 		}
 	}
@@ -953,9 +952,26 @@ ECheckBoxState SDMEditor::GetSlotCheckState(EDMMaterialPropertyType InProperty)
 		{
 			if (UDynamicMaterialModelEditorOnlyData* EditorOnlyData = Slot->GetMaterialModelEditorOnlyData())
 			{
-				return EditorOnlyData->GetSlotForMaterialProperty(InProperty) == Slot
-					? ECheckBoxState::Checked
-					: ECheckBoxState::Unchecked;
+				if (UDMMaterialSlot* PropertySlot = EditorOnlyData->GetSlotForMaterialProperty(InProperty))
+				{
+					return PropertySlot == Slot
+						? ECheckBoxState::Checked
+						: ECheckBoxState::Unchecked;
+				}
+
+				if (InProperty == EDMMaterialPropertyType::BaseColor)
+				{
+					return EditorOnlyData->GetSlotForMaterialProperty(EDMMaterialPropertyType::EmissiveColor) == Slot
+						? ECheckBoxState::Checked
+						: ECheckBoxState::Unchecked;
+				}
+
+				if (InProperty == EDMMaterialPropertyType::Opacity)
+				{
+					return EditorOnlyData->GetSlotForMaterialProperty(EDMMaterialPropertyType::OpacityMask) == Slot
+						? ECheckBoxState::Checked
+						: ECheckBoxState::Unchecked;
+				}	
 			}
 		}
 	}
@@ -983,9 +999,28 @@ void SDMEditor::OnSlotCheckStateChanged(ECheckBoxState InCheckState, EDMMaterial
 			{
 				if (UDynamicMaterialModelEditorOnlyData* EditorOnlyData = Slot->GetMaterialModelEditorOnlyData())
 				{
-					if (EditorOnlyData->GetSlotForMaterialProperty(InProperty) == Slot)
+					if (UDMMaterialSlot* PropertySlot = EditorOnlyData->GetSlotForMaterialProperty(InProperty))
+					{
+						if (PropertySlot == Slot)
+						{
+							SetActiveSlotIndex(SlotIndex);
+						}
+					}
+
+					if (InProperty == EDMMaterialPropertyType::BaseColor)
+					{
+						if (EditorOnlyData->GetSlotForMaterialProperty(EDMMaterialPropertyType::EmissiveColor) == Slot)
+						{
+							SetActiveSlotIndex(SlotIndex);
+						}
+					}
+
+					if (InProperty == EDMMaterialPropertyType::Opacity)
 					{
-						SetActiveSlotIndex(SlotIndex);
+						if (EditorOnlyData->GetSlotForMaterialProperty(EDMMaterialPropertyType::OpacityMask) == Slot)
+						{
+							SetActiveSlotIndex(SlotIndex);
+						}
 					}
 				}
 			}
diff --git a/Engine/Plugins/Experimental/DynamicMaterial/Source/DynamicMaterialEditor/Private/Slate/SDMMaterialWizard.cpp b/Engine/Plugins/Experimental/DynamicMaterial/Source/DynamicMaterialEditor/Private/Slate/SDMMaterialWizard.cpp
--- a/Engine/Plugins/Experimental/DynamicMaterial/Source/DynamicMaterialEditor/Private/Slate/SDMMaterialWizard.cpp
+++ b/Engine/Plugins/Experimental/DynamicMaterial/Source/DynamicMaterialEditor/Private/Slate/SDMMaterialWizard.cpp
@@ -1,12 +1,10 @@
 // Copyright Epic Games, Inc. All Rights Reserved.
 
 #include "Slate/SDMMaterialWizard.h"
-
 #include "DMDefs.h"
 #include "DynamicMaterialEditorSettings.h"
 #include "DynamicMaterialEditorStyle.h"
 #include "Engine/EngineTypes.h"
-#include "MaterialDomain.h"
 #include "Model/DynamicMaterialModel.h"
 #include "Model/DynamicMaterialModelEditorOnlyData.h"
 #include "SDMEditor.h"
@@ -117,109 +115,6 @@ TSharedRef<SWidget> SDMMaterialWizard::CreateLayout()
 				]
 			]
 
-			+ SVerticalBox::Slot()
-			.AutoHeight()
-			.HAlign(HAlign_Fill)
-			.Padding(0.0f, SeparationDistance, 0.0f, 0.f)
-			[
-				SNew(SExpandableArea)
-				.BorderImage(FAppStyle::GetBrush(TEXT("Menu.Background")))
-				.BodyBorderImage(FAppStyle::GetBrush("DetailsView.CategoryTop"))
-				.HeaderContent()
-				[
-					SNew(STextBlock)
-					.Text(LOCTEXT("AdvancedSettings", "Advanced Settings"))
-				]
-				.BodyContent()
-				[
-					SNew(SBox)
-					.Padding(FMargin(SeparationDistance, SeparationDistance, SeparationDistance, SeparationDistance))
-					[
-						SNew(SVerticalBox)
-
-						+ SVerticalBox::Slot()
-						.AutoHeight()
-						.HAlign(HAlign_Left)
-						.Padding(0.0f, 0.f, 0.0f, TitleContentDistance)
-						[
-							SNew(STextBlock)
-							.Text(LOCTEXT("MaterialDomain", "Material Domain"))
-						]
-
-						+ SVerticalBox::Slot()
-						.AutoHeight()
-						.HAlign(EHorizontalAlignment::HAlign_Fill)
-						[
-							CreateMaterialDomainOptions()
-						]
-
-						+ SVerticalBox::Slot()
-						.AutoHeight()
-						.HAlign(HAlign_Left)
-						.Padding(0.0f, SeparationDistance, 0.0f, TitleContentDistance)
-						[
-							SNew(STextBlock)
-							.Text(LOCTEXT("BlendMode", "Blend Mode"))
-						]
-
-						+ SVerticalBox::Slot()
-						.AutoHeight()
-						.HAlign(EHorizontalAlignment::HAlign_Fill)
-						[
-							CreateBlendModeOptions()
-						]
-
-						+ SVerticalBox::Slot()
-						.AutoHeight()
-						.HAlign(HAlign_Left)
-						.Padding(0.0f, SeparationDistance, 0.0f, TitleContentDistance)
-						[
-							SNew(STextBlock)
-							.Text(LOCTEXT("ShadingModel", "Shading Model"))
-						]
-
-						+ SVerticalBox::Slot()
-						.AutoHeight()
-						.HAlign(EHorizontalAlignment::HAlign_Fill)
-						[
-							CreateShadingModelOptions()
-						]
-
-						+ SVerticalBox::Slot()
-						.AutoHeight()
-						.HAlign(HAlign_Left)
-						.Padding(0.0f, SeparationDistance, 0.0f, TitleContentDistance)
-						[
-							SNew(STextBlock)
-							.Text(LOCTEXT("MotionVectors", "Motion Vectors"))
-						]
-
-						+ SVerticalBox::Slot()
-						.AutoHeight()
-						.HAlign(EHorizontalAlignment::HAlign_Fill)
-						[
-							CreateAnimationOptions()
-						]
-
-						+ SVerticalBox::Slot()
-						.AutoHeight()
-						.HAlign(HAlign_Left)
-						.Padding(0.0f, SeparationDistance, 0.0f, TitleContentDistance)
-						[
-							SNew(STextBlock)
-							.Text(LOCTEXT("Geometry", "Geometry"))
-						]
-
-						+ SVerticalBox::Slot()
-						.AutoHeight()
-						.HAlign(EHorizontalAlignment::HAlign_Fill)
-						[
-							CreateTwoSidedOptions()
-						]
-					]
-				]
-			]
-
 			+ SVerticalBox::Slot()
 			.AutoHeight()
 			.HAlign(EHorizontalAlignment::HAlign_Fill)
@@ -230,196 +125,6 @@ TSharedRef<SWidget> SDMMaterialWizard::CreateLayout()
 		];
 }
 
-TSharedRef<SWidget> SDMMaterialWizard::CreateMaterialDomainOptions()
-{
-	using namespace UE::DynamicMaterialEditor::Private;
-
-	UEnum* DomainEnum = StaticEnum<EMaterialDomain>();
-
-	TSharedRef<SWrapBox> DomainOptions = SNew(SWrapBox)
-		.UseAllottedSize(true)
-		.InnerSlotPadding(WrapBoxSlotPadding)
-		.Orientation(EOrientation::Orient_Horizontal);
-
-	for (EMaterialDomain Domain : UDynamicMaterialModelEditorOnlyData::SupportedDomains)
-	{
-		DomainOptions->AddSlot()
-			[
-				SNew(SCheckBox)
-				.Style(FAppStyle::Get(), "DetailsView.SectionButton")
-				.HAlign(EHorizontalAlignment::HAlign_Center)
-				.Padding(ButtonPadding)
-				.IsEnabled(this, &SDMMaterialWizard::Domain_IsEnabled)
-				.IsChecked(this, &SDMMaterialWizard::Domain_GetState, Domain)
-				.OnCheckStateChanged(this, &SDMMaterialWizard::Domain_OnChange, Domain)
-				[
-					SNew(STextBlock)
-					.Text(DomainEnum->GetDisplayNameTextByValue(Domain))
-				]
-			];
-	}
-
-	return DomainOptions;
-}
-
-TSharedRef<SWidget> SDMMaterialWizard::CreateBlendModeOptions()
-{
-	using namespace UE::DynamicMaterialEditor::Private;
-
-	UEnum* BlendModeEnum = StaticEnum<EBlendMode>();
-
-	TSharedRef<SWrapBox> BlendModeOptions = SNew(SWrapBox)
-		.UseAllottedSize(true)
-		.InnerSlotPadding(WrapBoxSlotPadding)
-		.Orientation(EOrientation::Orient_Horizontal);
-
-	for (EBlendMode BlendMode : UDynamicMaterialModelEditorOnlyData::SupportedBlendModes)
-	{
-		BlendModeOptions->AddSlot()
-			[
-				SNew(SCheckBox)
-				.Style(FAppStyle::Get(), "DetailsView.SectionButton")
-				.HAlign(EHorizontalAlignment::HAlign_Center)
-				.Padding(ButtonPadding)
-				.IsEnabled(this, &SDMMaterialWizard::BlendMode_IsEnabled)
-				.IsChecked(this, &SDMMaterialWizard::BlendMode_GetState, BlendMode)
-				.OnCheckStateChanged(this, &SDMMaterialWizard::BlendMode_OnChange, BlendMode)
-				[
-					SNew(STextBlock)
-					.Text(BlendModeEnum->GetDisplayNameTextByValue(BlendMode))
-				]
-			];
-	}
-
-	return BlendModeOptions;
-}
-
-TSharedRef<SWidget> SDMMaterialWizard::CreateShadingModelOptions()
-{
-	using namespace UE::DynamicMaterialEditor::Private;
-
-	TSharedRef<SWrapBox> MaterialShadingModelOptions = SNew(SWrapBox)
-		.UseAllottedSize(true)
-		.InnerSlotPadding(WrapBoxSlotPadding)
-		.Orientation(EOrientation::Orient_Horizontal);
-
-	MaterialShadingModelOptions->AddSlot()
-		[
-			SNew(SCheckBox)
-			.Style(FAppStyle::Get(), "DetailsView.SectionButton")
-			.HAlign(EHorizontalAlignment::HAlign_Center)
-			.Padding(ButtonPadding)
-			.IsEnabled(this, &SDMMaterialWizard::Unlit_IsEnabled)
-			.IsChecked(this, &SDMMaterialWizard::Unlit_GetState, EDMMaterialShadingModel::Unlit)
-			.OnCheckStateChanged(this, &SDMMaterialWizard::Unlit_OnChange, EDMMaterialShadingModel::Unlit)
-			[
-				SNew(STextBlock)
-				.Text(LOCTEXT("Emissive", "Emissive"))
-			]
-		];
-
-	MaterialShadingModelOptions->AddSlot()
-		[
-			SNew(SCheckBox)
-			.Style(FAppStyle::Get(), "DetailsView.SectionButton")
-			.HAlign(EHorizontalAlignment::HAlign_Center)
-			.Padding(ButtonPadding)
-			.IsEnabled(this, &SDMMaterialWizard::Unlit_IsEnabled)
-			.IsChecked(this, &SDMMaterialWizard::Unlit_GetState, EDMMaterialShadingModel::DefaultLit)
-			.OnCheckStateChanged(this, &SDMMaterialWizard::Unlit_OnChange, EDMMaterialShadingModel::DefaultLit)
-			[
-				SNew(STextBlock)
-				.Text(LOCTEXT("DefaultLit", "Default Lit"))
-			]
-		];
-
-	return MaterialShadingModelOptions;
-}
-
-TSharedRef<SWidget> SDMMaterialWizard::CreateAnimationOptions()
-{
-	using namespace UE::DynamicMaterialEditor::Private;
-
-	TSharedRef<SWrapBox> AnimatedOptions = SNew(SWrapBox)
-		.UseAllottedSize(true)
-		.InnerSlotPadding(WrapBoxSlotPadding)
-		.Orientation(EOrientation::Orient_Horizontal);
-
-	AnimatedOptions->AddSlot()
-		[
-			SNew(SCheckBox)
-				.Style(FAppStyle::Get(), "DetailsView.SectionButton")
-				.HAlign(EHorizontalAlignment::HAlign_Center)
-				.Padding(ButtonPadding)
-				.IsEnabled(this, &SDMMaterialWizard::Animated_IsEnabled)
-				.IsChecked(this, &SDMMaterialWizard::Animated_GetState, false)
-				.OnCheckStateChanged(this, &SDMMaterialWizard::Animated_OnChange, false)
-				[
-					SNew(STextBlock)
-					.Text(LOCTEXT("Static", "Static"))
-				]
-		];
-
-	AnimatedOptions->AddSlot()
-		[
-			SNew(SCheckBox)
-				.Style(FAppStyle::Get(), "DetailsView.SectionButton")
-				.HAlign(EHorizontalAlignment::HAlign_Center)
-				.Padding(ButtonPadding)
-				.IsEnabled(this, &SDMMaterialWizard::Animated_IsEnabled)
-				.IsChecked(this, &SDMMaterialWizard::Animated_GetState, true)
-				.OnCheckStateChanged(this, &SDMMaterialWizard::Animated_OnChange, true)
-				[
-					SNew(STextBlock)
-					.Text(LOCTEXT("Animated", "Animated"))
-				]
-		];
-
-	return AnimatedOptions;
-}
-
-TSharedRef<SWidget> SDMMaterialWizard::CreateTwoSidedOptions()
-{
-	using namespace UE::DynamicMaterialEditor::Private;
-
-	TSharedRef<SWrapBox> TwoSidedOptions = SNew(SWrapBox)
-		.UseAllottedSize(true)
-		.InnerSlotPadding(WrapBoxSlotPadding)
-		.Orientation(EOrientation::Orient_Horizontal);
-
-	TwoSidedOptions->AddSlot()
-		[
-			SNew(SCheckBox)
-				.Style(FAppStyle::Get(), "DetailsView.SectionButton")
-				.HAlign(EHorizontalAlignment::HAlign_Center)
-				.Padding(ButtonPadding)
-				.IsEnabled(this, &SDMMaterialWizard::TwoSided_IsEnabled)
-				.IsChecked(this, &SDMMaterialWizard::TwoSided_GetState, false)
-				.OnCheckStateChanged(this, &SDMMaterialWizard::TwoSided_OnChange, false)
-				[
-					SNew(STextBlock)
-					.Text(LOCTEXT("OneSided", "One Sided"))
-				]
-		];
-
-	TwoSidedOptions->AddSlot()
-		[
-			SNew(SCheckBox)
-				.Style(FAppStyle::Get(), "DetailsView.SectionButton")
-				.HAlign(EHorizontalAlignment::HAlign_Center)
-				.Padding(ButtonPadding)
-				.IsEnabled(this, &SDMMaterialWizard::TwoSided_IsEnabled)
-				.IsChecked(this, &SDMMaterialWizard::TwoSided_GetState, true)
-				.OnCheckStateChanged(this, &SDMMaterialWizard::TwoSided_OnChange, true)
-				[
-					SNew(STextBlock)
-					.Text(LOCTEXT("TwoSided", "Two Sided"))
-				]
-		];
-
-	return TwoSidedOptions;
-}
-
 TSharedRef<SWidget> SDMMaterialWizard::CreateChannelPresets()
 {
 	using namespace UE::DynamicMaterialEditor::Private;
@@ -513,201 +218,6 @@ TSharedRef<SWidget> SDMMaterialWizard::CreateAcceptButton()
 		];
 }
 
-bool SDMMaterialWizard::Domain_IsEnabled() const
-{
-	if (UDynamicMaterialModel* MaterialModel = GetMaterialModel())
-	{
-		return !!UDynamicMaterialModelEditorOnlyData::Get(MaterialModel);
-	}
-
-	return false;
-}
-
-ECheckBoxState SDMMaterialWizard::Domain_GetState(EMaterialDomain InDomain) const
-{
-	if (UDynamicMaterialModel* MaterialModel = GetMaterialModel())
-	{
-		if (UDynamicMaterialModelEditorOnlyData* EditorOnlyData = UDynamicMaterialModelEditorOnlyData::Get(MaterialModel))
-		{
-			return EditorOnlyData->GetDomain() == InDomain
-				? ECheckBoxState::Checked
-				: ECheckBoxState::Unchecked;
-		}
-	}
-
-	return ECheckBoxState::Undetermined;
-}
-
-void SDMMaterialWizard::Domain_OnChange(ECheckBoxState InState, EMaterialDomain InDomain)
-{
-	if (InState == ECheckBoxState::Checked)
-	{
-		if (UDynamicMaterialModel* MaterialModel = GetMaterialModel())
-		{
-			if (UDynamicMaterialModelEditorOnlyData* EditorOnlyData = UDynamicMaterialModelEditorOnlyData::Get(MaterialModel))
-			{
-				EditorOnlyData->SetDomain(InDomain);
-			}
-		}
-	}
-}
-
-bool SDMMaterialWizard::BlendMode_IsEnabled() const
-{
-	if (UDynamicMaterialModel* MaterialModel = GetMaterialModel())
-	{
-		return !!UDynamicMaterialModelEditorOnlyData::Get(MaterialModel);
-	}
-
-	return false;
-}
-
-ECheckBoxState SDMMaterialWizard::BlendMode_GetState(EBlendMode InBlendMode) const
-{
-	if (UDynamicMaterialModel* MaterialModel = GetMaterialModel())
-	{
-		if (UDynamicMaterialModelEditorOnlyData* EditorOnlyData = UDynamicMaterialModelEditorOnlyData::Get(MaterialModel))
-		{
-			return EditorOnlyData->GetBlendMode() == InBlendMode
-				? ECheckBoxState::Checked
-				: ECheckBoxState::Unchecked;
-		}
-	}
-
-	return ECheckBoxState::Undetermined;
-}
-
-void SDMMaterialWizard::BlendMode_OnChange(ECheckBoxState InState, EBlendMode InBlendMode)
-{
-	if (InState == ECheckBoxState::Checked)
-	{
-		if (UDynamicMaterialModel* MaterialModel = GetMaterialModel())
-		{
-			if (UDynamicMaterialModelEditorOnlyData* EditorOnlyData = UDynamicMaterialModelEditorOnlyData::Get(MaterialModel))
-			{
-				EditorOnlyData->SetBlendMode(InBlendMode);
-			}
-		}
-	}
-}
-
-bool SDMMaterialWizard::Unlit_IsEnabled() const
-{
-	if (UDynamicMaterialModel* MaterialModel = GetMaterialModel())
-	{
-		return !!UDynamicMaterialModelEditorOnlyData::Get(MaterialModel);
-	}
-
-	return false;
-}
-
-ECheckBoxState SDMMaterialWizard::Unlit_GetState(EDMMaterialShadingModel InValue) const
-{
-	if (UDynamicMaterialModel* MaterialModel = GetMaterialModel())
-	{
-		if (UDynamicMaterialModelEditorOnlyData* EditorOnlyData = UDynamicMaterialModelEditorOnlyData::Get(MaterialModel))
-		{
-			return EditorOnlyData->GetShadingModel() == InValue
-				? ECheckBoxState::Checked
-				: ECheckBoxState::Unchecked;
-		}
-	}
-
-	return ECheckBoxState::Undetermined;
-}
-
-void SDMMaterialWizard::Unlit_OnChange(ECheckBoxState InState, EDMMaterialShadingModel InValue)
-{
-	if (InState == ECheckBoxState::Checked)
-	{
-		if (UDynamicMaterialModel* MaterialModel = GetMaterialModel())
-		{
-			if (UDynamicMaterialModelEditorOnlyData* EditorOnlyData = UDynamicMaterialModelEditorOnlyData::Get(MaterialModel))
-			{
-				EditorOnlyData->SetShadingModel(InValue);
-			}
-		}
-	}
-}
-
-bool SDMMaterialWizard::Animated_IsEnabled() const
-{
-	if (UDynamicMaterialModel* MaterialModel = GetMaterialModel())
-	{
-		return !!UDynamicMaterialModelEditorOnlyData::Get(MaterialModel);
-	}
-
-	return false;
-}
-
-ECheckBoxState SDMMaterialWizard::Animated_GetState(bool bInValue) const
-{
-	if (UDynamicMaterialModel* MaterialModel = GetMaterialModel())
-	{
-		if (UDynamicMaterialModelEditorOnlyData* EditorOnlyData = UDynamicMaterialModelEditorOnlyData::Get(MaterialModel))
-		{
-			return EditorOnlyData->IsPixelAnimationFlagSet() == bInValue
-				? ECheckBoxState::Checked
-				: ECheckBoxState::Unchecked;
-		}
-	}
-
-	return ECheckBoxState::Undetermined;
-}
-
-void SDMMaterialWizard::Animated_OnChange(ECheckBoxState InState, bool bInValue)
-{
-	if (InState == ECheckBoxState::Checked)
-	{
-		if (UDynamicMaterialModel* MaterialModel = GetMaterialModel())
-		{
-			if (UDynamicMaterialModelEditorOnlyData* EditorOnlyData = UDynamicMaterialModelEditorOnlyData::Get(MaterialModel))
-			{
-				EditorOnlyData->SetPixelAnimationFlag(bInValue);
-			}
-		}
-	}
-}
-
-bool SDMMaterialWizard::TwoSided_IsEnabled() const
-{
-	if (UDynamicMaterialModel* MaterialModel = GetMaterialModel())
-	{
-		return !!UDynamicMaterialModelEditorOnlyData::Get(MaterialModel);
-	}
-
-	return false;
-}
-
-ECheckBoxState SDMMaterialWizard::TwoSided_GetState(bool bInValue) const
-{
-	if (UDynamicMaterialModel* MaterialModel = GetMaterialModel())
-	{
-		if (UDynamicMaterialModelEditorOnlyData* EditorOnlyData = UDynamicMaterialModelEditorOnlyData::Get(MaterialModel))
-		{
-			return EditorOnlyData->IsTwoSidedFlagSet() == bInValue
-				? ECheckBoxState::Checked
-				: ECheckBoxState::Unchecked;
-		}
-	}
-
-	return ECheckBoxState::Undetermined;
-}
-
-void SDMMaterialWizard::TwoSided_OnChange(ECheckBoxState InState, bool bInValue)
-{
-	if (InState == ECheckBoxState::Checked)
-	{
-		if (UDynamicMaterialModel* MaterialModel = GetMaterialModel())
-		{
-			if (UDynamicMaterialModelEditorOnlyData* EditorOnlyData = UDynamicMaterialModelEditorOnlyData::Get(MaterialModel))
-			{
-				EditorOnlyData->SetTwoSidedFlag(bInValue);
-			}
-		}
-	}
-}
-
 ECheckBoxState SDMMaterialWizard::Preset_GetState(FName InPresetName) const
 {
 	return CurrentPreset == InPresetName
diff --git a/Engine/Plugins/Experimental/DynamicMaterial/Source/DynamicMaterialEditor/Private/Slate/SDMMaterialWizard.h b/Engine/Plugins/Experimental/DynamicMaterial/Source/DynamicMaterialEditor/Private/Slate/SDMMaterialWizard.h
--- a/Engine/Plugins/Experimental/DynamicMaterial/Source/DynamicMaterialEditor/Private/Slate/SDMMaterialWizard.h
+++ b/Engine/Plugins/Experimental/DynamicMaterial/Source/DynamicMaterialEditor/Private/Slate/SDMMaterialWizard.h
@@ -32,35 +32,10 @@ protected:
 	TSharedPtr<SBox> PresetChannelContainer;
 
 	TSharedRef<SWidget> CreateLayout();
-	TSharedRef<SWidget> CreateMaterialDomainOptions();
-	TSharedRef<SWidget> CreateBlendModeOptions();
-	TSharedRef<SWidget> CreateShadingModelOptions();
-	TSharedRef<SWidget> CreateAnimationOptions();
-	TSharedRef<SWidget> CreateTwoSidedOptions();
 	TSharedRef<SWidget> CreateChannelPresets();
 	TSharedRef<SWidget> CreateChannelList();
 	TSharedRef<SWidget> CreateAcceptButton();
 
-	bool Domain_IsEnabled() const;
-	ECheckBoxState Domain_GetState(EMaterialDomain InDomain) const;
-	void Domain_OnChange(ECheckBoxState InState, EMaterialDomain InDomain);
-
-	bool BlendMode_IsEnabled() const;
-	ECheckBoxState BlendMode_GetState(EBlendMode InBlendMode) const;
-	void BlendMode_OnChange(ECheckBoxState InState, EBlendMode InBlendMode);
-
-	bool Unlit_IsEnabled() const;
-	ECheckBoxState Unlit_GetState(EDMMaterialShadingModel InValue) const;
-	void Unlit_OnChange(ECheckBoxState InState, EDMMaterialShadingModel InValue);
-
-	bool Animated_IsEnabled() const;
-	ECheckBoxState Animated_GetState(bool bInValue) const;
-	void Animated_OnChange(ECheckBoxState InState, bool bInValue);
-
-	bool TwoSided_IsEnabled() const;
-	ECheckBoxState TwoSided_GetState(bool bInValue) const;
-	void TwoSided_OnChange(ECheckBoxState InState, bool bInValue);
-
 	ECheckBoxState Preset_GetState(FName InPresetName) const;
 	void Preset_OnChange(ECheckBoxState InState, FName InPresetName);
 
diff --git a/Engine/Plugins/Experimental/DynamicMaterial/Source/DynamicMaterialEditor/Public/IDynamicMaterialEditorModule.h b/Engine/Plugins/Experimental/DynamicMaterial/Source/DynamicMaterialEditor/Public/IDynamicMaterialEditorModule.h
--- a/Engine/Plugins/Experimental/DynamicMaterial/Source/DynamicMaterialEditor/Public/IDynamicMaterialEditorModule.h
+++ b/Engine/Plugins/Experimental/DynamicMaterial/Source/DynamicMaterialEditor/Public/IDynamicMaterialEditorModule.h
@@ -10,7 +10,7 @@ class UClass;
 class UDynamicMaterialModel;
 class UObject;
 struct FDMObjectMaterialProperty;
-struct IDMMaterialModelCreatedCallback;
+struct IDMOnWizardCompleteCallback;
 
 DECLARE_DELEGATE_RetVal_OneParam(TArray<FDMObjectMaterialProperty>, FDMGetObjectMaterialPropertiesDelegate, UObject* InObject)
 
@@ -35,12 +35,12 @@ public:
 
 	virtual void RegisterCustomMaterialPropertyGenerator(UClass* InClass, FDMGetObjectMaterialPropertiesDelegate InGenerator) = 0;
 
-	virtual void RegisterMaterialModelCreatedCallback(const TSharedRef<IDMMaterialModelCreatedCallback> InCallback) = 0;
+	virtual void RegisterMaterialModelCreatedCallback(const TSharedRef<IDMOnWizardCompleteCallback> InCallback) = 0;
 
-	virtual void UnregisterMaterialModelCreatedCallback(const TSharedRef<IDMMaterialModelCreatedCallback> InCallback) = 0;
+	virtual void UnregisterMaterialModelCreatedCallback(const TSharedRef<IDMOnWizardCompleteCallback> InCallback) = 0;
 
 	template<typename InCallbackType, typename... InArgsType
-		UE_REQUIRES(std::is_base_of_v<IDMMaterialModelCreatedCallback, InCallbackType>)>
+		UE_REQUIRES(std::is_base_of_v<IDMOnWizardCompleteCallback, InCallbackType>)>
 	TSharedRef<InCallbackType> RegisterMaterialModelCreatedCallback(InArgsType&&... InArgs)
 	{
 		TSharedRef<InCallbackType> NewCallback = MakeShared<InCallbackType>(Forward<InArgsType>(InArgs)...);
diff --git a/Engine/Plugins/Experimental/DynamicMaterial/Source/DynamicMaterialEditor/Public/Model/DMModelCreatedCallback.h b/Engine/Plugins/Experimental/DynamicMaterial/Source/DynamicMaterialEditor/Public/Model/DMModelCreatedCallback.h
deleted file mode 100644
diff --git a/Engine/Plugins/Experimental/DynamicMaterial/Source/DynamicMaterialEditor/Public/Model/DMOnWizardCompleteCallback.h b/Engine/Plugins/Experimental/DynamicMaterial/Source/DynamicMaterialEditor/Public/Model/DMOnWizardCompleteCallback.h
new file mode 100644
--- /dev/null
+++ b/Engine/Plugins/Experimental/DynamicMaterial/Source/DynamicMaterialEditor/Public/Model/DMOnWizardCompleteCallback.h
@@ -0,0 +1,69 @@
+// Copyright Epic Games, Inc. All Rights Reserved.
+
+#pragma once
+
+#include "Delegates/Delegate.h"
+#include "Delegates/DelegateCombinations.h"
+#include "HAL/Platform.h"
+#include "Templates/SharedPointer.h"
+
+class AActor;
+class UActorComponent;
+class UDynamicMaterialModelEditorOnlyData;
+class UDynamicMaterialModel;
+class UObject;
+
+struct FDMOnWizardCompleteCallbackParams
+{
+	UDynamicMaterialModel* MaterialModel;
+	UDynamicMaterialModelEditorOnlyData* EditorOnlyData;
+	UObject* Outer;
+	UActorComponent* OuterComponent;
+	AActor* OuterActor;
+};
+
+// Callback interface
+struct IDMOnWizardCompleteCallback : public TSharedPtr<IDMOnWizardCompleteCallback>
+{
+	virtual ~IDMOnWizardCompleteCallback() = default;
+
+	virtual uint32 GetPriority() const = 0;
+	virtual void OnModelCreated(const FDMOnWizardCompleteCallbackParams& InParams) = 0;
+
+	virtual bool operator<(const IDMOnWizardCompleteCallback& InOther) const
+	{
+		return GetPriority() < InOther.GetPriority();
+	}
+};
+
+// Default implementation
+struct FDMMaterialModelCreatedCallbackBase : public IDMOnWizardCompleteCallback
+{
+	FDMMaterialModelCreatedCallbackBase(uint32 InPriority);
+
+	virtual ~FDMMaterialModelCreatedCallbackBase() override = default;
+
+	//~ Begin IDMMaterialModelCreatedCallback
+	virtual uint32 GetPriority() const override;
+	//~ End IDMMaterialModelCreatedCallback
+
+protected:
+	uint32 Priority;
+};
+
+// Default implementation
+struct FDMMaterialModelCreatedCallbackDelegate : public FDMMaterialModelCreatedCallbackBase
+{
+	DECLARE_DELEGATE_OneParam(FOnModelCreated, const FDMOnWizardCompleteCallbackParams&)
+
+	FDMMaterialModelCreatedCallbackDelegate(uint32 InPriority, const FOnModelCreated& InOnModelCreatedDelegate);
+
+	virtual ~FDMMaterialModelCreatedCallbackDelegate() override = default;
+
+	//~ Begin IDMMaterialModelCreatedCallback
+	virtual void OnModelCreated(const FDMOnWizardCompleteCallbackParams& InParams) override;
+	//~ End IDMMaterialModelCreatedCallback
+
+protected:
+	FOnModelCreated OnModelCreatedDelegate;
+};
