Material Designer
- Fixed an issue when dragging a texture across the viewport into the material design would open the material instance of the preview actor.

diff --git a/Engine/Plugins/Experimental/DynamicMaterial/Source/DynamicMaterialEditor/Private/Material/DynamicMaterialInstanceFactory.cpp b/Engine/Plugins/Experimental/DynamicMaterial/Source/DynamicMaterialEditor/Private/Material/DynamicMaterialInstanceFactory.cpp
--- a/Engine/Plugins/Experimental/DynamicMaterial/Source/DynamicMaterialEditor/Private/Material/DynamicMaterialInstanceFactory.cpp
+++ b/Engine/Plugins/Experimental/DynamicMaterial/Source/DynamicMaterialEditor/Private/Material/DynamicMaterialInstanceFactory.cpp
@@ -4,6 +4,7 @@
 #include "AssetToolsModule.h"
 #include "DMPrivate.h"
 #include "EngineAnalytics.h"
+#include "GameFramework/Actor.h"
 #include "IAssetTools.h"
 #include "Material/DynamicMaterialInstance.h"
 #include "Model/DynamicMaterialModel.h"
@@ -51,6 +52,18 @@ UObject* UDynamicMaterialInstanceFactory::FactoryCreateNew(UClass* Class, UObjec
 
 	NewInstance->InitializeMIDPublic();
 
+	if (InParent)
+	{
+		if (AActor* Actor = InParent->GetTypedOuter<AActor>())
+		{
+			if (Actor->bIsEditorPreviewActor)
+			{
+				// If it is a preview actor do not trigger analytics or open in editor.
+				return NewInstance;
+			}
+		}
+	}
+
 	IAssetTools& AssetTools = FModuleManager::LoadModuleChecked<FAssetToolsModule>("AssetTools").Get();
 	AssetTools.OpenEditorForAssets({NewInstance});
 
