Material Designer
- Ability to rename parameters in the details panel.

diff --git a/Engine/Plugins/Experimental/DynamicMaterial/Source/DynamicMaterial/Private/Components/DMMaterialValue.cpp b/Engine/Plugins/Experimental/DynamicMaterial/Source/DynamicMaterial/Private/Components/DMMaterialValue.cpp
--- a/Engine/Plugins/Experimental/DynamicMaterial/Source/DynamicMaterial/Private/Components/DMMaterialValue.cpp
+++ b/Engine/Plugins/Experimental/DynamicMaterial/Source/DynamicMaterial/Private/Components/DMMaterialValue.cpp
@@ -26,7 +26,6 @@
 TMap<EDMValueType, TStrongObjectPtr<UClass>> UDMMaterialValue::TypeClasses = {};
 
 const FString UDMMaterialValue::ParameterPathToken = FString(TEXT("Parameter"));
-const TCHAR* UDMMaterialValue::ParameterNamePrefix = TEXT("VALUE_");
 
 #if WITH_EDITOR
 const FName UDMMaterialValue::ValueName = "Value";
@@ -151,7 +150,7 @@ void UDMMaterialValue::BeginDestroy()
 
 FName UDMMaterialValue::GetMaterialParameterName() const
 {
-	return Parameter ? FName(ParameterNamePrefix + Parameter->GetParameterName().ToString()) : GetFName();
+	return Parameter ? Parameter->GetParameterName() : GetFName();
 }
 
 UDMMaterialComponent* UDMMaterialValue::GetSubComponentByPath(FDMComponentPath& InPath, const FDMComponentPathSegment& InPathSegment) const
diff --git a/Engine/Plugins/Experimental/DynamicMaterial/Source/DynamicMaterial/Private/Components/DMTextureUV.cpp b/Engine/Plugins/Experimental/DynamicMaterial/Source/DynamicMaterial/Private/Components/DMTextureUV.cpp
--- a/Engine/Plugins/Experimental/DynamicMaterial/Source/DynamicMaterial/Private/Components/DMTextureUV.cpp
+++ b/Engine/Plugins/Experimental/DynamicMaterial/Source/DynamicMaterial/Private/Components/DMTextureUV.cpp
@@ -217,6 +217,7 @@ TSharedPtr<IPropertyHandle> UDMTextureUV::GetPropertyHandle(FName InProperty)
 
 	return PropertyHandles.FindChecked(InProperty);
 }
+#endif
 
 TArray<UDMMaterialParameter*> UDMTextureUV::GetParameters() const
 {
@@ -231,6 +232,92 @@ TArray<UDMMaterialParameter*> UDMTextureUV::GetParameters() const
 	return Parameters;
 }
 
+UDMMaterialParameter* UDMTextureUV::GetMaterialParameter(FName InPropertyName, int32 InComponent) const
+{
+	using namespace UE::DynamicMaterial;
+
+	const TObjectPtr<UDMMaterialParameter>* ParameterPtr = nullptr;
+
+	if (InPropertyName == NAME_Offset)
+	{
+		switch (InComponent)
+		{
+			case 0:
+				ParameterPtr = MaterialParameters.Find(ParamID::OffsetX);
+				break;
+
+			case 1:
+				ParameterPtr = MaterialParameters.Find(ParamID::OffsetY);
+				break;
+		}
+	}
+	else if (InPropertyName == NAME_Pivot)
+	{
+		switch (InComponent)
+		{
+			case 0:
+				ParameterPtr = MaterialParameters.Find(ParamID::PivotX);
+				break;
+
+			case 1:
+				ParameterPtr = MaterialParameters.Find(ParamID::PivotY);
+				break;
+		}
+	}
+	else if (InPropertyName == NAME_Rotation)
+	{
+		switch (InComponent)
+		{
+			case 0:
+				ParameterPtr = MaterialParameters.Find(ParamID::Rotation);
+				break;
+		}
+	}
+	else if (InPropertyName == NAME_Scale)
+	{
+		switch (InComponent)
+		{
+			case 0:
+				ParameterPtr = MaterialParameters.Find(ParamID::ScaleX);
+				break;
+
+			case 1:
+				ParameterPtr = MaterialParameters.Find(ParamID::ScaleY);
+				break;
+		}
+	}
+
+	if (ParameterPtr)
+	{
+		return (*ParameterPtr).Get();
+	}
+
+	return nullptr;
+}
+
+FName UDMTextureUV::GetMaterialParameterName(FName InPropertyName, int32 InComponent) const
+{
+	if (UDMMaterialParameter* Parameter = GetMaterialParameter(InPropertyName, InComponent))
+	{
+		return Parameter->GetParameterName();
+	}
+
+	return NAME_None;
+}
+
+#if WITH_EDITOR
+bool UDMTextureUV::SetMaterialParameterName(FName InPropertyName, int32 InComponent, FName InNewName)
+{
+	if (UDMMaterialParameter* Parameter = GetMaterialParameter(InPropertyName, InComponent))
+	{
+		const FName CurrentName = Parameter->GetParameterName();
+		Parameter->RenameParameter(InNewName);
+		return Parameter->GetParameterName() != CurrentName;
+	}
+
+	return false;
+}
+
 void UDMTextureUV::PostEditorDuplicate(UDynamicMaterialModel* InMaterialModel, UDMMaterialComponent* InParent)
 {
 	Super::PostEditorDuplicate(InMaterialModel, InParent);
diff --git a/Engine/Plugins/Experimental/DynamicMaterial/Source/DynamicMaterial/Private/Model/DynamicMaterialModel.cpp b/Engine/Plugins/Experimental/DynamicMaterial/Source/DynamicMaterial/Private/Model/DynamicMaterialModel.cpp
--- a/Engine/Plugins/Experimental/DynamicMaterial/Source/DynamicMaterial/Private/Model/DynamicMaterialModel.cpp
+++ b/Engine/Plugins/Experimental/DynamicMaterial/Source/DynamicMaterial/Private/Model/DynamicMaterialModel.cpp
@@ -241,6 +241,11 @@ UDMMaterialParameter* UDynamicMaterialModel::CreateUniqueParameter(FName InBaseN
 	ParameterMap.Emplace(NewParameter->GetParameterName(), NewParameter);
 	NewParameter->SetComponentState(EDMComponentLifetimeState::Added);
 
+	if (IDynamicMaterialModelEditorOnlyDataInterface* ModelEditorOnlyData = GetEditorOnlyData())
+	{
+		ModelEditorOnlyData->RequestMaterialBuild();
+	}
+
 	return NewParameter;
 }
 
@@ -261,6 +266,11 @@ void UDynamicMaterialModel::RenameParameter(UDMMaterialParameter* InParameter, F
 
 	InParameter->ParameterName = CreateUniqueParameterName(InBaseName);
 	ParameterMap.Emplace(InParameter->ParameterName, InParameter);
+
+	if (IDynamicMaterialModelEditorOnlyDataInterface* ModelEditorOnlyData = GetEditorOnlyData())
+	{
+		ModelEditorOnlyData->RequestMaterialBuild();
+	}
 }
 
 void UDynamicMaterialModel::FreeParameter(UDMMaterialParameter* InParameter)
@@ -289,6 +299,11 @@ void UDynamicMaterialModel::FreeParameter(UDMMaterialParameter* InParameter)
 
 	InParameter->ParameterName = NAME_None;
 	InParameter->SetComponentState(EDMComponentLifetimeState::Removed);
+
+	if (IDynamicMaterialModelEditorOnlyDataInterface* ModelEditorOnlyData = GetEditorOnlyData())
+	{
+		ModelEditorOnlyData->RequestMaterialBuild();
+	}
 }
 
 bool UDynamicMaterialModel::ConditionalFreeParameter(UDMMaterialParameter* InParameter)
diff --git a/Engine/Plugins/Experimental/DynamicMaterial/Source/DynamicMaterial/Public/Components/DMMaterialValue.h b/Engine/Plugins/Experimental/DynamicMaterial/Source/DynamicMaterial/Public/Components/DMMaterialValue.h
--- a/Engine/Plugins/Experimental/DynamicMaterial/Source/DynamicMaterial/Public/Components/DMMaterialValue.h
+++ b/Engine/Plugins/Experimental/DynamicMaterial/Source/DynamicMaterial/Public/Components/DMMaterialValue.h
@@ -41,7 +41,6 @@ class DYNAMICMATERIAL_API UDMMaterialValue : public UDMMaterialLinkedComponent
  
 public:
 	static const FString ParameterPathToken;
-	static const TCHAR* ParameterNamePrefix;
 
 #if WITH_EDITOR
 	static const FName ValueName;
diff --git a/Engine/Plugins/Experimental/DynamicMaterial/Source/DynamicMaterial/Public/Components/DMTextureUV.h b/Engine/Plugins/Experimental/DynamicMaterial/Source/DynamicMaterial/Public/Components/DMTextureUV.h
--- a/Engine/Plugins/Experimental/DynamicMaterial/Source/DynamicMaterial/Public/Components/DMTextureUV.h
+++ b/Engine/Plugins/Experimental/DynamicMaterial/Source/DynamicMaterial/Public/Components/DMTextureUV.h
@@ -138,8 +138,20 @@ public:
 	TSharedPtr<IPropertyHandle> GetPropertyHandle(FName InProperty);
 #endif
 
+	UFUNCTION(BlueprintPure, Category = "Material Designer")
 	TArray<UDMMaterialParameter*> GetParameters() const;
 
+	UFUNCTION(BlueprintPure, Category = "Material Designer")
+	UDMMaterialParameter* GetMaterialParameter(FName InPropertyName, int32 InComponent) const;
+
+	UFUNCTION(BlueprintPure, Category = "Material Designer")
+	FName GetMaterialParameterName(FName InPropertyName, int32 InComponent) const;
+
+#if WITH_EDITOR
+	UFUNCTION(BlueprintCallable, Category = "Material Designer")
+	bool SetMaterialParameterName(FName InPropertyName, int32 InComponent, FName InNewName);
+#endif
+
 	void SetMIDParameters(UMaterialInstanceDynamic* InMID);
 
 	//~ Begin UObject
diff --git a/Engine/Plugins/Experimental/DynamicMaterial/Source/DynamicMaterialEditor/Private/DetailsPanel/DMValueDetailsRowExtensions.cpp b/Engine/Plugins/Experimental/DynamicMaterial/Source/DynamicMaterialEditor/Private/DetailsPanel/DMValueDetailsRowExtensions.cpp
new file mode 100644
--- /dev/null
+++ b/Engine/Plugins/Experimental/DynamicMaterial/Source/DynamicMaterialEditor/Private/DetailsPanel/DMValueDetailsRowExtensions.cpp
@@ -0,0 +1,400 @@
+// Copyright Epic Games, Inc. All Rights Reserved.
+
+#include "DetailsPanel/DMValueDetailsRowExtensions.h"
+#include "Components/DMMaterialParameter.h"
+#include "Components/DMMaterialValue.h"
+#include "Components/DMTextureUV.h"
+#include "Delegates/IDelegateInstance.h"
+#include "DetailRowMenuContext.h"
+#include "DynamicMaterialEditorModule.h"
+#include "Framework/SlateDelegates.h"
+#include "Modules/ModuleManager.h"
+#include "PropertyEditorDelegates.h"
+#include "PropertyEditorModule.h"
+#include "PropertyHandle.h"
+#include "ScopedTransaction.h"
+#include "Styling/CoreStyle.h"
+#include "ToolMenu.h"
+#include "ToolMenus.h"
+#include "ToolMenuSection.h"
+#include "Widgets/Input/SEditableTextBox.h"
+
+#define LOCTEXT_NAMESPACE "DMValueRowExtensions"
+
+namespace UE::DynamicMaterialEditor::Private
+{
+	constexpr int32 ComponentX = 0;
+	constexpr int32 ComponentY = 1;
+	const FName RowExtensionName = "DMValueRowExtensionContextSection";
+	const FName MaterialDesignerMenuName = "MaterialDesigner";
+	const FName PropertyEditorModuleName = "PropertyEditor";
+	const FText SectionName = LOCTEXT("MaterialDesigner", "Material Designer");
+	const FText MenuEntry = LOCTEXT("ParameterName", "Parameter Name");
+	const FText MenuEntryX = LOCTEXT("ParameterNameX", "Parameter Name X");
+	const FText MenuEntryY = LOCTEXT("ParameterNameY", "Parameter Name Y");
+	const FText ToolTip = LOCTEXT("SetParameterNameToolTip", "Set the name of the parameter this property is exposed as within the generated material.");
+	const FText InvalidCharacterError = LOCTEXT("InvalidCharacterError", "Valid characters are A-Z, A-z, _ and -");
+	const FText TooShortError = LOCTEXT("TooShortError", "Min 3 characters.");
+	const FText TooLongError = LOCTEXT("TooLongError", "Max 50 characters.");
+	const TArray<TPair<TCHAR, TCHAR>> ValidCharacterRanges = {
+		{'A', 'Z'},
+		{'a', 'z'},
+		{'0', '9'},
+		{'-', '-'},
+		{'_', '_'}
+	};
+}
+
+FDMValueDetailsRowExtensions& FDMValueDetailsRowExtensions::Get()
+{
+	static FDMValueDetailsRowExtensions Instance;
+	return Instance;
+}
+
+FDMValueDetailsRowExtensions::~FDMValueDetailsRowExtensions()
+{
+	UnregisterRowExtensions();
+}
+
+void FDMValueDetailsRowExtensions::RegisterRowExtensions()
+{
+	using namespace UE::DynamicMaterialEditor::Private;
+
+	FPropertyEditorModule& PropertyEditorModule = FModuleManager::LoadModuleChecked<FPropertyEditorModule>(PropertyEditorModuleName);
+	FOnGenerateGlobalRowExtension& RowExtensionDelegate = PropertyEditorModule.GetGlobalRowExtensionDelegate();
+	RowExtensionHandle = RowExtensionDelegate.AddStatic(&HandleCreatePropertyRowExtension);
+}
+
+void FDMValueDetailsRowExtensions::UnregisterRowExtensions()
+{
+	using namespace UE::DynamicMaterialEditor::Private;
+
+	if (RowExtensionHandle.IsValid() && FModuleManager::Get().IsModuleLoaded(PropertyEditorModuleName))
+	{
+		FPropertyEditorModule& PropertyEditorModule = FModuleManager::LoadModuleChecked<FPropertyEditorModule>(PropertyEditorModuleName);
+		PropertyEditorModule.GetGlobalRowExtensionDelegate().Remove(RowExtensionHandle);
+		RowExtensionHandle.Reset();
+	}
+}
+
+void FDMValueDetailsRowExtensions::HandleCreatePropertyRowExtension(const FOnGenerateGlobalRowExtensionArgs& InArgs, TArray<FPropertyRowExtensionButton>& OutExtensions)
+{
+	using namespace UE::DynamicMaterialEditor::Private;
+
+	if (!InArgs.Property && !InArgs.PropertyHandle.IsValid())
+	{
+		return;
+	}
+
+	UToolMenus* Menus = UToolMenus::Get();
+	check(Menus);
+
+	UToolMenu* ContextMenu = Menus->FindMenu(UE::PropertyEditor::RowContextMenuName);
+
+	if (ContextMenu->ContainsSection(RowExtensionName))
+	{
+		return;
+	}
+
+	ContextMenu->AddDynamicSection(
+		RowExtensionName,
+		FNewToolMenuDelegate::CreateStatic(&FillPropertyRightClickMenu)
+	);
+}
+
+void FDMValueDetailsRowExtensions::FillPropertyRightClickMenu(UToolMenu* InToolMenu)
+{
+	using namespace UE::DynamicMaterialEditor::Private;
+
+	UDetailRowMenuContext* RowMenuContext = InToolMenu->FindContext<UDetailRowMenuContext>();
+
+	if (!RowMenuContext)
+	{
+		return;
+	}
+
+	TSharedPtr<IPropertyHandle> PropertyHandle;
+
+	for (const TSharedPtr<IPropertyHandle>& ContextPropertyHandle : RowMenuContext->PropertyHandles)
+	{
+		if (ContextPropertyHandle.IsValid())
+		{
+			PropertyHandle = ContextPropertyHandle;
+			break;
+		}
+	}
+
+	if (!PropertyHandle.IsValid())
+	{
+		return;
+	}
+
+	if (RowMenuContext->DetailsView && !RowMenuContext->DetailsView->IsPropertyEditingEnabled())
+	{
+		return;
+	}
+
+	if (PropertyHandle->IsEditConst() || !PropertyHandle->IsEditable())
+	{
+		return;
+	}
+
+	TArray<UObject*> Outers;
+	PropertyHandle->GetOuterObjects(Outers);
+
+	if (Outers.IsEmpty())
+	{
+		return;
+	}
+
+	FText CurrentParameterName;
+
+	FToolMenuSection& Section = InToolMenu->AddSection(MaterialDesignerMenuName, SectionName);
+
+	if (UDMMaterialValue* Value = Cast<UDMMaterialValue>(Outers[0]))
+	{
+		FillPropertyRightClickMenu_Value(Section, Value);
+	}
+	else if (UDMTextureUV* TextureUV = Cast<UDMTextureUV>(Outers[0]))
+	{
+		FillPropertyRightClickMenu_TextureUV(Section, TextureUV, PropertyHandle->GetProperty()->GetFName());
+	}
+}
+
+void FDMValueDetailsRowExtensions::FillPropertyRightClickMenu_TextureUV(FToolMenuSection& InSection, UDMTextureUV* InTextureUV, FName InPropertyName)
+{
+	using namespace UE::DynamicMaterialEditor::Private;
+
+	TWeakObjectPtr<UDMTextureUV> TextureUVWeak = InTextureUV;
+
+	auto CreateMenuEntry = [&InSection, &TextureUVWeak, &InPropertyName](const FText& InMenuText, int32 InComponent)
+		{
+			FNewToolMenuChoice RenameChoice(FOnGetContent::CreateLambda(
+				[TextureUVWeak, InPropertyName, InComponent]() -> TSharedRef<SWidget>
+				{
+					UDMTextureUV* TextureUV = TextureUVWeak.Get();
+
+					if (!IsValid(TextureUV))
+					{
+						return SNullWidget::NullWidget;
+					}
+
+					return SNew(SEditableTextBox)
+						.Text(FText::FromName(TextureUV->GetMaterialParameterName(InPropertyName, InComponent)))
+						.OnVerifyTextChanged_Static(&VerifyParameterName)
+						.AllowContextMenu(false)
+						.ClearKeyboardFocusOnCommit(true)
+						.MinDesiredWidth(100.f)
+						.OnTextCommitted_Lambda(
+							[TextureUVWeak, InPropertyName](const FText& InText, ETextCommit::Type InCommitType)
+							{
+								if (InCommitType != ETextCommit::OnEnter)
+								{
+									return;
+								}
+
+								UDMTextureUV* TextureUV = TextureUVWeak.Get();
+
+								if (!IsValid(TextureUV))
+								{
+									return;
+								}
+
+								FText Unused;
+
+								if (!VerifyParameterName(InText, Unused))
+								{
+									return;
+								}
+
+								SetTextureUVParameterName(TextureUVWeak, InPropertyName, ComponentX, FName(*InText.ToString()));
+							});
+				}));
+
+			FName MenuName = FName(InPropertyName, InComponent);
+
+			InSection.AddSubMenu(
+				MenuName,
+				InMenuText,
+				ToolTip,
+				RenameChoice,
+				/* Open submenu on click */ true,
+				TAttribute<FSlateIcon>(),
+				/* Should close after selection */ false
+			);
+		};
+
+	if (InPropertyName == UDMTextureUV::NAME_Rotation)
+	{
+		CreateMenuEntry(MenuEntry, ComponentX);
+	}
+
+	if (InPropertyName == UDMTextureUV::NAME_Offset 
+		|| InPropertyName == UDMTextureUV::NAME_Pivot
+		|| InPropertyName == UDMTextureUV::NAME_Scale)
+	{
+		CreateMenuEntry(MenuEntryX, ComponentX);
+		CreateMenuEntry(MenuEntryY, ComponentY);
+	}
+}
+
+void FDMValueDetailsRowExtensions::FillPropertyRightClickMenu_Value(FToolMenuSection& InSection, UDMMaterialValue* InValue)
+{
+	using namespace UE::DynamicMaterialEditor::Private;
+
+	TWeakObjectPtr<UDMMaterialValue> ValueWeak = InValue;
+
+	FNewToolMenuChoice RenameChoice(FOnGetContent::CreateLambda(
+		[ValueWeak]() -> TSharedRef<SWidget>
+		{
+			UDMMaterialValue* Value = ValueWeak.Get();
+
+			if (!IsValid(Value))
+			{
+				return SNullWidget::NullWidget;
+			}
+
+			return SNew(SEditableTextBox)
+				.Text(FText::FromName(Value->GetMaterialParameterName()))
+				.OnVerifyTextChanged_Static(&VerifyParameterName)
+				.AllowContextMenu(false)
+				.ClearKeyboardFocusOnCommit(true)
+				.MinDesiredWidth(100.f)
+				.OnTextCommitted_Lambda(
+					[ValueWeak](const FText& InText, ETextCommit::Type InCommitType)
+					{
+						if (InCommitType != ETextCommit::OnEnter)
+						{
+							return;
+						}
+
+						UDMMaterialValue* Value = ValueWeak.Get();
+
+						if (!IsValid(Value))
+						{
+							return;
+						}
+
+						FText Unused;
+
+						if (!VerifyParameterName(InText, Unused))
+						{
+							return;
+						}
+
+						SetValueParameterName(ValueWeak, FName(*InText.ToString()));
+					});
+		}));
+
+	InSection.AddSubMenu(
+		UDMMaterialValue::ValueName,
+		MenuEntry,
+		ToolTip,
+		RenameChoice,
+		/* Open submenu on click */ true,
+		TAttribute<FSlateIcon>(),
+		/* Should close after selection */ false
+	);
+}
+
+void FDMValueDetailsRowExtensions::SetValueParameterName(TWeakObjectPtr<UDMMaterialValue> InValueWeak, FName InName)
+{
+	using namespace UE::DynamicMaterialEditor::Private;
+
+	UDMMaterialValue* Value = InValueWeak.Get();
+
+	if (!Value)
+	{
+		return;
+	}
+
+	if (InName == NAME_None)
+	{
+		return;
+	}
+
+	const FName CurrentName = Value->GetMaterialParameterName();
+
+	if (CurrentName == InName)
+	{
+		return;
+	}
+
+	FScopedTransaction Transaction(LOCTEXT("SetParameterName", "Material Designer Set Parameter Name"));
+	Value->Modify();
+	Value->SetParameterName(InName);
+
+	UE_LOG(LogDynamicMaterialEditor, Log, TEXT("Parameter renamed to: %s"), *Value->GetMaterialParameterName().ToString());
+}
+
+void FDMValueDetailsRowExtensions::SetTextureUVParameterName(TWeakObjectPtr<UDMTextureUV> InTextureUVWeak, FName InPropertyName, 
+	int32 InComponent, FName InName)
+{
+	using namespace UE::DynamicMaterialEditor::Private;
+
+	UDMTextureUV* TextureUV = InTextureUVWeak.Get();
+
+	if (!TextureUV)
+	{
+		return;
+	}
+
+	if (InName == NAME_None)
+	{
+		return;
+	}
+
+	const FName CurrentName = TextureUV->GetMaterialParameterName(InPropertyName, InComponent);
+
+	if (CurrentName == InName)
+	{
+		return;
+	}
+
+	TextureUV->SetMaterialParameterName(InPropertyName, InComponent, InName);
+	UE_LOG(LogDynamicMaterialEditor, Log, TEXT("Parameter renamed to: %s"), *TextureUV->GetMaterialParameterName(InPropertyName, InComponent).ToString());
+}
+
+bool FDMValueDetailsRowExtensions::VerifyParameterName(const FText& InValue, FText& OutErrorText)
+{
+	using namespace UE::DynamicMaterialEditor::Private;
+
+	const FString ValueStr = InValue.ToString();
+	const int32 ValueLen = ValueStr.Len();
+
+	if (ValueLen < 3)
+	{
+		OutErrorText = TooShortError;
+		return false;
+	}
+
+	if (ValueLen > 50)
+	{
+		OutErrorText = TooLongError;
+		return false;
+	}
+
+	for (int32 Index = 0; Index < ValueLen; ++Index)
+	{
+		bool bInRange = false;
+
+		for (const TPair<TCHAR, TCHAR>& ValidCharacterRange : ValidCharacterRanges)
+		{
+			if (ValueStr[Index] >= ValidCharacterRange.Key && ValueStr[Index] <= ValidCharacterRange.Value)
+			{
+				bInRange = true;
+				break;
+			}
+		}
+
+		if (!bInRange)
+		{
+			OutErrorText = InvalidCharacterError;
+			return false;
+		}
+	}
+
+	return true;
+}
+
+#undef LOCTEXT_NAMESPACE
diff --git a/Engine/Plugins/Experimental/DynamicMaterial/Source/DynamicMaterialEditor/Private/DetailsPanel/DMValueDetailsRowExtensions.h b/Engine/Plugins/Experimental/DynamicMaterial/Source/DynamicMaterialEditor/Private/DetailsPanel/DMValueDetailsRowExtensions.h
new file mode 100644
--- /dev/null
+++ b/Engine/Plugins/Experimental/DynamicMaterial/Source/DynamicMaterialEditor/Private/DetailsPanel/DMValueDetailsRowExtensions.h
@@ -0,0 +1,46 @@
+// Copyright Epic Games, Inc. All Rights Reserved.
+
+#pragma once
+
+#include "Containers/ContainersFwd.h"
+#include "Delegates/IDelegateInstance.h"
+#include "UObject/WeakObjectPtrFwd.h"
+
+class FName;
+class UDMMaterialValue;
+class UDMTextureUV;
+class UToolMenu;
+struct FOnGenerateGlobalRowExtensionArgs;
+struct FPropertyRowExtensionButton;
+struct FToolMenuContext;
+struct FToolMenuSection;
+
+class FDMValueDetailsRowExtensions
+{
+public:
+	static FDMValueDetailsRowExtensions& Get();
+
+	~FDMValueDetailsRowExtensions();
+
+	void RegisterRowExtensions();
+
+	void UnregisterRowExtensions();
+
+private:
+	FDelegateHandle RowExtensionHandle;
+
+	static void HandleCreatePropertyRowExtension(const FOnGenerateGlobalRowExtensionArgs& InArgs, TArray<FPropertyRowExtensionButton>& OutExtensions);
+
+	static void FillPropertyRightClickMenu(UToolMenu* InToolMenu);
+
+	static void FillPropertyRightClickMenu_TextureUV(FToolMenuSection& InSection, UDMTextureUV* InTextureUV, FName InPropertyName);
+
+	static void FillPropertyRightClickMenu_Value(FToolMenuSection& InSection, UDMMaterialValue* InValue);
+
+	static void SetValueParameterName(TWeakObjectPtr<UDMMaterialValue> InValueWeak, FName InName);
+
+	static void SetTextureUVParameterName(TWeakObjectPtr<UDMTextureUV> InTextureUVWeak,
+		FName InPropertyName, int32 InComponent, FName InName);
+
+	static bool VerifyParameterName(const FText& InValue, FText& OutErrorText);
+};
diff --git a/Engine/Plugins/Experimental/DynamicMaterial/Source/DynamicMaterialEditor/Private/DynamicMaterialEditorModule.cpp b/Engine/Plugins/Experimental/DynamicMaterial/Source/DynamicMaterialEditor/Private/DynamicMaterialEditorModule.cpp
--- a/Engine/Plugins/Experimental/DynamicMaterial/Source/DynamicMaterialEditor/Private/DynamicMaterialEditorModule.cpp
+++ b/Engine/Plugins/Experimental/DynamicMaterial/Source/DynamicMaterialEditor/Private/DynamicMaterialEditorModule.cpp
@@ -1,6 +1,7 @@
 // Copyright Epic Games, Inc. All Rights Reserved.
 
 #include "DynamicMaterialEditorModule.h"
+#include "Components/ActorComponent.h"
 #include "Components/DMMaterialComponent.h"
 #include "Components/DMMaterialEffectFunction.h"
 #include "Components/DMMaterialStageThroughput.h"
@@ -15,10 +16,10 @@
 #include "Components/MaterialValues/DMMaterialValueFloat3XYZ.h"
 #include "Components/MaterialValues/DMMaterialValueFloat4.h"
 #include "Components/MaterialValues/DMMaterialValueTexture.h"
-#include "Components/ActorComponent.h"
 #include "Components/PrimitiveComponent.h"
 #include "DetailsPanel/DMMaterialInterfaceTypeCustomizer.h"
 #include "DetailsPanel/DMPropertyTypeCustomizer.h"
+#include "DetailsPanel/DMValueDetailsRowExtensions.h"
 #include "DetailsPanel/Slate/SDMMaterialListExtensionWidget.h"
 #include "DMMaterialFunctionLibrary.h"
 #include "DMWorldSubsystem.h"
@@ -32,11 +33,11 @@
 #include "Material/DynamicMaterialInstance.h"
 #include "MaterialList.h"
 #include "Model/DMMaterialModelDefaults.h"
+#include "Model/DMOnWizardCompleteCallback.h"
 #include "Model/DynamicMaterialModel.h"
 #include "Model/DynamicMaterialModelEditorOnlyData.h"
 #include "Modules/ModuleManager.h"
 #include "PropertyEditorModule.h"
-#include "Model/DMOnWizardCompleteCallback.h"
 #include "Slate/Properties/Editors/SDMPropertyEditBoolValue.h"
 #include "Slate/Properties/Editors/SDMPropertyEditFloat1Value.h"
 #include "Slate/Properties/Editors/SDMPropertyEditFloat2Value.h"
@@ -304,6 +305,8 @@ void FDynamicMaterialEditorModule::StartupModule()
 		{
 			return UDynamicMaterialEditorSettings::Get()->DefaultRGBTexture.LoadSynchronous();
 		});
+
+	FDMValueDetailsRowExtensions::Get().RegisterRowExtensions();
 }
 
 void FDynamicMaterialEditorModule::ShutdownModule()
@@ -327,6 +330,8 @@ void FDynamicMaterialEditorModule::ShutdownModule()
 	BuildRequestList.Empty();
 
 	UDMMaterialValueTexture::GetDefaultRGBTexture.Unbind();
+
+	FDMValueDetailsRowExtensions::Get().UnregisterRowExtensions();
 }
 
 void FDynamicMaterialEditorModule::SetDynamicMaterialModel(UDynamicMaterialModel* InMaterialModel, UWorld* InWorld, bool bInInvokeTab)
