[Backout] - CL32858377
Original CL Desc
-----------------------------------------------------------------
Material Designer
- Ability to rename parameters in the details panel.

diff --git a/Engine/Plugins/Experimental/DynamicMaterial/Source/DynamicMaterial/Private/Components/DMMaterialValue.cpp b/Engine/Plugins/Experimental/DynamicMaterial/Source/DynamicMaterial/Private/Components/DMMaterialValue.cpp
--- a/Engine/Plugins/Experimental/DynamicMaterial/Source/DynamicMaterial/Private/Components/DMMaterialValue.cpp
+++ b/Engine/Plugins/Experimental/DynamicMaterial/Source/DynamicMaterial/Private/Components/DMMaterialValue.cpp
@@ -26,6 +26,7 @@
 TMap<EDMValueType, TStrongObjectPtr<UClass>> UDMMaterialValue::TypeClasses = {};
 
 const FString UDMMaterialValue::ParameterPathToken = FString(TEXT("Parameter"));
+const TCHAR* UDMMaterialValue::ParameterNamePrefix = TEXT("VALUE_");
 
 #if WITH_EDITOR
 const FName UDMMaterialValue::ValueName = "Value";
@@ -150,7 +151,7 @@ void UDMMaterialValue::BeginDestroy()
 
 FName UDMMaterialValue::GetMaterialParameterName() const
 {
-	return Parameter ? Parameter->GetParameterName() : GetFName();
+	return Parameter ? FName(ParameterNamePrefix + Parameter->GetParameterName().ToString()) : GetFName();
 }
 
 UDMMaterialComponent* UDMMaterialValue::GetSubComponentByPath(FDMComponentPath& InPath, const FDMComponentPathSegment& InPathSegment) const
diff --git a/Engine/Plugins/Experimental/DynamicMaterial/Source/DynamicMaterial/Private/Components/DMTextureUV.cpp b/Engine/Plugins/Experimental/DynamicMaterial/Source/DynamicMaterial/Private/Components/DMTextureUV.cpp
--- a/Engine/Plugins/Experimental/DynamicMaterial/Source/DynamicMaterial/Private/Components/DMTextureUV.cpp
+++ b/Engine/Plugins/Experimental/DynamicMaterial/Source/DynamicMaterial/Private/Components/DMTextureUV.cpp
@@ -217,7 +217,6 @@ TSharedPtr<IPropertyHandle> UDMTextureUV::GetPropertyHandle(FName InProperty)
 
 	return PropertyHandles.FindChecked(InProperty);
 }
-#endif
 
 TArray<UDMMaterialParameter*> UDMTextureUV::GetParameters() const
 {
@@ -232,92 +231,6 @@ TArray<UDMMaterialParameter*> UDMTextureUV::GetParameters() const
 	return Parameters;
 }
 
-UDMMaterialParameter* UDMTextureUV::GetMaterialParameter(FName InPropertyName, int32 InComponent) const
-{
-	using namespace UE::DynamicMaterial;
-
-	const TObjectPtr<UDMMaterialParameter>* ParameterPtr = nullptr;
-
-	if (InPropertyName == NAME_Offset)
-	{
-		switch (InComponent)
-		{
-			case 0:
-				ParameterPtr = MaterialParameters.Find(ParamID::OffsetX);
-				break;
-
-			case 1:
-				ParameterPtr = MaterialParameters.Find(ParamID::OffsetY);
-				break;
-		}
-	}
-	else if (InPropertyName == NAME_Pivot)
-	{
-		switch (InComponent)
-		{
-			case 0:
-				ParameterPtr = MaterialParameters.Find(ParamID::PivotX);
-				break;
-
-			case 1:
-				ParameterPtr = MaterialParameters.Find(ParamID::PivotY);
-				break;
-		}
-	}
-	else if (InPropertyName == NAME_Rotation)
-	{
-		switch (InComponent)
-		{
-			case 0:
-				ParameterPtr = MaterialParameters.Find(ParamID::Rotation);
-				break;
-		}
-	}
-	else if (InPropertyName == NAME_Scale)
-	{
-		switch (InComponent)
-		{
-			case 0:
-				ParameterPtr = MaterialParameters.Find(ParamID::ScaleX);
-				break;
-
-			case 1:
-				ParameterPtr = MaterialParameters.Find(ParamID::ScaleY);
-				break;
-		}
-	}
-
-	if (ParameterPtr)
-	{
-		return (*ParameterPtr).Get();
-	}
-
-	return nullptr;
-}
-
-FName UDMTextureUV::GetMaterialParameterName(FName InPropertyName, int32 InComponent) const
-{
-	if (UDMMaterialParameter* Parameter = GetMaterialParameter(InPropertyName, InComponent))
-	{
-		return Parameter->GetParameterName();
-	}
-
-	return NAME_None;
-}
-
-#if WITH_EDITOR
-bool UDMTextureUV::SetMaterialParameterName(FName InPropertyName, int32 InComponent, FName InNewName)
-{
-	if (UDMMaterialParameter* Parameter = GetMaterialParameter(InPropertyName, InComponent))
-	{
-		const FName CurrentName = Parameter->GetParameterName();
-		Parameter->RenameParameter(InNewName);
-		return Parameter->GetParameterName() != CurrentName;
-	}
-
-	return false;
-}
-
 void UDMTextureUV::PostEditorDuplicate(UDynamicMaterialModel* InMaterialModel, UDMMaterialComponent* InParent)
 {
 	Super::PostEditorDuplicate(InMaterialModel, InParent);
diff --git a/Engine/Plugins/Experimental/DynamicMaterial/Source/DynamicMaterial/Private/Model/DynamicMaterialModel.cpp b/Engine/Plugins/Experimental/DynamicMaterial/Source/DynamicMaterial/Private/Model/DynamicMaterialModel.cpp
--- a/Engine/Plugins/Experimental/DynamicMaterial/Source/DynamicMaterial/Private/Model/DynamicMaterialModel.cpp
+++ b/Engine/Plugins/Experimental/DynamicMaterial/Source/DynamicMaterial/Private/Model/DynamicMaterialModel.cpp
@@ -241,11 +241,6 @@ UDMMaterialParameter* UDynamicMaterialModel::CreateUniqueParameter(FName InBaseN
 	ParameterMap.Emplace(NewParameter->GetParameterName(), NewParameter);
 	NewParameter->SetComponentState(EDMComponentLifetimeState::Added);
 
-	if (IDynamicMaterialModelEditorOnlyDataInterface* ModelEditorOnlyData = GetEditorOnlyData())
-	{
-		ModelEditorOnlyData->RequestMaterialBuild();
-	}
-
 	return NewParameter;
 }
 
@@ -266,11 +261,6 @@ void UDynamicMaterialModel::RenameParameter(UDMMaterialParameter* InParameter, F
 
 	InParameter->ParameterName = CreateUniqueParameterName(InBaseName);
 	ParameterMap.Emplace(InParameter->ParameterName, InParameter);
-
-	if (IDynamicMaterialModelEditorOnlyDataInterface* ModelEditorOnlyData = GetEditorOnlyData())
-	{
-		ModelEditorOnlyData->RequestMaterialBuild();
-	}
 }
 
 void UDynamicMaterialModel::FreeParameter(UDMMaterialParameter* InParameter)
@@ -299,11 +289,6 @@ void UDynamicMaterialModel::FreeParameter(UDMMaterialParameter* InParameter)
 
 	InParameter->ParameterName = NAME_None;
 	InParameter->SetComponentState(EDMComponentLifetimeState::Removed);
-
-	if (IDynamicMaterialModelEditorOnlyDataInterface* ModelEditorOnlyData = GetEditorOnlyData())
-	{
-		ModelEditorOnlyData->RequestMaterialBuild();
-	}
 }
 
 bool UDynamicMaterialModel::ConditionalFreeParameter(UDMMaterialParameter* InParameter)
diff --git a/Engine/Plugins/Experimental/DynamicMaterial/Source/DynamicMaterial/Public/Components/DMMaterialValue.h b/Engine/Plugins/Experimental/DynamicMaterial/Source/DynamicMaterial/Public/Components/DMMaterialValue.h
--- a/Engine/Plugins/Experimental/DynamicMaterial/Source/DynamicMaterial/Public/Components/DMMaterialValue.h
+++ b/Engine/Plugins/Experimental/DynamicMaterial/Source/DynamicMaterial/Public/Components/DMMaterialValue.h
@@ -41,6 +41,7 @@ class DYNAMICMATERIAL_API UDMMaterialValue : public UDMMaterialLinkedComponent
  
 public:
 	static const FString ParameterPathToken;
+	static const TCHAR* ParameterNamePrefix;
 
 #if WITH_EDITOR
 	static const FName ValueName;
diff --git a/Engine/Plugins/Experimental/DynamicMaterial/Source/DynamicMaterial/Public/Components/DMTextureUV.h b/Engine/Plugins/Experimental/DynamicMaterial/Source/DynamicMaterial/Public/Components/DMTextureUV.h
--- a/Engine/Plugins/Experimental/DynamicMaterial/Source/DynamicMaterial/Public/Components/DMTextureUV.h
+++ b/Engine/Plugins/Experimental/DynamicMaterial/Source/DynamicMaterial/Public/Components/DMTextureUV.h
@@ -138,20 +138,8 @@ public:
 	TSharedPtr<IPropertyHandle> GetPropertyHandle(FName InProperty);
 #endif
 
-	UFUNCTION(BlueprintPure, Category = "Material Designer")
 	TArray<UDMMaterialParameter*> GetParameters() const;
 
-	UFUNCTION(BlueprintPure, Category = "Material Designer")
-	UDMMaterialParameter* GetMaterialParameter(FName InPropertyName, int32 InComponent) const;
-
-	UFUNCTION(BlueprintPure, Category = "Material Designer")
-	FName GetMaterialParameterName(FName InPropertyName, int32 InComponent) const;
-
-#if WITH_EDITOR
-	UFUNCTION(BlueprintCallable, Category = "Material Designer")
-	bool SetMaterialParameterName(FName InPropertyName, int32 InComponent, FName InNewName);
-#endif
-
 	void SetMIDParameters(UMaterialInstanceDynamic* InMID);
 
 	//~ Begin UObject
diff --git a/Engine/Plugins/Experimental/DynamicMaterial/Source/DynamicMaterialEditor/Private/DetailsPanel/DMValueDetailsRowExtensions.cpp b/Engine/Plugins/Experimental/DynamicMaterial/Source/DynamicMaterialEditor/Private/DetailsPanel/DMValueDetailsRowExtensions.cpp
deleted file mode 100644
diff --git a/Engine/Plugins/Experimental/DynamicMaterial/Source/DynamicMaterialEditor/Private/DetailsPanel/DMValueDetailsRowExtensions.h b/Engine/Plugins/Experimental/DynamicMaterial/Source/DynamicMaterialEditor/Private/DetailsPanel/DMValueDetailsRowExtensions.h
deleted file mode 100644
diff --git a/Engine/Plugins/Experimental/DynamicMaterial/Source/DynamicMaterialEditor/Private/DynamicMaterialEditorModule.cpp b/Engine/Plugins/Experimental/DynamicMaterial/Source/DynamicMaterialEditor/Private/DynamicMaterialEditorModule.cpp
--- a/Engine/Plugins/Experimental/DynamicMaterial/Source/DynamicMaterialEditor/Private/DynamicMaterialEditorModule.cpp
+++ b/Engine/Plugins/Experimental/DynamicMaterial/Source/DynamicMaterialEditor/Private/DynamicMaterialEditorModule.cpp
@@ -1,7 +1,6 @@
 // Copyright Epic Games, Inc. All Rights Reserved.
 
 #include "DynamicMaterialEditorModule.h"
-#include "Components/ActorComponent.h"
 #include "Components/DMMaterialComponent.h"
 #include "Components/DMMaterialEffectFunction.h"
 #include "Components/DMMaterialStageThroughput.h"
@@ -16,10 +15,10 @@
 #include "Components/MaterialValues/DMMaterialValueFloat3XYZ.h"
 #include "Components/MaterialValues/DMMaterialValueFloat4.h"
 #include "Components/MaterialValues/DMMaterialValueTexture.h"
+#include "Components/ActorComponent.h"
 #include "Components/PrimitiveComponent.h"
 #include "DetailsPanel/DMMaterialInterfaceTypeCustomizer.h"
 #include "DetailsPanel/DMPropertyTypeCustomizer.h"
-#include "DetailsPanel/DMValueDetailsRowExtensions.h"
 #include "DetailsPanel/Slate/SDMMaterialListExtensionWidget.h"
 #include "DMMaterialFunctionLibrary.h"
 #include "DMWorldSubsystem.h"
@@ -33,11 +32,11 @@
 #include "Material/DynamicMaterialInstance.h"
 #include "MaterialList.h"
 #include "Model/DMMaterialModelDefaults.h"
-#include "Model/DMOnWizardCompleteCallback.h"
 #include "Model/DynamicMaterialModel.h"
 #include "Model/DynamicMaterialModelEditorOnlyData.h"
 #include "Modules/ModuleManager.h"
 #include "PropertyEditorModule.h"
+#include "Model/DMOnWizardCompleteCallback.h"
 #include "Slate/Properties/Editors/SDMPropertyEditBoolValue.h"
 #include "Slate/Properties/Editors/SDMPropertyEditFloat1Value.h"
 #include "Slate/Properties/Editors/SDMPropertyEditFloat2Value.h"
@@ -305,8 +304,6 @@ void FDynamicMaterialEditorModule::StartupModule()
 		{
 			return UDynamicMaterialEditorSettings::Get()->DefaultRGBTexture.LoadSynchronous();
 		});
-
-	FDMValueDetailsRowExtensions::Get().RegisterRowExtensions();
 }
 
 void FDynamicMaterialEditorModule::ShutdownModule()
@@ -330,8 +327,6 @@ void FDynamicMaterialEditorModule::ShutdownModule()
 	BuildRequestList.Empty();
 
 	UDMMaterialValueTexture::GetDefaultRGBTexture.Unbind();
-
-	FDMValueDetailsRowExtensions::Get().UnregisterRowExtensions();
 }
 
 void FDynamicMaterialEditorModule::SetDynamicMaterialModel(UDynamicMaterialModel* InMaterialModel, UWorld* InWorld, bool bInInvokeTab)
