Material Designer
- Changed the global opacity slider so that it only appears when a material is actually translucent.

diff --git a/Engine/Plugins/Experimental/DynamicMaterial/Source/DynamicMaterialEditor/Private/Slate/SDMEditor.cpp b/Engine/Plugins/Experimental/DynamicMaterial/Source/DynamicMaterialEditor/Private/Slate/SDMEditor.cpp
--- a/Engine/Plugins/Experimental/DynamicMaterial/Source/DynamicMaterialEditor/Private/Slate/SDMEditor.cpp
+++ b/Engine/Plugins/Experimental/DynamicMaterial/Source/DynamicMaterialEditor/Private/Slate/SDMEditor.cpp
@@ -562,6 +562,7 @@ TSharedRef<SWidget> SDMEditor::CreateMainLayout()
 			SAssignNew(GlobalOpacityContainer, SBox)
 			.HAlign(HAlign_Fill)
 			.VAlign(VAlign_Fill)
+			.Visibility(this, &SDMEditor::GetGlobalOpacityVisibility)
 			[
 				CreateGlobalOpacityWidget()
 			]
@@ -613,12 +614,17 @@ TSharedRef<SWidget> SDMEditor::CreateGlobalOpacityWidget()
 	}
 
 	TSharedRef<SDMPropertyEdit> GlobalOpacityWidget = SNew(SDMPropertyEditOpacity, SharedThis(this), OpacityValue);
-	GlobalOpacityWidget->SetEnabled(TAttribute<bool>::CreateSP(this, &SDMEditor::IsGlobalOpacityEnabled));
 
 	TSharedRef<SWidget> GlobalOpacityButtons = SDMComponentEdit::CreateExtensionButtons(SharedThis(this), OpacityValue, UDMMaterialValue::ValueName, true, FSimpleDelegate());
-	GlobalOpacityButtons->SetEnabled(TAttribute<bool>::CreateSP(this, &SDMEditor::IsGlobalOpacityEnabled));
 
-	TSharedPtr<SWidget> RowLabel;
+	TSharedRef<SWidget> RowLabel = SNew(SBox)
+		.Padding(0.f, 0.f, 5.f, 0.f)
+		[
+			SNew(STextBlock)
+				.Text(LOCTEXT("GlobalOpacity", "Global Opacity"))
+		];
+
+	RowLabel->SetOnMouseButtonDown(FPointerEventHandler::CreateStatic(&SDMPropertyEdit::CreateRightClickDetailsMenu, GlobalOpacityWidget.ToWeakPtr()));
 
 	TSharedRef<SWidget> Row =
 		SNew(SHorizontalBox)
@@ -628,12 +634,7 @@ TSharedRef<SWidget> SDMEditor::CreateGlobalOpacityWidget()
 		.VAlign(VAlign_Top)
 		.Padding(10.0f, 5.0f)
 		[
-			SAssignNew(RowLabel, SBox)
-			.Padding(0.f, 0.f, 5.f, 0.f)
-			[
-				SNew(STextBlock)
-				.Text(LOCTEXT("GlobalOpacity", "Global Opacity"))
-			]
+			RowLabel
 		]
 		+ SHorizontalBox::Slot()
 		.FillWidth(1.0f)
@@ -650,8 +651,6 @@ TSharedRef<SWidget> SDMEditor::CreateGlobalOpacityWidget()
 			GlobalOpacityButtons
 		];
 
-	RowLabel->SetOnMouseButtonDown(FPointerEventHandler::CreateStatic(&SDMPropertyEdit::CreateRightClickDetailsMenu, GlobalOpacityWidget.ToWeakPtr()));
-
 	return Row;
 }
 
@@ -899,14 +898,17 @@ void SDMEditor::OnSlotsUpdated(UDynamicMaterialModel* InMaterialModel)
 	RefreshSlotsList();
 }
 
-bool SDMEditor::IsGlobalOpacityEnabled() const
+EVisibility SDMEditor::GetGlobalOpacityVisibility() const
 {
 	if (UDynamicMaterialModelEditorOnlyData* ModelEditorOnlyData = UDynamicMaterialModelEditorOnlyData::Get(MaterialModelWeak))
 	{
-		return ModelEditorOnlyData->GetBlendMode() == BLEND_Translucent || ModelEditorOnlyData->GetBlendMode() == BLEND_Masked;
+		if (ModelEditorOnlyData->GetBlendMode() != BLEND_Opaque)
+		{
+			return EVisibility::Visible;
+		}
 	}
 
-	return false;
+	return EVisibility::Collapsed;
 }
 
 bool SDMEditor::IsPropertyValidForModel(EDMMaterialPropertyType InProperty) const
diff --git a/Engine/Plugins/Experimental/DynamicMaterial/Source/DynamicMaterialEditor/Private/Slate/SDMEditor.h b/Engine/Plugins/Experimental/DynamicMaterial/Source/DynamicMaterialEditor/Private/Slate/SDMEditor.h
--- a/Engine/Plugins/Experimental/DynamicMaterial/Source/DynamicMaterialEditor/Private/Slate/SDMEditor.h
+++ b/Engine/Plugins/Experimental/DynamicMaterial/Source/DynamicMaterialEditor/Private/Slate/SDMEditor.h
@@ -146,7 +146,7 @@ protected:
 	void OnValuesUpdated(UDynamicMaterialModel* InMaterialModel);
 	void OnSlotsUpdated(UDynamicMaterialModel* InMaterialModel);
 
-	bool IsGlobalOpacityEnabled() const;
+	EVisibility GetGlobalOpacityVisibility() const;
 
 	bool IsPropertyValidForModel(EDMMaterialPropertyType InProperty) const;
 	ECheckBoxState GetSlotCheckState(EDMMaterialPropertyType InProperty) const;
