Material Designer
- Removed the text value type.
- Added the idea of 'render target renderers' for runtime re-rendering of a texture.
- Changed the text renderer to be a 'render target renderer'.
- Changed methods which use EDMValueType to take a TSubclassOf<UDMMaterialValue> as an overload.
- Fixed bug with clamp texture not working.

diff --git a/Engine/Plugins/Experimental/DynamicMaterial/Source/DynamicMaterial/Private/Components/DMMaterialValue.cpp b/Engine/Plugins/Experimental/DynamicMaterial/Source/DynamicMaterial/Private/Components/DMMaterialValue.cpp
--- a/Engine/Plugins/Experimental/DynamicMaterial/Source/DynamicMaterial/Private/Components/DMMaterialValue.cpp
+++ b/Engine/Plugins/Experimental/DynamicMaterial/Source/DynamicMaterial/Private/Components/DMMaterialValue.cpp
@@ -30,20 +30,24 @@ const FString UDMMaterialValue::ParameterPathToken = FString(TEXT("Parameter"));
 #if WITH_EDITOR
 const FName UDMMaterialValue::ValueName = "Value";
 
+UDMMaterialValue* UDMMaterialValue::CreateMaterialValue(UDynamicMaterialModel* InMaterialModel, const FString& InName, 
+	EDMValueType InValueType, bool bInLocal)
+{
+	return CreateMaterialValue(
+		InMaterialModel,
+		InName,
+		UDMValueDefinitionLibrary::GetValueDefinition(InValueType).GetValueClass(),
+		bInLocal
+	);
+}
+
 UDMMaterialValue* UDMMaterialValue::CreateMaterialValue(UDynamicMaterialModel* InMaterialModel, const FString& InName,
-	EDMValueType InType, bool bInLocal)
+	TSubclassOf<UDMMaterialValue> InValueClass, bool bInLocal)
 {
 	check(InMaterialModel);
  
-	TStrongObjectPtr<UClass>* ValueClassPtr = TypeClasses.Find(InType);
-	check(ValueClassPtr);
- 
-	UClass* ValueClass = (*ValueClassPtr).Get();
-	check(ValueClass);
- 
-	UDMMaterialValue* NewValue = NewObject<UDMMaterialValue>(InMaterialModel, ValueClass, NAME_None, RF_Transactional);
+	UDMMaterialValue* NewValue = NewObject<UDMMaterialValue>(InMaterialModel, InValueClass, NAME_None, RF_Transactional);
 	check(NewValue);
-	NewValue->Type = InType;
 	NewValue->bLocal = bInLocal;
  
 	if (InName.IsEmpty() == false)
diff --git a/Engine/Plugins/Experimental/DynamicMaterial/Source/DynamicMaterial/Private/Components/DMRenderTargetRenderer.cpp b/Engine/Plugins/Experimental/DynamicMaterial/Source/DynamicMaterial/Private/Components/DMRenderTargetRenderer.cpp
new file mode 100644
--- /dev/null
+++ b/Engine/Plugins/Experimental/DynamicMaterial/Source/DynamicMaterial/Private/Components/DMRenderTargetRenderer.cpp
@@ -0,0 +1,67 @@
+// Copyright Epic Games, Inc. All Rights Reserved.
+
+#include "Components/DMRenderTargetRenderer.h"
+
+#include "Components/MaterialValues/DMMaterialValueRenderTarget.h"
+#include "Misc/CoreDelegates.h"
+#include "Templates/SubclassOf.h"
+
+UDMRenderTargetRenderer* UDMRenderTargetRenderer::CreateRenderTargetRenderer(TSubclassOf<UDMRenderTargetRenderer> InRendererClass, 
+	UDMMaterialValueRenderTarget* InRenderTargetValue)
+{
+	UClass* Class = InRendererClass.Get();
+
+	check(Class);
+	check(!Class->HasAnyClassFlags(UE::DynamicMaterial::InvalidClassFlags));
+
+	check(InRenderTargetValue);
+
+	UDMRenderTargetRenderer* Renderer = NewObject<UDMRenderTargetRenderer>(InRenderTargetValue, Class, NAME_None, RF_Transactional);
+	InRenderTargetValue->SetRenderer(Renderer);
+
+	return Renderer;
+}
+
+UDMMaterialValueRenderTarget* UDMRenderTargetRenderer::GetRenderTargetValue() const
+{
+	return Cast<UDMMaterialValueRenderTarget>(GetOuterSafe());
+}
+
+void UDMRenderTargetRenderer::UpdateRenderTarget()
+{
+	if (EndOfFrameDelegateHandle.IsValid())
+	{
+		FCoreDelegates::OnEndFrame.Remove(EndOfFrameDelegateHandle);
+		EndOfFrameDelegateHandle.Reset();
+	}
+
+	if (bUpdating)
+	{
+		return;
+	}
+
+	TGuardValue<bool> Guard(bUpdating, true);
+
+	UpdateRenderTarget_Internal();
+}
+
+void UDMRenderTargetRenderer::AsyncUpdateRenderTarget()
+{
+	if (bUpdating)
+	{
+		return;
+	}
+
+	if (!EndOfFrameDelegateHandle.IsValid())
+	{
+		EndOfFrameDelegateHandle = FCoreDelegates::OnEndFrame.AddUObject(this, &UDMRenderTargetRenderer::UpdateRenderTarget);
+	}
+}
+
+void UDMRenderTargetRenderer::FlushUpdateRenderTarget()
+{
+	if (EndOfFrameDelegateHandle.IsValid())
+	{
+		UpdateRenderTarget();
+	}
+}
diff --git a/Engine/Plugins/Experimental/DynamicMaterial/Source/DynamicMaterial/Private/Components/MaterialValues/DMMaterialValueRenderTarget.cpp b/Engine/Plugins/Experimental/DynamicMaterial/Source/DynamicMaterial/Private/Components/MaterialValues/DMMaterialValueRenderTarget.cpp
new file mode 100644
--- /dev/null
+++ b/Engine/Plugins/Experimental/DynamicMaterial/Source/DynamicMaterial/Private/Components/MaterialValues/DMMaterialValueRenderTarget.cpp
@@ -0,0 +1,271 @@
+// Copyright Epic Games, Inc. All Rights Reserved.
+
+#include "Components/MaterialValues/DMMaterialValueRenderTarget.h"
+#include "Components/DMRenderTargetRenderer.h"
+#include "DMComponentPath.h"
+#include "Engine/TextureRenderTarget2D.h"
+#include "Misc/CoreDelegates.h"
+
+const FString UDMMaterialValueRenderTarget::RendererPathToken = "Renderer";
+
+UDMMaterialValueRenderTarget::UDMMaterialValueRenderTarget()
+	: TextureSize(FIntPoint(512, 512))
+	, TextureFormat(ETextureRenderTargetFormat::RTF_RGBA16f)
+	, ClearColor(FLinearColor::Black)
+	, Renderer(nullptr)
+{
+#if WITH_EDITOR
+	EditableProperties.Add(GET_MEMBER_NAME_CHECKED(UDMMaterialValueRenderTarget, TextureSize));
+	EditableProperties.Add(GET_MEMBER_NAME_CHECKED(UDMMaterialValueRenderTarget, TextureFormat));
+	EditableProperties.Add(GET_MEMBER_NAME_CHECKED(UDMMaterialValueRenderTarget, ClearColor));
+	EditableProperties.Add(GET_MEMBER_NAME_CHECKED(UDMMaterialValueRenderTarget, Renderer));
+#endif
+}
+
+UDMMaterialValueRenderTarget::~UDMMaterialValueRenderTarget()
+{
+	FCoreDelegates::OnEndFrame.Remove(EndOfFrameDelegateHandle);
+	EndOfFrameDelegateHandle.Reset();
+}
+
+UTextureRenderTarget2D* UDMMaterialValueRenderTarget::GetRenderTarget() const
+{
+	return Cast<UTextureRenderTarget2D>(GetValue());
+}
+
+const FIntPoint& UDMMaterialValueRenderTarget::GetTextureSize() const
+{
+	return TextureSize;
+}
+
+void UDMMaterialValueRenderTarget::SetTextureSize(const FIntPoint& InTextureSize)
+{
+	if (InTextureSize.X <= 0 || InTextureSize.Y <= 0 || InTextureSize == TextureSize)
+	{
+		return;
+	}
+
+	TextureSize = InTextureSize;
+
+	AsyncCreateRenderTarget();
+}
+
+ETextureRenderTargetFormat UDMMaterialValueRenderTarget::GetTextureFormat() const
+{
+	return TextureFormat;
+}
+
+void UDMMaterialValueRenderTarget::SetTextureFormat(ETextureRenderTargetFormat InTextureFormat)
+{
+	if (InTextureFormat == TextureFormat)
+	{
+		return;
+	}
+
+	TextureFormat = InTextureFormat;
+
+	AsyncCreateRenderTarget();
+}
+
+const FLinearColor& UDMMaterialValueRenderTarget::GetClearColor() const
+{
+	return ClearColor;
+}
+
+void UDMMaterialValueRenderTarget::SetClearColor(const FLinearColor& InClearColor)
+{
+	if (InClearColor == ClearColor)
+	{
+		return;
+	}
+
+	ClearColor = InClearColor;
+
+	AsyncCreateRenderTarget();
+}
+
+UDMRenderTargetRenderer* UDMMaterialValueRenderTarget::GetRenderer() const
+{
+	return Renderer;
+}
+
+void UDMMaterialValueRenderTarget::SetRenderer(UDMRenderTargetRenderer* InRenderer)
+{
+#if WITH_EDITOR
+	if (Renderer == InRenderer)
+	{
+		return;
+	}
+
+	if (Renderer)
+	{
+		Renderer->SetComponentState(EDMComponentLifetimeState::Removed);
+	}
+
+	Renderer = InRenderer;
+
+	if (IsComponentAdded())
+	{
+		Renderer->SetComponentState(EDMComponentLifetimeState::Added);
+	}
+#endif
+}
+
+void UDMMaterialValueRenderTarget::EnsureRenderTarget(bool bInAsync)
+{
+	if (IsValid(GetRenderTarget()))
+	{
+		return;
+	}
+
+	if (bInAsync)
+	{
+		AsyncCreateRenderTarget();
+	}
+	else
+	{
+		CreateRenderTarget();
+	}
+}
+
+void UDMMaterialValueRenderTarget::FlushCreateRenderTarget()
+{
+	if (EndOfFrameDelegateHandle.IsValid() || !IsValid(GetRenderTarget()))
+	{
+		CreateRenderTarget();
+	}
+}
+
+void UDMMaterialValueRenderTarget::Update(EDMUpdateType InUpdateType)
+{
+	Super::Update(InUpdateType);
+
+	if (!IsValid(GetRenderTarget()))
+	{
+		AsyncCreateRenderTarget();
+	}
+	else
+	{
+		UpdateRenderTarget();
+	}
+}
+
+#if WITH_EDITOR
+void UDMMaterialValueRenderTarget::PostEditorDuplicate(UDynamicMaterialModel* InMaterialModel, UDMMaterialComponent* InParent)
+{
+	Super::PostEditorDuplicate(InMaterialModel, InParent);
+
+	if (IsValid(Renderer))
+	{
+		Renderer->PostEditorDuplicate(InMaterialModel, this);
+	}
+
+	// Make sure we have a unique render target
+	AsyncCreateRenderTarget();
+}
+
+void UDMMaterialValueRenderTarget::PostEditChangeProperty(FPropertyChangedEvent& InPropertyChangedEvent)
+{
+	Super::PostEditChangeProperty(InPropertyChangedEvent);
+
+	static const FName TextureSizeName = GET_MEMBER_NAME_CHECKED(UDMMaterialValueRenderTarget, TextureSize);
+	static const FName ClearColorName = GET_MEMBER_NAME_CHECKED(UDMMaterialValueRenderTarget, ClearColor);
+
+	const FName MemberName = InPropertyChangedEvent.GetMemberPropertyName();
+
+	if (MemberName == TextureSizeName || MemberName == ClearColorName)
+	{
+		AsyncCreateRenderTarget();
+	}
+}
+#endif
+
+void UDMMaterialValueRenderTarget::PostLoad()
+{
+	Super::PostLoad();
+
+	if (!IsValid(GetRenderTarget()))
+	{
+		CreateRenderTarget();
+	}
+}
+
+void UDMMaterialValueRenderTarget::AsyncCreateRenderTarget()
+{
+	if (!EndOfFrameDelegateHandle.IsValid())
+	{
+		EndOfFrameDelegateHandle = FCoreDelegates::OnEndFrame.AddUObject(this, &UDMMaterialValueRenderTarget::CreateRenderTarget);
+	}
+}
+
+void UDMMaterialValueRenderTarget::CreateRenderTarget()
+{
+	if (EndOfFrameDelegateHandle.IsValid())
+	{
+		FCoreDelegates::OnEndFrame.Remove(EndOfFrameDelegateHandle);
+		EndOfFrameDelegateHandle.Reset();
+	}
+
+	UTextureRenderTarget2D* RenderTarget = NewObject<UTextureRenderTarget2D>(this, NAME_None,
+		EObjectFlags::RF_Transactional | EObjectFlags::RF_DuplicateTransient | EObjectFlags::RF_TextExportTransient
+	);
+
+	check(RenderTarget);
+	RenderTarget->RenderTargetFormat = TextureFormat;
+	RenderTarget->ClearColor = ClearColor;
+	RenderTarget->bAutoGenerateMips = false;
+	RenderTarget->bCanCreateUAV = false;
+	RenderTarget->InitAutoFormat(TextureSize.X, TextureSize.Y);
+	RenderTarget->UpdateResourceImmediate(true);
+
+	SetValue(RenderTarget);
+}
+
+void UDMMaterialValueRenderTarget::UpdateRenderTarget()
+{
+	if (Renderer)
+	{
+		if (UTextureRenderTarget2D* RenderTarget = GetRenderTarget())
+		{
+			Renderer->UpdateRenderTarget();
+		}
+		else
+		{
+			AsyncCreateRenderTarget();
+		}
+	}
+}
+
+UDMMaterialComponent* UDMMaterialValueRenderTarget::GetSubComponentByPath(FDMComponentPath& InPath, const FDMComponentPathSegment& InPathSegment) const
+{
+	if (InPathSegment.GetToken() == RendererPathToken)
+	{
+		return Renderer;
+	}
+
+	return Super::GetSubComponentByPath(InPath, InPathSegment);
+}
+
+#if WITH_EDITOR
+void UDMMaterialValueRenderTarget::OnComponentAdded()
+{
+	Super::OnComponentAdded();
+
+	if (Renderer)
+	{
+		Renderer->SetComponentState(EDMComponentLifetimeState::Added);
+	}
+
+	AsyncCreateRenderTarget();
+}
+
+void UDMMaterialValueRenderTarget::OnComponentRemoved()
+{
+	Super::OnComponentRemoved();
+
+	if (Renderer)
+	{
+		Renderer->SetComponentState(EDMComponentLifetimeState::Removed);
+	}
+}
+#endif
diff --git a/Engine/Plugins/Experimental/DynamicMaterial/Source/DynamicMaterial/Private/Components/MaterialValues/DMMaterialValueText.cpp b/Engine/Plugins/Experimental/DynamicMaterial/Source/DynamicMaterial/Private/Components/MaterialValues/DMMaterialValueText.cpp
deleted file mode 100644
diff --git a/Engine/Plugins/Experimental/DynamicMaterial/Source/DynamicMaterial/Private/Components/RenderTargetRenderers/DMRenderTargetTextRenderer.cpp b/Engine/Plugins/Experimental/DynamicMaterial/Source/DynamicMaterial/Private/Components/RenderTargetRenderers/DMRenderTargetTextRenderer.cpp
new file mode 100644
--- /dev/null
+++ b/Engine/Plugins/Experimental/DynamicMaterial/Source/DynamicMaterial/Private/Components/RenderTargetRenderers/DMRenderTargetTextRenderer.cpp
@@ -0,0 +1,731 @@
+// Copyright Epic Games, Inc. All Rights Reserved.
+
+#include "Components/RenderTargetRenderers/DMRenderTargetTextRenderer.h"
+#include "CanvasItem.h"
+#include "CanvasTypes.h"
+#include "Components/MaterialValues/DMMaterialValueRenderTarget.h"
+#include "Engine/Canvas.h"
+#include "Engine/Engine.h"
+#include "Engine/TextureRenderTarget2D.h"
+#include "TextureResource.h"
+
+#define LOCTEXT_NAMESPACE "DMRenderTargetTextRenderer"
+
+namespace UE::DynamicMaterial::Private
+{
+	const FIntPoint MinimumTextTextureSize = FIntPoint(5, 10);
+	const FIntPoint MaxTexScale = FIntPoint(64, 64);
+}
+
+#if WITH_EDITOR
+struct FDMRenderTargetTextRenderer
+{
+	static const inline FName FontName = GET_MEMBER_NAME_CHECKED(UDMRenderTargetTextRenderer, Font);
+	static const inline FName TextName = GET_MEMBER_NAME_CHECKED(UDMRenderTargetTextRenderer, Text);
+	static const inline FName TextColorName = GET_MEMBER_NAME_CHECKED(UDMRenderTargetTextRenderer, TextColor);
+	static const inline FName JustifyName = GET_MEMBER_NAME_CHECKED(UDMRenderTargetTextRenderer, Justify);
+	static const inline FName KerningName = GET_MEMBER_NAME_CHECKED(UDMRenderTargetTextRenderer, Kerning);
+	static const inline FName LineHeightName = GET_MEMBER_NAME_CHECKED(UDMRenderTargetTextRenderer, LineHeight);
+	static const inline FName PaddingLeftName = GET_MEMBER_NAME_CHECKED(UDMRenderTargetTextRenderer, PaddingLeft);
+	static const inline FName PaddingRightName = GET_MEMBER_NAME_CHECKED(UDMRenderTargetTextRenderer, PaddingRight);
+	static const inline FName PaddingTopName = GET_MEMBER_NAME_CHECKED(UDMRenderTargetTextRenderer, PaddingTop);
+	static const inline FName PaddingBottomName = GET_MEMBER_NAME_CHECKED(UDMRenderTargetTextRenderer, PaddingBottom);
+	static const inline FName TextScaleName = GET_MEMBER_NAME_CHECKED(UDMRenderTargetTextRenderer, TextScale);
+	static const inline FName TextureSizeOverrideName = GET_MEMBER_NAME_CHECKED(UDMRenderTargetTextRenderer, TextureSizeOverride);
+	static const inline FName bOutlineName = GET_MEMBER_NAME_CHECKED(UDMRenderTargetTextRenderer, bOutline);
+	static const inline FName OutlineColorName = GET_MEMBER_NAME_CHECKED(UDMRenderTargetTextRenderer, OutlineColor);
+	static const inline FName bShadowName = GET_MEMBER_NAME_CHECKED(UDMRenderTargetTextRenderer, bShadow);
+	static const inline FName ShadowColorName = GET_MEMBER_NAME_CHECKED(UDMRenderTargetTextRenderer, ShadowColor);
+	static const inline FName ShadowOffsetName = GET_MEMBER_NAME_CHECKED(UDMRenderTargetTextRenderer, ShadowOffset);
+	static const inline FName bGlowName = GET_MEMBER_NAME_CHECKED(UDMRenderTargetTextRenderer, bGlow);
+	static const inline FName GlowColorName = GET_MEMBER_NAME_CHECKED(UDMRenderTargetTextRenderer, GlowColor);
+	static const inline FName GlowInnerRadiusName = GET_MEMBER_NAME_CHECKED(UDMRenderTargetTextRenderer, GlowInnerRadius);
+	static const inline FName GlowOuterRadiusName = GET_MEMBER_NAME_CHECKED(UDMRenderTargetTextRenderer, GlowOuterRadius);
+
+	static const inline TSet<FName> PropertyNames = {
+		FontName, TextName, TextColorName, JustifyName, KerningName, LineHeightName, PaddingLeftName, PaddingRightName,
+		PaddingTopName, PaddingBottomName, TextScaleName, TextureSizeOverrideName, bOutlineName, OutlineColorName, bShadowName,
+		ShadowColorName, ShadowOffsetName, bGlowName, GlowColorName, GlowInnerRadiusName, GlowOuterRadiusName
+	};
+};
+#endif
+
+UDMRenderTargetTextRenderer::UDMRenderTargetTextRenderer()
+{
+#if WITH_EDITOR
+	EditableProperties.Append(FDMRenderTargetTextRenderer::PropertyNames.Array());
+#endif
+}
+
+void UDMRenderTargetTextRenderer::PostLoad()
+{
+	Super::PostLoad();
+
+	if (UDMMaterialValueRenderTarget* RenderTargetValue = GetRenderTargetValue())
+	{
+		RenderTargetValue->EnsureRenderTarget();
+	}
+}
+
+#if WITH_EDITOR
+void UDMRenderTargetTextRenderer::PostEditChangeProperty(FPropertyChangedEvent& InPropertyChangedEvent)
+{
+	Super::PostEditChangeProperty(InPropertyChangedEvent);
+
+	const FName PropertyName = InPropertyChangedEvent.GetMemberPropertyName();
+
+	if (PropertyName == FDMRenderTargetTextRenderer::TextName)
+	{
+		UpdateTextLines();
+	}
+	else if (PropertyName == FDMRenderTargetTextRenderer::FontName
+		|| PropertyName == FDMRenderTargetTextRenderer::KerningName
+		|| PropertyName == FDMRenderTargetTextRenderer::LineHeightName
+		|| PropertyName == FDMRenderTargetTextRenderer::PaddingLeftName
+		|| PropertyName == FDMRenderTargetTextRenderer::PaddingRightName
+		|| PropertyName == FDMRenderTargetTextRenderer::PaddingTopName
+		|| PropertyName == FDMRenderTargetTextRenderer::PaddingBottomName
+		|| PropertyName == FDMRenderTargetTextRenderer::TextScaleName
+		|| PropertyName == FDMRenderTargetTextRenderer::TextureSizeOverrideName)
+	{
+		TextScale.X = FMath::Clamp(TextScale.X, 1, UE::DynamicMaterial::Private::MaxTexScale.X);
+		TextScale.Y = FMath::Clamp(TextScale.Y, 1, UE::DynamicMaterial::Private::MaxTexScale.Y);
+		bRecalculateTextSize = true;
+		AsyncUpdateRenderTarget();
+	}
+	else if (PropertyName == FDMRenderTargetTextRenderer::bGlowName
+		|| PropertyName == FDMRenderTargetTextRenderer::bShadowName
+		|| PropertyName == FDMRenderTargetTextRenderer::bOutlineName)
+	{
+		// Cause details panel refresh.
+		Update(EDMUpdateType::Structure);
+		AsyncUpdateRenderTarget();
+	}
+	else if (FDMRenderTargetTextRenderer::PropertyNames.Contains(PropertyName))
+	{
+		AsyncUpdateRenderTarget();
+	}
+}
+#endif
+
+UFont* UDMRenderTargetTextRenderer::GetFont() const
+{
+	return Font;
+}
+
+void UDMRenderTargetTextRenderer::SetFont(UFont* InFont)
+{
+	if (Font == InFont)
+	{
+		return;
+	}
+
+	Font = InFont;
+
+	UpdateTextureSize();
+}
+
+const FText& UDMRenderTargetTextRenderer::GetText() const
+{
+	return Text;
+}
+
+void UDMRenderTargetTextRenderer::SetText(const FText& InText)
+{
+	if (Text.EqualTo(InText, ETextComparisonLevel::Default))
+	{
+		return;
+	}
+
+	Text = InText;
+
+	UpdateTextLines();
+	UpdateTextureSize();
+}
+
+const FLinearColor& UDMRenderTargetTextRenderer::GetTextColor() const
+{
+	return TextColor;
+}
+
+void UDMRenderTargetTextRenderer::SetTextColor(const FLinearColor& InColor)
+{
+	if (TextColor == InColor)
+	{
+		return;
+	}
+
+	TextColor = InColor;
+
+	AsyncUpdateRenderTarget();
+}
+
+const FLinearColor& UDMRenderTargetTextRenderer::GetBackgroundColor() const
+{
+	if (UDMMaterialValueRenderTarget* RenderTargetValue = GetRenderTargetValue())
+	{
+		return RenderTargetValue->GetClearColor();
+	}
+
+	return FLinearColor::Black;
+}
+
+void UDMRenderTargetTextRenderer::SetBackgroundColor(const FLinearColor& InBackgroundColor)
+{
+	if (UDMMaterialValueRenderTarget* RenderTargetValue = GetRenderTargetValue())
+	{
+		RenderTargetValue->SetClearColor(InBackgroundColor);
+	}
+}
+
+ETextJustify::Type UDMRenderTargetTextRenderer::GetJustify() const
+{
+	return Justify;
+}
+
+void UDMRenderTargetTextRenderer::SetJustify(ETextJustify::Type InJustify)
+{
+	if (Justify == InJustify)
+	{
+		return;
+	}
+
+	Justify = InJustify;
+
+	AsyncUpdateRenderTarget();
+}
+
+float UDMRenderTargetTextRenderer::GetKerning() const
+{
+	return Kerning;
+}
+
+void UDMRenderTargetTextRenderer::SetKerning(float InKerning)
+{
+	if (Kerning == InKerning)
+	{
+		return;
+	}
+
+	Kerning = InKerning;
+
+	UpdateTextureSize();
+}
+
+const FLinearColor& UDMRenderTargetTextRenderer::GetOutlineColor() const
+{
+	return OutlineColor;
+}
+
+void UDMRenderTargetTextRenderer::SetOutlineColor(const FLinearColor& InOutlineColor)
+{
+	if (OutlineColor == InOutlineColor)
+	{
+		return;
+	}
+
+	OutlineColor = InOutlineColor;
+
+	AsyncUpdateRenderTarget();
+}
+
+float UDMRenderTargetTextRenderer::GetLineHeight() const
+{
+	return LineHeight;
+}
+
+void UDMRenderTargetTextRenderer::SetLineHeight(float InLineHeight)
+{
+	if (LineHeight == InLineHeight)
+	{
+		return;
+	}
+
+	LineHeight = InLineHeight;
+
+	UpdateTextureSize();
+}
+
+float UDMRenderTargetTextRenderer::GetPaddingLeft() const
+{
+	return PaddingLeft;
+}
+
+void UDMRenderTargetTextRenderer::SetPaddingLeft(float InPaddingLeft)
+{
+	if (PaddingLeft == InPaddingLeft)
+	{
+		return;
+	}
+
+	PaddingLeft = InPaddingLeft;
+
+	UpdateTextureSize();
+}
+
+float UDMRenderTargetTextRenderer::GetPaddingRight() const
+{
+	return PaddingRight;
+}
+
+void UDMRenderTargetTextRenderer::SetPaddingRight(float InPaddingRight)
+{
+	if (PaddingRight == InPaddingRight)
+	{
+		return;
+	}
+
+	PaddingRight = InPaddingRight;
+
+	UpdateTextureSize();
+}
+
+float UDMRenderTargetTextRenderer::GetPaddingTop() const
+{
+	return PaddingTop;
+}
+
+void UDMRenderTargetTextRenderer::SetPaddingTop(float InPaddingTop)
+{
+	if (PaddingTop == InPaddingTop)
+	{
+		return;
+	}
+
+	PaddingTop = InPaddingTop;
+
+	UpdateTextureSize();
+}
+
+float UDMRenderTargetTextRenderer::GetPaddingBottom() const
+{
+	return PaddingBottom;
+}
+
+void UDMRenderTargetTextRenderer::SetPaddingBottom(float InPaddingBottom)
+{
+	if (PaddingBottom == InPaddingBottom)
+	{
+		return;
+	}
+
+	PaddingBottom = InPaddingBottom;
+
+	UpdateTextureSize();
+}
+
+const FIntPoint& UDMRenderTargetTextRenderer::GetTextScale() const
+{
+	return TextScale;
+}
+
+void UDMRenderTargetTextRenderer::SetTextScale(const FIntPoint& InTextScale)
+{
+	FIntPoint NewTextScale = InTextScale;
+	NewTextScale.X = FMath::Clamp(InTextScale.X, 1, UE::DynamicMaterial::Private::MaxTexScale.X);
+	NewTextScale.Y = FMath::Clamp(InTextScale.Y, 1, UE::DynamicMaterial::Private::MaxTexScale.Y);
+
+	if (TextScale == NewTextScale)
+	{
+		return;
+	}
+
+	TextScale = NewTextScale;
+
+	UpdateTextureSize();
+}
+
+const FIntPoint& UDMRenderTargetTextRenderer::GetTextureSizeOverride() const
+{
+	return TextureSizeOverride;
+}
+
+void UDMRenderTargetTextRenderer::SetTextureSizeOverride(const FIntPoint& InTextureSizeOverride)
+{
+	if (TextureSizeOverride == InTextureSizeOverride)
+	{
+		return;
+	}
+
+	TextureSizeOverride = InTextureSizeOverride;
+
+	UpdateTextureSize();
+}
+
+bool UDMRenderTargetTextRenderer::GetHasOutline() const
+{
+	return bOutline;
+}
+
+void UDMRenderTargetTextRenderer::SetHasOutline(bool bInHasOutline)
+{
+	if (bOutline == bInHasOutline)
+	{
+		return;
+	}
+
+	bOutline = bInHasOutline;
+
+	AsyncUpdateRenderTarget();
+}
+
+bool UDMRenderTargetTextRenderer::GetHasShadow() const
+{
+	return bShadow;
+}
+
+void UDMRenderTargetTextRenderer::SetHasShadow(bool bInHasShadow)
+{
+	if (bShadow == bInHasShadow)
+	{
+		return;
+	}
+
+	bShadow = bInHasShadow;
+
+	AsyncUpdateRenderTarget();
+}
+
+const FLinearColor& UDMRenderTargetTextRenderer::GetShadowColor() const
+{
+	return ShadowColor;
+}
+
+void UDMRenderTargetTextRenderer::SetShadowColor(const FLinearColor& InShadowColor)
+{
+	if (ShadowColor == InShadowColor)
+	{
+		return;
+	}
+
+	ShadowColor = InShadowColor;
+
+	AsyncUpdateRenderTarget();
+}
+
+const FVector2D& UDMRenderTargetTextRenderer::GetShadowOffset() const
+{
+	return ShadowOffset;
+}
+
+void UDMRenderTargetTextRenderer::SetShadowOffset(const FVector2D& InShadowOffset)
+{
+	if (ShadowOffset == InShadowOffset)
+	{
+		return;
+	}
+
+	ShadowOffset = InShadowOffset;
+
+	AsyncUpdateRenderTarget();
+}
+
+bool UDMRenderTargetTextRenderer::GetHasGlow() const
+{
+	return bGlow;
+}
+
+void UDMRenderTargetTextRenderer::SetHasGlow(bool bInHasGlow)
+{
+	if (bGlow == bInHasGlow)
+	{
+		return;
+	}
+
+	bGlow = bInHasGlow;
+
+	AsyncUpdateRenderTarget();
+}
+
+const FLinearColor& UDMRenderTargetTextRenderer::GetGlowColor() const
+{
+	return GlowColor;
+}
+
+void UDMRenderTargetTextRenderer::SetGlowColor(const FLinearColor& InGlowColor)
+{
+	if (GlowColor == InGlowColor)
+	{
+		return;
+	}
+
+	GlowColor = InGlowColor;
+
+	AsyncUpdateRenderTarget();
+}
+
+const FVector2D& UDMRenderTargetTextRenderer::GetGlowInnerRadius() const
+{
+	return GlowInnerRadius;
+}
+
+void UDMRenderTargetTextRenderer::SetGlowInnerRadius(const FVector2D& InGlowInnerRadius)
+{
+	if (GlowInnerRadius == InGlowInnerRadius)
+	{
+		return;
+	}
+
+	GlowInnerRadius = InGlowInnerRadius;
+
+	AsyncUpdateRenderTarget();
+}
+
+const FVector2D& UDMRenderTargetTextRenderer::GetGlowOuterRadius() const
+{
+	return GlowOuterRadius;
+}
+
+void UDMRenderTargetTextRenderer::SetGlowOuterRadius(const FVector2D& InGlowOuterRadius)
+{
+	if (GlowOuterRadius == InGlowOuterRadius)
+	{
+		return;
+	}
+
+	GlowOuterRadius = InGlowOuterRadius;
+
+	AsyncUpdateRenderTarget();
+}
+
+void UDMRenderTargetTextRenderer::CalculateLineSizes()
+{
+	FTextSizingParameters Params(0, 0, 0, 0, Font);
+	Params.Scaling = FVector2D(1, 1);
+
+	for (FDMTextLine& Line : Lines)
+	{
+		Params.DrawXL = 0;
+		Params.DrawYL = 0;
+
+		UCanvas::CanvasStringSize(Params, *Line.Line);
+
+		int32 RequiredWidth = Params.DrawXL;
+		RequiredWidth += Line.Line.Len() > 0 ? ((Line.Line.Len() - 1) * Kerning) : 0;
+
+		Line.Width = RequiredWidth;
+	}
+}
+
+FIntPoint UDMRenderTargetTextRenderer::GetRequiredTextureSize() const
+{
+	if (TextureSizeOverride.X > 0 && TextureSizeOverride.Y > 0)
+	{
+		return TextureSizeOverride;
+	}
+
+	if (!IsValid(Font) || Lines.IsEmpty())
+	{
+		return UE::DynamicMaterial::Private::MinimumTextTextureSize;
+	}
+
+	const float MaxCharHeight = Font->GetMaxCharHeight();
+
+	if (MaxCharHeight <= 0)
+	{
+		return UE::DynamicMaterial::Private::MinimumTextTextureSize;
+	}
+
+	FVector2f Size = FVector2f::ZeroVector;
+
+	for (const FDMTextLine& Line : Lines)
+	{
+		Size.X = FMath::Max(Size.X, Line.Width);
+		Size.Y += MaxCharHeight * LineHeight;
+	}
+
+	Size.X += PaddingLeft + PaddingRight;
+	Size.Y += PaddingTop + PaddingBottom;
+
+	Size.X = FMath::Max(Size.X, UE::DynamicMaterial::Private::MinimumTextTextureSize.X);
+	Size.Y = FMath::Max(Size.Y, UE::DynamicMaterial::Private::MinimumTextTextureSize.Y);
+
+	FIntPoint RequiredSize = FIntPoint::ZeroValue;
+	RequiredSize.X = FMath::CeilToInt(Size.X * TextScale.X);
+	RequiredSize.Y = FMath::CeilToInt(Size.Y * TextScale.Y);
+
+	return RequiredSize;
+}
+
+void UDMRenderTargetTextRenderer::UpdateTextLines()
+{
+	const FString TextStr = Text.ToString();
+	TArray<FString> NewLines;
+	TextStr.ParseIntoArray(NewLines, TEXT("\n"), /* Cull empty */ false);
+
+	Lines.Empty();
+	Lines.Reserve(NewLines.Num());
+
+	for (FString& NewLine : NewLines)
+	{
+		if (NewLine.EndsWith(TEXT("\r")))
+		{
+			NewLine.LeftChopInline(1);
+		}
+
+		Lines.Add({NewLine, -1.f});
+	}
+
+	bRecalculateTextSize = true;
+	AsyncUpdateRenderTarget();
+}
+
+FCanvasTextItem UDMRenderTargetTextRenderer::CreateTextItem(const FVector2D& InPosition, const FText& InText) const
+{
+	FCanvasTextItem TextItem(InPosition, InText, Font, TextColor);
+
+	TextItem.BlendMode = SE_BLEND_Translucent; // Must be translucent for fonts!
+	TextItem.HorizSpacingAdjust = Kerning;
+
+	if (bOutline)
+	{
+		TextItem.bOutlined = true;
+		TextItem.OutlineColor = OutlineColor;
+	}
+	else
+	{
+		TextItem.bOutlined = false;
+	}
+
+	if (bShadow)
+	{
+		TextItem.EnableShadow(ShadowColor, ShadowOffset);
+	}
+	else
+	{
+		TextItem.DisableShadow();
+	}
+
+	if (bGlow)
+	{
+		TextItem.FontRenderInfo.GlowInfo.bEnableGlow = true;
+		TextItem.FontRenderInfo.GlowInfo.GlowColor = GlowColor;
+		TextItem.FontRenderInfo.GlowInfo.GlowInnerRadius = GlowInnerRadius;
+		TextItem.FontRenderInfo.GlowInfo.GlowOuterRadius = GlowOuterRadius;
+	}
+	else
+	{
+		TextItem.FontRenderInfo.GlowInfo.bEnableGlow = false;
+	}
+
+	return TextItem;
+}
+
+void UDMRenderTargetTextRenderer::UpdateTextureSize()
+{
+	UDMMaterialValueRenderTarget* RenderTargetValue = GetRenderTargetValue();
+
+	if (!RenderTargetValue)
+	{
+		return;
+	}
+
+	const FIntPoint RequiredSize = GetRequiredTextureSize();
+
+	RenderTargetValue->SetTextureSize(RequiredSize);
+	RenderTargetValue->FlushCreateRenderTarget();
+}
+
+void UDMRenderTargetTextRenderer::UpdateRenderTarget_Internal()
+{
+	Super::UpdateRenderTarget();
+
+	if (!IsValid(Font) || Lines.IsEmpty())
+	{
+		return;
+	}
+
+	UDMMaterialValueRenderTarget* RenderTargetValue = GetRenderTargetValue();
+
+	if (!RenderTargetValue)
+	{
+		return;
+	}
+
+	if (bRecalculateTextSize)
+	{
+		CalculateLineSizes();
+		UpdateTextureSize();
+	}
+
+	UTextureRenderTarget2D* RenderTarget = RenderTargetValue->GetRenderTarget();
+
+	if (!RenderTarget)
+	{
+		return;
+	}
+
+	RenderTarget->UpdateResourceImmediate(true);
+
+	FTextureRenderTargetResource* RenderTargetResource = RenderTarget->GameThread_GetRenderTargetResource();
+
+	if (!RenderTargetResource)
+	{
+		return;
+	}
+
+	static const FGameTime GameTime = FGameTime::CreateUndilated(0, 0);
+
+	FCanvas RenderCanvas(
+		RenderTargetResource,
+		nullptr,
+		GameTime,
+		GEngine->GetDefaultWorldFeatureLevel(),
+		FCanvas::ECanvasDrawMode::CDM_ImmediateDrawing
+	);
+
+	const FVector2D RenderTargetSize = FVector2D(RenderTarget->SizeX, RenderTarget->SizeY);
+
+	FCanvasTileItem TileItem = FCanvasTileItem(FVector2D::ZeroVector, RenderTargetSize, RenderTarget->ClearColor);
+	RenderCanvas.DrawItem(TileItem);
+
+	const float MaxCharHeight = Font->GetMaxCharHeight();
+
+	FCanvasTextItem TextItem = CreateTextItem(FVector2D::ZeroVector, FText::GetEmpty());
+
+	const float LeftOffset = PaddingLeft * static_cast<float>(TextScale.X);
+	const float RightOffset = PaddingRight * static_cast<float>(TextScale.X);
+
+	TextItem.Position.X = LeftOffset;
+	TextItem.Position.Y = PaddingTop * static_cast<float>(TextScale.Y);
+	TextItem.Scale.X = static_cast<float>(TextScale.X);
+	TextItem.Scale.Y = static_cast<float>(TextScale.Y);
+
+	if (!FMath::IsNearlyEqual(LineHeight, 1.f))
+	{
+		TextItem.Position.Y += MaxCharHeight * (LineHeight - 1.f) * 0.5f * static_cast<float>(TextScale.Y);
+	}
+
+	const float ActualLineHeight = MaxCharHeight * LineHeight * static_cast<float>(TextScale.Y);
+
+	for (const FDMTextLine& Line : Lines)
+	{
+		switch (Justify)
+		{
+			default:
+			case ETextJustify::Left:
+			case ETextJustify::InvariantLeft:
+				TextItem.Position.X = LeftOffset;
+				break;
+
+			case ETextJustify::Center:
+				TextItem.Position.X = (RenderTarget->SizeX - (Line.Width * TextItem.Scale.X)) * 0.5f;
+				break;
+
+			case ETextJustify::Right:
+			case ETextJustify::InvariantRight:
+				TextItem.Position.X = RenderTarget->SizeX - (Line.Width * TextItem.Scale.X) - RightOffset;
+				break;
+		}
+
+		TextItem.Text = FText::FromString(Line.Line);
+		RenderCanvas.DrawItem(TextItem);
+		TextItem.Position.Y += ActualLineHeight;
+	}
+
+	RenderCanvas.Flush_GameThread();
+
+	RenderTarget->UpdateResourceImmediate(false);
+}
+
+#undef LOCTEXT_NAMESPACE
diff --git a/Engine/Plugins/Experimental/DynamicMaterial/Source/DynamicMaterial/Private/DMValueDefinition.cpp b/Engine/Plugins/Experimental/DynamicMaterial/Source/DynamicMaterial/Private/DMValueDefinition.cpp
--- a/Engine/Plugins/Experimental/DynamicMaterial/Source/DynamicMaterial/Private/DMValueDefinition.cpp
+++ b/Engine/Plugins/Experimental/DynamicMaterial/Source/DynamicMaterial/Private/DMValueDefinition.cpp
@@ -1,6 +1,16 @@
 // Copyright Epic Games, Inc. All Rights Reserved.
 
 #include "DMValueDefinition.h"
+
+#include "Components/MaterialValues/DMMaterialValueBool.h"
+#include "Components/MaterialValues/DMMaterialValueColorAtlas.h"
+#include "Components/MaterialValues/DMMaterialValueFloat1.h"
+#include "Components/MaterialValues/DMMaterialValueFloat2.h"
+#include "Components/MaterialValues/DMMaterialValueFloat3RGB.h"
+#include "Components/MaterialValues/DMMaterialValueFloat3RPY.h"
+#include "Components/MaterialValues/DMMaterialValueFloat3XYZ.h"
+#include "Components/MaterialValues/DMMaterialValueFloat4.h"
+#include "Components/MaterialValues/DMMaterialValueTexture.h"
 #include "Containers/Map.h"
 
 #define LOCTEXT_NAMESPACE "DMValueDefinition"
@@ -16,8 +26,7 @@ namespace UE::MaterialDesigner::Private
 			EDMValueType::VT_Float3_XYZ,
 			EDMValueType::VT_Float4_RGBA,
 			EDMValueType::VT_Texture,
-			EDMValueType::VT_ColorAtlas,
-			EDMValueType::VT_Text
+			EDMValueType::VT_ColorAtlas
 	};
 
 	const TMap<EDMValueType, FDMValueDefinition> TypeDefinitions = {
@@ -25,26 +34,30 @@ namespace UE::MaterialDesigner::Private
 			{EDMValueType::VT_None,
 			0,
 			LOCTEXT("None", "None"),
-			{}}
+			{},
+			TSubclassOf<UDMMaterialValue>()}
 		},
 		{EDMValueType::VT_Bool,
 			{EDMValueType::VT_Bool,
 			0,
 			LOCTEXT("Bool", "Bool"),
-			{LOCTEXT("Value", "Value")}}
+			{LOCTEXT("Value", "Value")},
+			UDMMaterialValueBool::StaticClass()}
 		},
 		{EDMValueType::VT_Float1,
 			{EDMValueType::VT_Float1,
 			1,
 			LOCTEXT("Float", "Float"),
-			{LOCTEXT("Value", "Value")}}
+			{LOCTEXT("Value", "Value")},
+			UDMMaterialValueFloat1::StaticClass()}
 		},
 		{EDMValueType::VT_Float2,
 			{EDMValueType::VT_Float2,
 			2,
 			LOCTEXT("Vector2D", "Vector 2D"),
 			{LOCTEXT("U", "U"),
-				LOCTEXT("V", "V")}}
+				LOCTEXT("V", "V")},
+			UDMMaterialValueFloat2::StaticClass()}
 		},
 		{EDMValueType::VT_Float3_RPY,
 			{EDMValueType::VT_Float3_RPY,
@@ -52,7 +65,8 @@ namespace UE::MaterialDesigner::Private
 			LOCTEXT("Rotator", "Rotator"),
 			{LOCTEXT("Roll", "Roll"),
 				LOCTEXT("Pitch", "Pitch"),
-				LOCTEXT("Yaw", "Yaw")}}
+				LOCTEXT("Yaw", "Yaw")},
+			UDMMaterialValueFloat3RPY::StaticClass()}
 		},
 		{EDMValueType::VT_Float3_RGB,
 			{EDMValueType::VT_Float3_RGB,
@@ -60,7 +74,8 @@ namespace UE::MaterialDesigner::Private
 			LOCTEXT("ColorRGB", "Color (RGB)"),
 			{LOCTEXT("Red", "Red"),
 				LOCTEXT("Green", "Green"),
-				LOCTEXT("Blue", "Blue")}}
+				LOCTEXT("Blue", "Blue")},
+			UDMMaterialValueFloat3RGB::StaticClass()}
 		},
 		{EDMValueType::VT_Float3_XYZ,
 			{EDMValueType::VT_Float3_XYZ,
@@ -68,7 +83,8 @@ namespace UE::MaterialDesigner::Private
 			LOCTEXT("Vector3D", "Vector 3D"),
 			{LOCTEXT("X", "X"),
 				LOCTEXT("Y", "Y"),
-				LOCTEXT("Z", "Z")}}
+				LOCTEXT("Z", "Z")},
+			UDMMaterialValueFloat3XYZ::StaticClass()}
 		},
 		{EDMValueType::VT_Float4_RGBA,
 			{EDMValueType::VT_Float4_RGBA,
@@ -77,19 +93,22 @@ namespace UE::MaterialDesigner::Private
 			{LOCTEXT("Red", "Red"),
 				LOCTEXT("Green", "Green"),
 				LOCTEXT("Blue", "Blue"),
-				LOCTEXT("Alpha", "Alpha")}}
+				LOCTEXT("Alpha", "Alpha")},
+			UDMMaterialValueFloat4::StaticClass()}
 		},
 		{EDMValueType::VT_Float_Any,
 			{EDMValueType::VT_Float_Any,
 			0,
 			LOCTEXT("FloatAny", "Float (Any)"),
-			{}}
+			{},
+			UDMMaterialValueFloat4::StaticClass()}
 		},
 		{EDMValueType::VT_Texture,
 			{EDMValueType::VT_Texture,
 			0,
 			LOCTEXT("Texture", "Texture"),
-			{}}
+			{},
+			UDMMaterialValueTexture::StaticClass()}
 		},
 		{EDMValueType::VT_ColorAtlas,
 			{EDMValueType::VT_ColorAtlas,
@@ -98,13 +117,8 @@ namespace UE::MaterialDesigner::Private
 			{LOCTEXT("Red", "Red"),
 				LOCTEXT("Green", "Green"),
 				LOCTEXT("Blue", "Blue"),
-				LOCTEXT("Alpha", "Alpha")}}
-		},
-		{EDMValueType::VT_Text,
-			{EDMValueType::VT_Text,
-			0,
-			LOCTEXT("Text", "Text"),
-			{}}
+				LOCTEXT("Alpha", "Alpha")},
+			UDMMaterialValueColorAtlas::StaticClass()}
 		}
 	};
 }
@@ -245,6 +259,11 @@ const FText& FDMValueDefinition::GetChannelName(int32 InChannel) const
 	return Error;
 }
 
+TSubclassOf<UDMMaterialValue> FDMValueDefinition::GetValueClass() const
+{
+	return ValueClass;
+}
+
 const TArray<EDMValueType>& UDMValueDefinitionLibrary::GetValueTypes()
 {
 	using namespace UE::MaterialDesigner::Private;
diff --git a/Engine/Plugins/Experimental/DynamicMaterial/Source/DynamicMaterial/Private/Model/DynamicMaterialModel.cpp b/Engine/Plugins/Experimental/DynamicMaterial/Source/DynamicMaterial/Private/Model/DynamicMaterialModel.cpp
--- a/Engine/Plugins/Experimental/DynamicMaterial/Source/DynamicMaterial/Private/Model/DynamicMaterialModel.cpp
+++ b/Engine/Plugins/Experimental/DynamicMaterial/Source/DynamicMaterial/Private/Model/DynamicMaterialModel.cpp
@@ -6,6 +6,7 @@
 #include "Components/MaterialValues/DMMaterialValueFloat1.h"
 #include "DMComponentPath.h"
 #include "DMDefs.h"
+#include "DMValueDefinition.h"
 #include "Material/DynamicMaterialInstance.h"
 #include "Model/IDynamicMaterialModelEditorOnlyDataInterface.h"
 
@@ -150,9 +151,14 @@ UDMMaterialValue* UDynamicMaterialModel::GetValueByIndex(int32 Index) const
 	return nullptr;
 }
 
-UDMMaterialValue* UDynamicMaterialModel::AddValue(EDMValueType InType)
+UDMMaterialValue* UDynamicMaterialModel::AddValue(EDMValueType InValueType)
 {
-	UDMMaterialValue* NewValue = UDMMaterialValue::CreateMaterialValue(this, TEXT(""), InType, false);
+	return AddValue(UDMValueDefinitionLibrary::GetValueDefinition(InValueType).GetValueClass());
+}
+
+UDMMaterialValue* UDynamicMaterialModel::AddValue(TSubclassOf<UDMMaterialValue> InValueClass)
+{
+	UDMMaterialValue* NewValue = UDMMaterialValue::CreateMaterialValue(this, TEXT(""), InValueClass, false);
 	Values.Add(NewValue);
 
 	if (IDynamicMaterialModelEditorOnlyDataInterface* ModelEditorOnlyData = GetEditorOnlyData())
diff --git a/Engine/Plugins/Experimental/DynamicMaterial/Source/DynamicMaterial/Public/Components/DMMaterialValue.h b/Engine/Plugins/Experimental/DynamicMaterial/Source/DynamicMaterial/Public/Components/DMMaterialValue.h
--- a/Engine/Plugins/Experimental/DynamicMaterial/Source/DynamicMaterial/Public/Components/DMMaterialValue.h
+++ b/Engine/Plugins/Experimental/DynamicMaterial/Source/DynamicMaterial/Public/Components/DMMaterialValue.h
@@ -45,8 +45,10 @@ public:
 #if WITH_EDITOR
 	static const FName ValueName;
 
+	static UDMMaterialValue* CreateMaterialValue(UDynamicMaterialModel* InMaterialModel, const FString& InName, EDMValueType InValueType, bool bInLocal);
+
 	UFUNCTION(BlueprintCallable, Category = "Material Designer")
-	static UDMMaterialValue* CreateMaterialValue(UDynamicMaterialModel* InMaterialModel, const FString& InName, EDMValueType InType, bool bInLocal);
+	static UDMMaterialValue* CreateMaterialValue(UDynamicMaterialModel* InMaterialModel, const FString& InName, TSubclassOf<UDMMaterialValue> InValueClass, bool bInLocal);
 #endif
  
 	UFUNCTION(BlueprintPure, Category = "Material Designer")
diff --git a/Engine/Plugins/Experimental/DynamicMaterial/Source/DynamicMaterial/Public/Components/DMRenderTargetRenderer.h b/Engine/Plugins/Experimental/DynamicMaterial/Source/DynamicMaterial/Public/Components/DMRenderTargetRenderer.h
new file mode 100644
--- /dev/null
+++ b/Engine/Plugins/Experimental/DynamicMaterial/Source/DynamicMaterial/Public/Components/DMRenderTargetRenderer.h
@@ -0,0 +1,49 @@
+// Copyright Epic Games, Inc. All Rights Reserved.
+
+#pragma once
+
+#include "DMMaterialComponent.h"
+#include "DMRenderTargetRenderer.generated.h"
+
+class UDMMaterialValueRenderTarget;
+template <typename T> class TSubclassOf;
+
+UCLASS(Abstract, BlueprintType, Blueprintable, ClassGroup = "Material Designer", meta = (DisplayName = "Material Designer Render Target Renderer"))
+class DYNAMICMATERIAL_API UDMRenderTargetRenderer : public UDMMaterialComponent
+{
+	GENERATED_BODY()
+
+public:
+	FDelegateHandle EndOfFrameDelegateHandle;
+
+	UFUNCTION(BlueprintCallable, Category = "Material Designer")
+	static UDMRenderTargetRenderer* CreateRenderTargetRenderer(TSubclassOf<UDMRenderTargetRenderer>
+		InRendererClass, UDMMaterialValueRenderTarget* InRenderTargetValue);
+
+	template<typename InRendererClass
+		UE_REQUIRES(std::is_base_of_v<UDMRenderTargetRenderer, InRendererClass>)>
+	static InRendererClass* CreateRenderTargetRenderer(UDMMaterialValueRenderTarget* InRenderTargetValue)
+	{
+		return Cast<InRendererClass>(CreateRenderTargetRenderer(InRendererClass::StaticClass(), InRenderTargetValue));
+	}
+
+	UFUNCTION(BlueprintPure, Category = "Material Designer")
+	UDMMaterialValueRenderTarget* GetRenderTargetValue() const;
+
+	UFUNCTION(BlueprintCallable, Category = "Material Designer")
+	void UpdateRenderTarget();
+
+	UFUNCTION(BlueprintCallable, Category = "Material Designer")
+	void AsyncUpdateRenderTarget();
+
+	/** Will trigger the end of frame update if it is currently queued */
+	void FlushUpdateRenderTarget();
+
+	UFUNCTION(BlueprintPure, Category = "Material Designer")
+	bool IsUpdating() const { return bUpdating; }
+
+protected:
+	bool bUpdating = false;
+
+	virtual void UpdateRenderTarget_Internal() {}
+};
diff --git a/Engine/Plugins/Experimental/DynamicMaterial/Source/DynamicMaterial/Public/Components/MaterialValues/DMMaterialValueRenderTarget.h b/Engine/Plugins/Experimental/DynamicMaterial/Source/DynamicMaterial/Public/Components/MaterialValues/DMMaterialValueRenderTarget.h
new file mode 100644
--- /dev/null
+++ b/Engine/Plugins/Experimental/DynamicMaterial/Source/DynamicMaterial/Public/Components/MaterialValues/DMMaterialValueRenderTarget.h
@@ -0,0 +1,115 @@
+// Copyright Epic Games, Inc. All Rights Reserved.
+
+#pragma once
+
+#include "Components/MaterialValues/DMMaterialValueTexture.h"
+#include "Delegates/IDelegateInstance.h"
+#include "DMMaterialValueRenderTarget.generated.h"
+
+class UDMRenderTargetRenderer;
+class UTextureRenderTarget2D;
+enum ETextureRenderTargetFormat : int;
+struct FLinearColor;
+
+UCLASS(BlueprintType, ClassGroup = "Material Designer")
+class DYNAMICMATERIAL_API UDMMaterialValueRenderTarget : public UDMMaterialValueTexture
+{
+	GENERATED_BODY()
+
+public:
+	static const FString RendererPathToken;
+
+	UDMMaterialValueRenderTarget();
+
+	virtual ~UDMMaterialValueRenderTarget() override;
+
+	UFUNCTION(BlueprintPure, Category = "Material Designer")
+	UTextureRenderTarget2D* GetRenderTarget() const;
+
+	UFUNCTION(BlueprintPure, Category = "Material Designer")
+	const FIntPoint& GetTextureSize() const;
+
+	UFUNCTION(BlueprintCallable, Category = "Material Designer")
+	void SetTextureSize(const FIntPoint& InTextureSize);
+
+	UFUNCTION(BlueprintPure, Category = "Material Designer")
+	ETextureRenderTargetFormat GetTextureFormat() const;
+
+	UFUNCTION(BlueprintCallable, Category = "Material Designer")
+	void SetTextureFormat(ETextureRenderTargetFormat InTextureFormat);
+
+	UFUNCTION(BlueprintPure, Category = "Material Designer")
+	const FLinearColor& GetClearColor() const;
+
+	UFUNCTION(BlueprintCallable, Category = "Material Designer")
+	void SetClearColor(const FLinearColor& InClearColor);
+
+	UFUNCTION(BlueprintPure, Category = "Material Designer")
+	UDMRenderTargetRenderer* GetRenderer() const;
+
+	UFUNCTION(BlueprintCallable, Category = "Material Designer")
+	void SetRenderer(UDMRenderTargetRenderer* InRenderer);
+
+	/**
+	 * Allows outside objects to ensure our render target is valid.
+	 * @param bInAsync If true, will create the render target on end of frame.
+	 */
+	void EnsureRenderTarget(bool bInAsync = false);
+
+	/** Will trigger the end of frame update if it is currently queued or is invalid */
+	void FlushCreateRenderTarget();
+
+#if WITH_EDITOR
+	//~ Begin UDMMaterialValue
+	/** Render target is handled internally. */
+	virtual bool AllowEditValue() const override { return false; }
+	//~ End UDMMaterialValue
+#endif
+
+	//~ Begin UDMMaterialComponent
+	virtual void Update(EDMUpdateType InUpdateType = EDMUpdateType::Value) override;
+#if WITH_EDITOR
+	virtual void PostEditorDuplicate(UDynamicMaterialModel* InMaterialModel, UDMMaterialComponent* InParent) override;
+#endif
+	//~ End UDMMaterialComponent
+
+	//~ Begin UObject
+#if WITH_EDITOR
+	virtual void PostEditChangeProperty(FPropertyChangedEvent& InPropertyChangedEvent) override;
+#endif
+	virtual void PostLoad() override;
+	//~ End UObject
+
+protected:
+	FDelegateHandle EndOfFrameDelegateHandle;
+
+	UPROPERTY(EditAnywhere, BlueprintReadWrite, Setter = SetTextureSize, BlueprintSetter = SetTextureSize, Category = "Material Designer",
+		meta = (AllowPrivateAccess = "true"))
+	FIntPoint TextureSize;
+
+	UPROPERTY(EditAnywhere, BlueprintReadWrite, Setter = SetTextureFormat, BlueprintSetter = SetTextureFormat, Category = "Material Designer",
+		meta = (AllowPrivateAccess = "true"))
+	TEnumAsByte<ETextureRenderTargetFormat> TextureFormat;
+
+	UPROPERTY(EditAnywhere, BlueprintReadWrite, Setter = SetClearColor, BlueprintSetter = SetClearColor, Category = "Material Designer",
+		meta = (AllowPrivateAccess = "true"))
+	FLinearColor ClearColor;
+
+	UPROPERTY(EditAnywhere, BlueprintReadWrite, Setter = SetRenderer, BlueprintSetter = SetRenderer, Category = "Material Designer",
+		meta = (AllowPrivateAccess = "true"))
+	TObjectPtr<UDMRenderTargetRenderer> Renderer;
+
+	void AsyncCreateRenderTarget();
+
+	void CreateRenderTarget();
+
+	void UpdateRenderTarget();
+
+	//~ Begin UDMMaterialComponent
+	virtual UDMMaterialComponent* GetSubComponentByPath(FDMComponentPath& InPath, const FDMComponentPathSegment& InPathSegment) const override;
+#if WITH_EDITOR
+	virtual void OnComponentAdded() override;
+	virtual void OnComponentRemoved() override;
+#endif
+	//~ End UDMMaterialComponent
+};
diff --git a/Engine/Plugins/Experimental/DynamicMaterial/Source/DynamicMaterial/Public/Components/MaterialValues/DMMaterialValueText.h b/Engine/Plugins/Experimental/DynamicMaterial/Source/DynamicMaterial/Public/Components/MaterialValues/DMMaterialValueText.h
deleted file mode 100644
diff --git a/Engine/Plugins/Experimental/DynamicMaterial/Source/DynamicMaterial/Public/Components/RenderTargetRenderers/DMRenderTargetTextRenderer.h b/Engine/Plugins/Experimental/DynamicMaterial/Source/DynamicMaterial/Public/Components/RenderTargetRenderers/DMRenderTargetTextRenderer.h
new file mode 100644
--- /dev/null
+++ b/Engine/Plugins/Experimental/DynamicMaterial/Source/DynamicMaterial/Public/Components/RenderTargetRenderers/DMRenderTargetTextRenderer.h
@@ -0,0 +1,287 @@
+// Copyright Epic Games, Inc. All Rights Reserved.
+
+#pragma once
+
+#include "Components/DMRenderTargetRenderer.h"
+#include "Delegates/IDelegateInstance.h"
+#include "Framework/Text/TextLayout.h"
+#include "Internationalization/Text.h"
+#include "Math/Color.h"
+#include "Math/Vector2D.h"
+#include "DMRenderTargetTextRenderer.generated.h"
+
+class FCanvasTextItem;
+class UCanvas;
+class UFont;
+
+USTRUCT()
+struct FDMTextLine
+{
+	GENERATED_BODY()
+
+	UPROPERTY()
+	FString Line;
+
+	UPROPERTY()
+	float Width = 0.f;
+};
+
+UCLASS(BlueprintType, ClassGroup = "Material Designer", meta = (DisplayName = "Material Designer Render Target Renderer"))
+class DYNAMICMATERIAL_API UDMRenderTargetTextRenderer : public UDMRenderTargetRenderer
+{
+	GENERATED_BODY()
+
+	friend struct FDMRenderTargetTextRenderer;
+
+public:
+	UDMRenderTargetTextRenderer();
+
+	UFUNCTION(BlueprintPure, Category = "Material Designer")
+	UFont* GetFont() const;
+
+	UFUNCTION(BlueprintCallable, Category = "Material Designer")
+	void SetFont(UFont* InFont);
+
+	UFUNCTION(BlueprintPure, Category = "Material Designer")
+	const FText& GetText() const;
+
+	UFUNCTION(BlueprintCallable, Category = "Material Designer")
+	void SetText(const FText& InText);
+
+	UFUNCTION(BlueprintPure, Category = "Material Designer")
+	const FLinearColor& GetTextColor() const;
+
+	UFUNCTION(BlueprintCallable, Category = "Material Designer")
+	void SetTextColor(const FLinearColor& InColor);
+
+	UFUNCTION(BlueprintPure, Category = "Material Designer")
+	const FLinearColor& GetBackgroundColor() const;
+
+	UFUNCTION(BlueprintCallable, Category = "Material Designer")
+	void SetBackgroundColor(const FLinearColor& InBackgroundColor);
+
+	UFUNCTION(BlueprintPure, Category = "Material Designer")
+	ETextJustify::Type GetJustify() const;
+
+	UFUNCTION(BlueprintCallable, Category = "Material Designer")
+	void SetJustify(ETextJustify::Type InJustify);
+
+	UFUNCTION(BlueprintPure, Category = "Material Designer")
+	float GetKerning() const;
+
+	UFUNCTION(BlueprintCallable, Category = "Material Designer")
+	void SetKerning(float InKerning);
+
+	UFUNCTION(BlueprintPure, Category = "Material Designer")
+	float GetLineHeight() const;
+
+	UFUNCTION(BlueprintCallable, Category = "Material Designer")
+	void SetLineHeight(float InLineHeight);
+
+	UFUNCTION(BlueprintPure, Category = "Material Designer")
+	float GetPaddingLeft() const;
+
+	UFUNCTION(BlueprintCallable, Category = "Material Designer")
+	void SetPaddingLeft(float InPadding);
+
+	UFUNCTION(BlueprintPure, Category = "Material Designer")
+	float GetPaddingRight() const;
+
+	UFUNCTION(BlueprintCallable, Category = "Material Designer")
+	void SetPaddingRight(float InPadding);
+
+	UFUNCTION(BlueprintPure, Category = "Material Designer")
+	float GetPaddingTop() const;
+
+	UFUNCTION(BlueprintCallable, Category = "Material Designer")
+	void SetPaddingTop(float InPadding);
+
+	UFUNCTION(BlueprintPure, Category = "Material Designer")
+	float GetPaddingBottom() const;
+
+	UFUNCTION(BlueprintCallable, Category = "Material Designer")
+	void SetPaddingBottom(float InPadding);
+
+	UFUNCTION(BlueprintPure, Category = "Material Designer")
+	const FIntPoint& GetTextScale() const;
+
+	UFUNCTION(BlueprintCallable, Category = "Material Designer")
+	void SetTextScale(const FIntPoint& InTextScale);
+
+	UFUNCTION(BlueprintPure, Category = "Material Designer")
+	const FIntPoint& GetTextureSizeOverride() const;
+
+	UFUNCTION(BlueprintCallable, Category = "Material Designer")
+	void SetTextureSizeOverride(const FIntPoint& InTextureSizeOverride);
+
+	UFUNCTION(BlueprintPure, Category = "Material Designer")
+	bool GetHasOutline() const;
+
+	UFUNCTION(BlueprintCallable, Category = "Material Designer")
+	void SetHasOutline(bool bInHasOutline);
+
+	UFUNCTION(BlueprintPure, Category = "Material Designer")
+	const FLinearColor& GetOutlineColor() const;
+
+	UFUNCTION(BlueprintCallable, Category = "Material Designer")
+	void SetOutlineColor(const FLinearColor& InOutlineColor);
+
+	UFUNCTION(BlueprintPure, Category = "Material Designer")
+	bool GetHasShadow() const;
+
+	UFUNCTION(BlueprintCallable, Category = "Material Designer")
+	void SetHasShadow(bool bInHasShadow);
+
+	UFUNCTION(BlueprintPure, Category = "Material Designer")
+	const FLinearColor& GetShadowColor() const;
+
+	UFUNCTION(BlueprintCallable, Category = "Material Designer")
+	void SetShadowColor(const FLinearColor& InShadowColor);
+
+	UFUNCTION(BlueprintPure, Category = "Material Designer")
+	const FVector2D& GetShadowOffset() const;
+
+	UFUNCTION(BlueprintCallable, Category = "Material Designer")
+	void SetShadowOffset(const FVector2D& InShadowOffset);
+
+	UFUNCTION(BlueprintPure, Category = "Material Designer")
+	bool GetHasGlow() const;
+
+	UFUNCTION(BlueprintCallable, Category = "Material Designer")
+	void SetHasGlow(bool bInHasGlow);
+
+	UFUNCTION(BlueprintPure, Category = "Material Designer")
+	const FLinearColor& GetGlowColor() const;
+
+	UFUNCTION(BlueprintCallable, Category = "Material Designer")
+	void SetGlowColor(const FLinearColor& InGlowColor);
+
+	UFUNCTION(BlueprintPure, Category = "Material Designer")
+	const FVector2D& GetGlowInnerRadius() const;
+
+	UFUNCTION(BlueprintCallable, Category = "Material Designer")
+	void SetGlowInnerRadius(const FVector2D& InGlowInnerRadius);
+
+	UFUNCTION(BlueprintPure, Category = "Material Designer")
+	const FVector2D& GetGlowOuterRadius() const;
+
+	UFUNCTION(BlueprintCallable, Category = "Material Designer")
+	void SetGlowOuterRadius(const FVector2D& InGlowOuterRadius);
+
+	//~ Begin UObject
+	virtual void PostLoad() override;
+#if WITH_EDITOR
+	virtual void PostEditChangeProperty(FPropertyChangedEvent& InPropertyChangedEvent) override;
+#endif
+	//~ End UObject
+
+protected:
+	UPROPERTY(EditInstanceOnly, BlueprintReadWrite, Setter = SetFont, BlueprintSetter = SetFont, Category = "Material Designer",
+		meta = (NotKeyframeable, AllowPrivateAccess = "true"))
+	UFont* Font = nullptr;
+
+	UPROPERTY(EditInstanceOnly, BlueprintReadWrite, Setter = SetText, BlueprintSetter = SetText, Category = "Material Designer",
+		meta = (MultiLine, NotKeyframeable, AllowPrivateAccess = "true"))
+	FText Text = FText::GetEmpty();
+
+	UPROPERTY()
+	TArray<FDMTextLine> Lines = {};
+
+	UPROPERTY(EditInstanceOnly, BlueprintReadWrite, Setter = SetTextColor, BlueprintSetter = SetTextColor, Category = "Material Designer",
+		meta = (AllowPrivateAccess = "true"))
+	FLinearColor TextColor = FLinearColor::White;
+
+	UPROPERTY(EditInstanceOnly, BlueprintReadWrite, Setter = SetJustify, BlueprintSetter = SetJustify, Category = "Material Designer",
+		meta = (NotKeyframeable, AllowPrivateAccess = "true"))
+	TEnumAsByte<ETextJustify::Type> Justify = ETextJustify::Left;
+
+	UPROPERTY(EditInstanceOnly, BlueprintReadWrite, Setter = SetKerning, BlueprintSetter = SetKerning, Category = "Material Designer",
+		meta = (AllowPrivateAccess = "true"))
+	float Kerning = 0.f;
+
+	/** Multiplier on the base font height. */
+	UPROPERTY(EditInstanceOnly, BlueprintReadWrite, Setter = SetLineHeight, BlueprintSetter = SetLineHeight, Category = "Material Designer",
+		meta = (NotKeyframeable, AllowPrivateAccess = "true", UIMin = 0.1, UIMax = 2.0))
+	float LineHeight = 1.f;
+
+	/** Extra space adding beyond the edge of the glyphs. Useful for shadows, glows, etc. */
+	UPROPERTY(EditInstanceOnly, BlueprintReadWrite, Setter = SetPaddingLeft, BlueprintSetter = SetPaddingLeft, Category = "Material Designer",
+		meta = (NotKeyframeable, AllowPrivateAccess = "true"))
+	float PaddingLeft = 0.f;
+
+	/** Extra space adding beyond the edge of the glyphs. Useful for shadows, glows, etc. */
+	UPROPERTY(EditInstanceOnly, BlueprintReadWrite, Setter = SetPaddingRight, BlueprintSetter = SetPaddingRight, Category = "Material Designer",
+		meta = (NotKeyframeable, AllowPrivateAccess = "true"))
+	float PaddingRight = 0.f;
+
+	/** Extra space adding beyond the edge of the glyphs. Useful for shadows, glows, etc. */
+	UPROPERTY(EditInstanceOnly, BlueprintReadWrite, Setter = SetPaddingTop, BlueprintSetter = SetPaddingTop, Category = "Material Designer",
+		meta = (NotKeyframeable, AllowPrivateAccess = "true"))
+	float PaddingTop = 0.f;
+
+	/** Extra space adding beyond the edge of the glyphs. Useful for shadows, glows, etc. */
+	UPROPERTY(EditInstanceOnly, BlueprintReadWrite, Setter = SetPaddingBottom, BlueprintSetter = SetPaddingBottom, Category = "Material Designer",
+		meta = (NotKeyframeable, AllowPrivateAccess = "true"))
+	float PaddingBottom = 0.f;
+
+	UPROPERTY(EditInstanceOnly, BlueprintReadWrite, Setter = SetTextScale, BlueprintSetter = SetTextScale, Category = "Material Designer",
+		meta = (NotKeyframeable, AllowPrivateAccess = "true", ClampMin = 1, UIMin = 1, ClampMax = 64, UIMax = 64,
+			ToolTip = "Increasing this will increase the resolution of the text."))
+	FIntPoint TextScale = FIntPoint(10, 10);
+
+	UPROPERTY(EditInstanceOnly, BlueprintReadWrite, Setter = SetTextureSizeOverride, BlueprintSetter = SetTextureSizeOverride, Category = "Material Designer",
+		meta = (NotKeyframeable, AllowPrivateAccess = "true", ToolTip = "Override the size of the render target. Set to 0 to disable."))
+	FIntPoint TextureSizeOverride = FIntPoint(0, 0);
+
+	UPROPERTY(EditInstanceOnly, BlueprintReadWrite, Setter = SetHasOutline, BlueprintSetter = SetHasOutline, Category = "Material Designer",
+		meta = (AllowPrivateAccess = "true"))
+	bool bOutline = false;
+
+	UPROPERTY(EditInstanceOnly, BlueprintReadWrite, Setter = SetOutlineColor, BlueprintSetter = SetOutlineColor, Category = "Material Designer",
+		meta = (AllowPrivateAccess = "true", EditCondition = "bOutline", EditConditionHides = true))
+	FLinearColor OutlineColor = FLinearColor::Black;
+
+	UPROPERTY(EditInstanceOnly, BlueprintReadWrite, Setter = SetHasShadow, BlueprintSetter = SetHasShadow, Category = "Material Designer",
+		meta = (AllowPrivateAccess = "true"))
+	bool bShadow = false;
+
+	UPROPERTY(EditInstanceOnly, BlueprintReadWrite, Setter = SetShadowColor, BlueprintSetter = SetShadowColor, Category = "Material Designer",
+		meta = (AllowPrivateAccess = "true", EditCondition = "bShadow", EditConditionHides = true))
+	FLinearColor ShadowColor = FLinearColor::Black;
+
+	UPROPERTY(EditInstanceOnly, BlueprintReadWrite, Setter = SetShadowOffset, BlueprintSetter = SetShadowOffset, Category = "Material Designer",
+		meta = (AllowPrivateAccess = "true", EditCondition = "bShadow", EditConditionHides = true))
+	FVector2D ShadowOffset = FVector2D(1.0, 1.0);
+
+	UPROPERTY(EditInstanceOnly, BlueprintReadWrite, Setter = SetHasGlow, BlueprintSetter = SetHasGlow, Category = "Material Designer",
+		meta = (AllowPrivateAccess = "true"))
+	bool bGlow = false;
+
+	UPROPERTY(EditInstanceOnly, BlueprintReadWrite, Setter = SetGlowColor, BlueprintSetter = SetGlowColor, Category = "Material Designer",
+		meta = (AllowPrivateAccess = "true", EditCondition = "bGlow", EditConditionHides = true))
+	FLinearColor GlowColor = FLinearColor::White;
+
+	UPROPERTY(EditInstanceOnly, BlueprintReadWrite, Setter = SetGlowInnerRadius, BlueprintSetter = SetGlowInnerRadius, Category = "Material Designer",
+		meta = (AllowPrivateAccess = "true", EditCondition = "bGlow", EditConditionHides = true, ClampMin = 0, UIMin = 0, ClampMax = 1, UIMax = 1))
+	FVector2D GlowInnerRadius = FVector2D(0.5, 0.5);
+
+	UPROPERTY(EditInstanceOnly, BlueprintReadWrite, Setter = SetGlowOuterRadius, BlueprintSetter = SetGlowOuterRadius, Category = "Material Designer",
+		meta = (AllowPrivateAccess = "true", EditCondition = "bGlow", EditConditionHides = true, ClampMin = 0, UIMin = 0, ClampMax = 1, UIMax = 1))
+	FVector2D GlowOuterRadius = FVector2D(0.5, 0.5);
+
+	bool bRecalculateTextSize = false;
+
+	void CalculateLineSizes();
+
+	FIntPoint GetRequiredTextureSize() const;
+
+	void UpdateTextLines();
+
+	FCanvasTextItem CreateTextItem(const FVector2D& InPosition, const FText& InText) const;
+
+	void UpdateTextureSize();
+
+	//~ Begin UDMRenderTargetRenderer
+	virtual void UpdateRenderTarget_Internal() override;
+	//~ End UDMRenderTargetRenderer
+};
diff --git a/Engine/Plugins/Experimental/DynamicMaterial/Source/DynamicMaterial/Public/DMDefs.h b/Engine/Plugins/Experimental/DynamicMaterial/Source/DynamicMaterial/Public/DMDefs.h
--- a/Engine/Plugins/Experimental/DynamicMaterial/Source/DynamicMaterial/Public/DMDefs.h
+++ b/Engine/Plugins/Experimental/DynamicMaterial/Source/DynamicMaterial/Public/DMDefs.h
@@ -15,7 +15,7 @@ UENUM(BlueprintType)
 enum class EDMMaterialPropertyType : uint8
 {
 	None = 0,
-	BaseColor UMETA(ShortName = "Base Color"),
+	BaseColor,
 	EmissiveColor UMETA(ShortName = "Emissive"),
 	Opacity,
 	OpacityMask UMETA(ShortName = "Mask"),
@@ -50,7 +50,7 @@ enum class EDMValueType : uint8
 	VT_Float_Any,
 	VT_Texture,
 	VT_ColorAtlas,
-	VT_Text,
+	VT_Unused,
 	VT_MAX
 };
 
@@ -167,6 +167,7 @@ namespace UE::DynamicMaterial
 	FString DYNAMICMATERIAL_API CreateNodeComment(const ANSICHAR* InFile, int InLine, const ANSICHAR* InFunction, const FString* InComment = nullptr);
 
 	constexpr int32 RenameFlags = REN_DontCreateRedirectors | REN_DoNotDirty | REN_ForceNoResetLoaders | REN_NonTransactional;
+	constexpr EClassFlags InvalidClassFlags = CLASS_Abstract | CLASS_Hidden | CLASS_Deprecated | CLASS_NewerVersionExists;
 }
 
 #define UE_DM_NodeComment_Default UE::DynamicMaterial::CreateNodeComment(__FILE__, __LINE__, __FUNCTION__)
diff --git a/Engine/Plugins/Experimental/DynamicMaterial/Source/DynamicMaterial/Public/DMValueDefinition.h b/Engine/Plugins/Experimental/DynamicMaterial/Source/DynamicMaterial/Public/DMValueDefinition.h
--- a/Engine/Plugins/Experimental/DynamicMaterial/Source/DynamicMaterial/Public/DMValueDefinition.h
+++ b/Engine/Plugins/Experimental/DynamicMaterial/Source/DynamicMaterial/Public/DMValueDefinition.h
@@ -8,29 +8,35 @@
 #include "Kismet/BlueprintFunctionLibrary.h"
 #include "DMValueDefinition.generated.h"
 
+class UDMMaterialValue;
+
 USTRUCT(BlueprintType)
 struct DYNAMICMATERIAL_API FDMValueDefinition
 {
 	GENERATED_BODY()
 
 	FDMValueDefinition()
-		: FDMValueDefinition(EDMValueType::VT_None, 0, FText::GetEmpty(), {})
+		: FDMValueDefinition(EDMValueType::VT_None, 0, FText::GetEmpty(), {}, TSubclassOf<UDMMaterialValue>())
 	{		
 	}
 
-	FDMValueDefinition(EDMValueType InType, uint8 InFloatCount, const FText& InDisplayName, const TArray<FText>& InChannelNames)
+	FDMValueDefinition(EDMValueType InType, uint8 InFloatCount, const FText& InDisplayName, const TArray<FText>& InChannelNames, 
+		TSubclassOf<UDMMaterialValue> InValueClass)
 		: Type(InType)
 		, FloatCount(InFloatCount)
 		, DisplayName(InDisplayName)
 		, ChannelNames(InChannelNames)
+		, ValueClass(InValueClass)
 	{
 	}
 
-	FDMValueDefinition(EDMValueType InType, uint8 InFloatCount, FText&& InDisplayName, TArray<FText>&& InChannelNames)
+	FDMValueDefinition(EDMValueType InType, uint8 InFloatCount, FText&& InDisplayName, TArray<FText>&& InChannelNames, 
+		TSubclassOf<UDMMaterialValue> InValueClass)
 		: Type(InType)
 		, FloatCount(InFloatCount)
 		, DisplayName(MoveTemp(InDisplayName))
 		, ChannelNames(MoveTemp(InChannelNames))
+		, ValueClass(InValueClass)
 	{
 	}
 
@@ -49,6 +55,8 @@ struct DYNAMICMATERIAL_API FDMValueDefinition
 	/** To be consistent without OutputChannel, 1 is the first channel, not 0. */
 	const FText& GetChannelName(int32 InChannel) const;
 
+	TSubclassOf<UDMMaterialValue> GetValueClass() const;
+
 protected:
 	UPROPERTY(VisibleAnywhere, BlueprintReadOnly, Category = "Material Designer")
 	EDMValueType Type;
@@ -61,6 +69,9 @@ protected:
 
 	UPROPERTY(VisibleAnywhere, BlueprintReadOnly, Category = "Material Designer")
 	TArray<FText> ChannelNames;
+
+	UPROPERTY(VisibleAnywhere, BlueprintReadOnly, Category = "Material Designer")
+	TSubclassOf<UDMMaterialValue> ValueClass;
 };
 
 UCLASS(BlueprintType)
@@ -75,7 +86,7 @@ public:
 	UFUNCTION(BlueprintPure, CallInEditor, Category = "Material Designer")
 	static const FDMValueDefinition& GetValueDefinition(EDMValueType InValueType);
 
-	UFUNCTION(BlueprintPure, CallInEditor, Category = "Material Designer", Meta = (DisplayName = "Are Types Compatible"))
+	UFUNCTION(BlueprintPure, CallInEditor, Category = "Material Designer", meta = (DisplayName = "Are Types Compatible"))
 	static bool BP_AreTypesCompatible(EDMValueType A, EDMValueType B, int32 AChannel, int32 BChannel)
 	{
 		return AreTypesCompatible(A, B, AChannel, BChannel);
diff --git a/Engine/Plugins/Experimental/DynamicMaterial/Source/DynamicMaterial/Public/Model/DynamicMaterialModel.h b/Engine/Plugins/Experimental/DynamicMaterial/Source/DynamicMaterial/Public/Model/DynamicMaterialModel.h
--- a/Engine/Plugins/Experimental/DynamicMaterial/Source/DynamicMaterial/Public/Model/DynamicMaterialModel.h
+++ b/Engine/Plugins/Experimental/DynamicMaterial/Source/DynamicMaterial/Public/Model/DynamicMaterialModel.h
@@ -85,8 +85,10 @@ public:
 	UFUNCTION(BlueprintPure, Category = "Material Designer")
 	UDMMaterialValue* GetValueByIndex(int32 Index) const;
 
+	UDMMaterialValue* AddValue(EDMValueType InValueType);
+
 	UFUNCTION(BlueprintCallable, Category = "Material Designer")
-	UDMMaterialValue* AddValue(EDMValueType InType);
+	UDMMaterialValue* AddValue(TSubclassOf<UDMMaterialValue> InValueClass);
 
 	UFUNCTION(BlueprintCallable, Category = "Material Designer")
 	void AddRuntimeComponentReference(UDMMaterialComponent* InValue);
diff --git a/Engine/Plugins/Experimental/DynamicMaterial/Source/DynamicMaterialEditor/Private/Components/DMMaterialStage.cpp b/Engine/Plugins/Experimental/DynamicMaterial/Source/DynamicMaterialEditor/Private/Components/DMMaterialStage.cpp
--- a/Engine/Plugins/Experimental/DynamicMaterial/Source/DynamicMaterialEditor/Private/Components/DMMaterialStage.cpp
+++ b/Engine/Plugins/Experimental/DynamicMaterial/Source/DynamicMaterialEditor/Private/Components/DMMaterialStage.cpp
@@ -770,7 +770,7 @@ UDMMaterialStageInput* UDMMaterialStage::ChangeInput(TSubclassOf<UDMMaterialStag
 	check(InputConnectors.IsValidIndex(InInputIdx));
 
 	check(InInputClass.Get());
-	check(!(InInputClass->ClassFlags & (CLASS_Abstract | CLASS_Hidden | CLASS_Deprecated | CLASS_NewerVersionExists)));
+	check(!InInputClass->HasAnyClassFlags(UE::DynamicMaterial::InvalidClassFlags));
 
 	if (UDMMaterialStageThroughput* ThroughputCDO = Cast<UDMMaterialStageThroughput>(InInputClass->GetDefaultObject(true)))
 	{
@@ -1215,7 +1215,7 @@ UDMMaterialStageSource* UDMMaterialStage::ChangeSource(TSubclassOf<UDMMaterialSt
 	}
 
 	check(InSourceClass);
-	check(!(InSourceClass->ClassFlags & (CLASS_Abstract | CLASS_Hidden | CLASS_Deprecated | CLASS_NewerVersionExists)));
+	check(!InSourceClass->HasAnyClassFlags(UE::DynamicMaterial::InvalidClassFlags));
 
 	UDMMaterialStageSource* NewSource = NewObject<UDMMaterialStageSource>(this, InSourceClass, NAME_None, RF_Transactional);
 	check(NewSource);
diff --git a/Engine/Plugins/Experimental/DynamicMaterial/Source/DynamicMaterialEditor/Private/Components/DMMaterialStageExpression.cpp b/Engine/Plugins/Experimental/DynamicMaterial/Source/DynamicMaterialEditor/Private/Components/DMMaterialStageExpression.cpp
--- a/Engine/Plugins/Experimental/DynamicMaterial/Source/DynamicMaterialEditor/Private/Components/DMMaterialStageExpression.cpp
+++ b/Engine/Plugins/Experimental/DynamicMaterial/Source/DynamicMaterialEditor/Private/Components/DMMaterialStageExpression.cpp
@@ -100,7 +100,7 @@ UDMMaterialStageExpression* UDMMaterialStageExpression::ChangeStageSource_Expres
 	}
 
 	check(InExpressionClass);
-	check(!(InExpressionClass->ClassFlags & (CLASS_Abstract | CLASS_Hidden | CLASS_Deprecated | CLASS_NewerVersionExists)));
+	check(!InExpressionClass->HasAnyClassFlags(UE::DynamicMaterial::InvalidClassFlags));
 
 	return InStage->ChangeSource<UDMMaterialStageExpression>(InExpressionClass);
 }
@@ -127,7 +127,7 @@ void UDMMaterialStageExpression::GenerateExpressionList()
 			continue;
 		}
 
-		if (ExpressionClass->ClassFlags & (CLASS_Abstract | CLASS_Hidden | CLASS_Deprecated | CLASS_NewerVersionExists))
+		if (ExpressionClass->HasAnyClassFlags(UE::DynamicMaterial::InvalidClassFlags))
 		{
 			continue;
 		}
diff --git a/Engine/Plugins/Experimental/DynamicMaterial/Source/DynamicMaterialEditor/Private/Components/DMMaterialStageGradient.cpp b/Engine/Plugins/Experimental/DynamicMaterial/Source/DynamicMaterialEditor/Private/Components/DMMaterialStageGradient.cpp
--- a/Engine/Plugins/Experimental/DynamicMaterial/Source/DynamicMaterialEditor/Private/Components/DMMaterialStageGradient.cpp
+++ b/Engine/Plugins/Experimental/DynamicMaterial/Source/DynamicMaterialEditor/Private/Components/DMMaterialStageGradient.cpp
@@ -69,7 +69,7 @@ UDMMaterialStageGradient* UDMMaterialStageGradient::ChangeStageSource_Gradient(U
 	}
 
 	check(InGradientClass);
-	check(!(InGradientClass->ClassFlags & (CLASS_Abstract | CLASS_Hidden | CLASS_Deprecated | CLASS_NewerVersionExists)));
+	check(!InGradientClass->HasAnyClassFlags(UE::DynamicMaterial::InvalidClassFlags));
 
 	return InStage->ChangeSource<UDMMaterialStageGradient>(InGradientClass);
 }
diff --git a/Engine/Plugins/Experimental/DynamicMaterial/Source/DynamicMaterialEditor/Private/Components/DMMaterialStageSource.cpp b/Engine/Plugins/Experimental/DynamicMaterial/Source/DynamicMaterialEditor/Private/Components/DMMaterialStageSource.cpp
--- a/Engine/Plugins/Experimental/DynamicMaterial/Source/DynamicMaterialEditor/Private/Components/DMMaterialStageSource.cpp
+++ b/Engine/Plugins/Experimental/DynamicMaterial/Source/DynamicMaterialEditor/Private/Components/DMMaterialStageSource.cpp
@@ -99,7 +99,7 @@ void UDMMaterialStageSource::GenerateClassList()
 			continue;
 		}
 
-		if (MSEClass->ClassFlags & (CLASS_Abstract | CLASS_Hidden | CLASS_Deprecated | CLASS_NewerVersionExists))
+		if (MSEClass->HasAnyClassFlags(UE::DynamicMaterial::InvalidClassFlags))
 		{
 			continue;
 		}
diff --git a/Engine/Plugins/Experimental/DynamicMaterial/Source/DynamicMaterialEditor/Private/Components/DMMaterialStageThroughput.cpp b/Engine/Plugins/Experimental/DynamicMaterial/Source/DynamicMaterialEditor/Private/Components/DMMaterialStageThroughput.cpp
--- a/Engine/Plugins/Experimental/DynamicMaterial/Source/DynamicMaterialEditor/Private/Components/DMMaterialStageThroughput.cpp
+++ b/Engine/Plugins/Experimental/DynamicMaterial/Source/DynamicMaterialEditor/Private/Components/DMMaterialStageThroughput.cpp
@@ -461,7 +461,6 @@ void UDMMaterialStageThroughput::AddDefaultInput(int32 InInputIndex) const
 		case EDMValueType::VT_Float3_RPY:
 		case EDMValueType::VT_Float3_XYZ:
 		case EDMValueType::VT_Texture:
-		case EDMValueType::VT_Text:
 			UDMMaterialStageInputValue::ChangeStageInput_NewLocalValue(
 				Stage, 
 				InInputIndex,
diff --git a/Engine/Plugins/Experimental/DynamicMaterial/Source/DynamicMaterialEditor/Private/Components/MaterialStageExpressions/DMMSETextRenderer.cpp b/Engine/Plugins/Experimental/DynamicMaterial/Source/DynamicMaterialEditor/Private/Components/MaterialStageExpressions/DMMSETextRenderer.cpp
deleted file mode 100644
diff --git a/Engine/Plugins/Experimental/DynamicMaterial/Source/DynamicMaterialEditor/Private/Components/MaterialStageExpressions/DMMSETextureSampleBase.cpp b/Engine/Plugins/Experimental/DynamicMaterial/Source/DynamicMaterialEditor/Private/Components/MaterialStageExpressions/DMMSETextureSampleBase.cpp
--- a/Engine/Plugins/Experimental/DynamicMaterial/Source/DynamicMaterialEditor/Private/Components/MaterialStageExpressions/DMMSETextureSampleBase.cpp
+++ b/Engine/Plugins/Experimental/DynamicMaterial/Source/DynamicMaterialEditor/Private/Components/MaterialStageExpressions/DMMSETextureSampleBase.cpp
@@ -212,6 +212,20 @@ void UDMMaterialStageExpressionTextureSampleBase::SetClampTextureEnabled(bool bI
 	Update(EDMUpdateType::Value);
 }
 
+void UDMMaterialStageExpressionTextureSampleBase::PostEditChangeProperty(FPropertyChangedEvent& InPropertyChangedEvent)
+{
+	Super::PostEditChangeProperty(InPropertyChangedEvent);
+
+	static const FName ClampTextureName = GET_MEMBER_NAME_CHECKED(UDMMaterialStageExpressionTextureSampleBase, bClampTexture);
+
+	const FName PropertyName = InPropertyChangedEvent.GetMemberPropertyName();
+
+	if (PropertyName == ClampTextureName)
+	{
+		Update(EDMUpdateType::Structure);
+	}
+}
+
 void UDMMaterialStageExpressionTextureSampleBase::UpdateMask()
 {
 	UDMMaterialStage* BaseTextureSampleStage = GetStage();
@@ -378,9 +392,13 @@ void UDMMaterialStageExpressionTextureSampleBase::UpdateMask()
 	UDMMaterialStage* MaskTextureSampleStage = MaskTextureSample->GetStage();
 	check(MaskTextureSampleStage);
 
-	UDMMaterialStageInputValue::ChangeStageInput_NewLocalValue(MaskTextureSampleStage, 0,
-		FDMMaterialStageConnectorChannel::WHOLE_CHANNEL, BaseTextureValue->GetType(),
-		FDMMaterialStageConnectorChannel::WHOLE_CHANNEL);
+	UDMMaterialStageInputValue::ChangeStageInput_NewLocalValue(
+		MaskTextureSampleStage, 
+		0,
+		FDMMaterialStageConnectorChannel::WHOLE_CHANNEL,
+		UDMMaterialValueTexture::StaticClass(),
+		FDMMaterialStageConnectorChannel::WHOLE_CHANNEL
+	);
 
 	UDMMaterialStageInputValue* NewInputValue = Cast<UDMMaterialStageInputValue>(MaskTextureSampleStage->GetInputs().Last());
 	check(NewInputValue);
diff --git a/Engine/Plugins/Experimental/DynamicMaterial/Source/DynamicMaterialEditor/Private/Components/MaterialStageInputs/DMMSIExpression.cpp b/Engine/Plugins/Experimental/DynamicMaterial/Source/DynamicMaterialEditor/Private/Components/MaterialStageInputs/DMMSIExpression.cpp
--- a/Engine/Plugins/Experimental/DynamicMaterial/Source/DynamicMaterialEditor/Private/Components/MaterialStageInputs/DMMSIExpression.cpp
+++ b/Engine/Plugins/Experimental/DynamicMaterial/Source/DynamicMaterialEditor/Private/Components/MaterialStageInputs/DMMSIExpression.cpp
@@ -61,7 +61,7 @@ UDMMaterialStageInputExpression* UDMMaterialStageInputExpression::ChangeStageSou
 	}
 
 	check(InExpressionClass);
-	check(!(InExpressionClass->ClassFlags & (CLASS_Abstract | CLASS_Hidden | CLASS_Deprecated | CLASS_NewerVersionExists)));
+	check(!InExpressionClass->HasAnyClassFlags(UE::DynamicMaterial::InvalidClassFlags));
 
 	UDMMaterialStageInputExpression* NewInputExpression = InStage->ChangeSource<UDMMaterialStageInputExpression>(
 		[InExpressionClass](UDMMaterialStage* InStage, UDMMaterialStageSource* InNewSource)
@@ -83,7 +83,7 @@ UDMMaterialStageInputExpression* UDMMaterialStageInputExpression::ChangeStageInp
 	check(Source);
 
 	check(InExpressionClass);
-	check(!(InExpressionClass->ClassFlags & (CLASS_Abstract | CLASS_Hidden | CLASS_Deprecated | CLASS_NewerVersionExists)));
+	check(!InExpressionClass->HasAnyClassFlags(UE::DynamicMaterial::InvalidClassFlags));
 
 	UDMMaterialStageInputExpression* NewInputExpression = InStage->ChangeInput<UDMMaterialStageInputExpression>(
 		InInputIdx, InInputChannel, InOutputIdx, InOutputChannel,
diff --git a/Engine/Plugins/Experimental/DynamicMaterial/Source/DynamicMaterialEditor/Private/Components/MaterialStageInputs/DMMSIGradient.cpp b/Engine/Plugins/Experimental/DynamicMaterial/Source/DynamicMaterialEditor/Private/Components/MaterialStageInputs/DMMSIGradient.cpp
--- a/Engine/Plugins/Experimental/DynamicMaterial/Source/DynamicMaterialEditor/Private/Components/MaterialStageInputs/DMMSIGradient.cpp
+++ b/Engine/Plugins/Experimental/DynamicMaterial/Source/DynamicMaterialEditor/Private/Components/MaterialStageInputs/DMMSIGradient.cpp
@@ -39,7 +39,7 @@ UDMMaterialStageInputGradient* UDMMaterialStageInputGradient::ChangeStageSource_
 	}
 
 	check(InGradientClass);
-	check(!(InGradientClass->ClassFlags & (CLASS_Abstract | CLASS_Hidden | CLASS_Deprecated | CLASS_NewerVersionExists)));
+	check(!InGradientClass->HasAnyClassFlags(UE::DynamicMaterial::InvalidClassFlags));
 
 	UDMMaterialStageInputGradient* NewInputGradient = InStage->ChangeSource<UDMMaterialStageInputGradient>(
 		[InGradientClass](UDMMaterialStage* InStage, UDMMaterialStageSource* InNewSource)
@@ -60,7 +60,7 @@ UDMMaterialStageInputGradient* UDMMaterialStageInputGradient::ChangeStageInput_G
 	check(Source);
 
 	check(InGradientClass);
-	check(!(InGradientClass->ClassFlags & (CLASS_Abstract | CLASS_Hidden | CLASS_Deprecated | CLASS_NewerVersionExists)));
+	check(!InGradientClass->HasAnyClassFlags(UE::DynamicMaterial::InvalidClassFlags));
 
 	UDMMaterialStageInputGradient* NewInputGradient = InStage->ChangeInput<UDMMaterialStageInputGradient>(
 		InInputIdx, InInputChannel, 0, InOutputChannel,
diff --git a/Engine/Plugins/Experimental/DynamicMaterial/Source/DynamicMaterialEditor/Private/Components/MaterialStageInputs/DMMSIThroughput.cpp b/Engine/Plugins/Experimental/DynamicMaterial/Source/DynamicMaterialEditor/Private/Components/MaterialStageInputs/DMMSIThroughput.cpp
--- a/Engine/Plugins/Experimental/DynamicMaterial/Source/DynamicMaterialEditor/Private/Components/MaterialStageInputs/DMMSIThroughput.cpp
+++ b/Engine/Plugins/Experimental/DynamicMaterial/Source/DynamicMaterialEditor/Private/Components/MaterialStageInputs/DMMSIThroughput.cpp
@@ -243,6 +243,11 @@ void UDMMaterialStageInputThroughput::PostEditorDuplicate(UDynamicMaterialModel*
 
 FText UDMMaterialStageInputThroughput::GetComponentDescription() const
 {
+	if (UDMMaterialStageThroughput* Throughput = GetMaterialStageThroughput())
+	{
+		return Throughput->GetComponentDescription();
+	}
+
 	if (TSubclassOf<UDMMaterialStageThroughput> MaterialStageThroughputClass = GetMaterialStageThroughputClass())
 	{
 		check(MaterialStageThroughputClass.Get());
diff --git a/Engine/Plugins/Experimental/DynamicMaterial/Source/DynamicMaterialEditor/Private/Components/MaterialStageInputs/DMMSIValue.cpp b/Engine/Plugins/Experimental/DynamicMaterial/Source/DynamicMaterialEditor/Private/Components/MaterialStageInputs/DMMSIValue.cpp
--- a/Engine/Plugins/Experimental/DynamicMaterial/Source/DynamicMaterialEditor/Private/Components/MaterialStageInputs/DMMSIValue.cpp
+++ b/Engine/Plugins/Experimental/DynamicMaterial/Source/DynamicMaterialEditor/Private/Components/MaterialStageInputs/DMMSIValue.cpp
@@ -34,7 +34,15 @@ UDMMaterialStage* UDMMaterialStageInputValue::CreateStage(UDMMaterialValue* InVa
 	return NewStage;
 }
 
-UDMMaterialStageInputValue* UDMMaterialStageInputValue::ChangeStageSource_NewLocalValue(UDMMaterialStage* InStage, EDMValueType InType)
+UDMMaterialStageInputValue* UDMMaterialStageInputValue::ChangeStageSource_NewLocalValue(UDMMaterialStage* InStage, EDMValueType InValueType)
+{
+	return ChangeStageSource_NewLocalValue(
+		InStage,
+		UDMValueDefinitionLibrary::GetValueDefinition(InValueType).GetValueClass()
+	);
+}
+
+UDMMaterialStageInputValue* UDMMaterialStageInputValue::ChangeStageSource_NewLocalValue(UDMMaterialStage* InStage, TSubclassOf<UDMMaterialValue> InValueClass)
 {
 	check(InStage);
 
@@ -57,7 +65,7 @@ UDMMaterialStageInputValue* UDMMaterialStageInputValue::ChangeStageSource_NewLoc
 
 	// Don't add via the builder, just parent it to the builder.
 	// It won't appear in the global value list.
-	UDMMaterialValue* NewValue = UDMMaterialValue::CreateMaterialValue(MaterialModel, TEXT(""), InType, true);
+	UDMMaterialValue* NewValue = UDMMaterialValue::CreateMaterialValue(MaterialModel, TEXT(""), InValueClass, true);
 	check(NewValue);
 
 	UDMMaterialStageInputValue* InputValue = InStage->ChangeSource<UDMMaterialStageInputValue>(
@@ -101,7 +109,16 @@ UDMMaterialStageInputValue* UDMMaterialStageInputValue::ChangeStageSource_Value(
 	return InputValue;
 }
 
-UDMMaterialStageInputValue* UDMMaterialStageInputValue::ChangeStageSource_NewValue(UDMMaterialStage* InStage, EDMValueType InType)
+UDMMaterialStageInputValue* UDMMaterialStageInputValue::ChangeStageSource_NewValue(UDMMaterialStage* InStage, EDMValueType InValueType)
+{
+	return ChangeStageSource_NewValue(
+		InStage,
+		UDMValueDefinitionLibrary::GetValueDefinition(InValueType).GetValueClass()
+	);
+}
+
+UDMMaterialStageInputValue* UDMMaterialStageInputValue::ChangeStageSource_NewValue(UDMMaterialStage* InStage, 
+	TSubclassOf<UDMMaterialValue> InValueClass)
 {
 	check(InStage);
 
@@ -122,7 +139,7 @@ UDMMaterialStageInputValue* UDMMaterialStageInputValue::ChangeStageSource_NewVal
 	UDynamicMaterialModel* MaterialModel = ModelEditorOnlyData->GetMaterialModel();
 	check(MaterialModel);
 
-	UDMMaterialValue* NewValue = MaterialModel->AddValue(InType);
+	UDMMaterialValue* NewValue = MaterialModel->AddValue(InValueClass);
 	check(NewValue);
 
 	UDMMaterialStageInputValue* InputValue = InStage->ChangeSource<UDMMaterialStageInputValue>(
@@ -135,8 +152,20 @@ UDMMaterialStageInputValue* UDMMaterialStageInputValue::ChangeStageSource_NewVal
 	return InputValue;
 }
 
+UDMMaterialStageInputValue* UDMMaterialStageInputValue::ChangeStageInput_NewLocalValue(UDMMaterialStage* InStage, int32 InInputIdx, 
+	int32 InInputChannel, EDMValueType InValueType, int32 InOutputChannel)
+{
+	return ChangeStageInput_NewLocalValue(
+		InStage,
+		InInputIdx,
+		InInputChannel,
+		UDMValueDefinitionLibrary::GetValueDefinition(InValueType).GetValueClass(),
+		InOutputChannel
+	);
+}
+
 UDMMaterialStageInputValue* UDMMaterialStageInputValue::ChangeStageInput_NewLocalValue(UDMMaterialStage* InStage,
-	int32 InInputIdx, int32 InInputChannel, EDMValueType InType, int32 InOutputChannel)
+	int32 InInputIdx, int32 InInputChannel, TSubclassOf<UDMMaterialValue> InValueClass, int32 InOutputChannel)
 {
 	check(InStage);
 
@@ -148,7 +177,7 @@ UDMMaterialStageInputValue* UDMMaterialStageInputValue::ChangeStageInput_NewLoca
 
 	const TArray<FDMMaterialStageConnector>& InputConnectors = Throughput->GetInputConnectors();
 	check(InputConnectors.IsValidIndex(InInputIdx));
-	check(Throughput->CanInputAcceptType(InInputIdx, InType));
+	check(Throughput->CanInputAcceptType(InInputIdx, InValueClass.GetDefaultObject()->GetType()));
 
 	UDMMaterialLayerObject* Layer = InStage->GetLayer();
 	check(Layer);
@@ -164,7 +193,7 @@ UDMMaterialStageInputValue* UDMMaterialStageInputValue::ChangeStageInput_NewLoca
 
 	// Don't add via the builder, just parent it to the builder.
 	// It won't appear in the global value list.
-	UDMMaterialValue* NewValue = UDMMaterialValue::CreateMaterialValue(MaterialModel, TEXT(""), InType, true);
+	UDMMaterialValue* NewValue = UDMMaterialValue::CreateMaterialValue(MaterialModel, TEXT(""), InValueClass, true);
 	check(NewValue);
 
 	UDMMaterialStageInputValue* NewInputValue = InStage->ChangeInput<UDMMaterialStageInputValue>(
@@ -232,7 +261,20 @@ UDMMaterialStageInputValue* UDMMaterialStageInputValue::ChangeStageInput_Value(U
 	return NewInputValue;
 }
 
-UDMMaterialStageInputValue* UDMMaterialStageInputValue::ChangeStageInput_NewValue(UDMMaterialStage* InStage, int32 InInputIdx, int32 InInputChannel, EDMValueType InType, int32 InOutputChannel)
+UDMMaterialStageInputValue* UDMMaterialStageInputValue::ChangeStageInput_NewValue(UDMMaterialStage* InStage, int32 InInputIdx, 
+	int32 InInputChannel, EDMValueType InValueType, int32 InOutputChannel)
+{
+	return ChangeStageInput_NewValue(
+		InStage,
+		InInputIdx,
+		InInputChannel,
+		UDMValueDefinitionLibrary::GetValueDefinition(InValueType).GetValueClass(),
+		InOutputChannel
+	);
+}
+
+UDMMaterialStageInputValue* UDMMaterialStageInputValue::ChangeStageInput_NewValue(UDMMaterialStage* InStage, int32 InInputIdx, 
+	int32 InInputChannel, TSubclassOf<UDMMaterialValue> InValueClass, int32 InOutputChannel)
 {
 	check(InStage);
 
@@ -244,7 +286,7 @@ UDMMaterialStageInputValue* UDMMaterialStageInputValue::ChangeStageInput_NewValu
 
 	const TArray<FDMMaterialStageConnector>& InputConnectors = Throughput->GetInputConnectors();
 	check(InputConnectors.IsValidIndex(InInputIdx));
-	check(Throughput->CanInputAcceptType(InInputIdx, InType));
+	check(Throughput->CanInputAcceptType(InInputIdx, InValueClass.GetDefaultObject()->GetType()));
 
 	UDMMaterialLayerObject* Layer = InStage->GetLayer();
 	check(Layer);
@@ -258,7 +300,7 @@ UDMMaterialStageInputValue* UDMMaterialStageInputValue::ChangeStageInput_NewValu
 	UDynamicMaterialModel* MaterialModel = ModelEditorOnlyData->GetMaterialModel();
 	check(MaterialModel);
 
-	UDMMaterialValue* NewValue = MaterialModel->AddValue(InType);
+	UDMMaterialValue* NewValue = MaterialModel->AddValue(InValueClass);
 	check(NewValue);
 
 	UDMMaterialStageInputValue* NewInputValue = InStage->ChangeInput<UDMMaterialStageInputValue>(
diff --git a/Engine/Plugins/Experimental/DynamicMaterial/Source/DynamicMaterialEditor/Private/DMBlueprintFunctionLibrary.cpp b/Engine/Plugins/Experimental/DynamicMaterial/Source/DynamicMaterialEditor/Private/DMBlueprintFunctionLibrary.cpp
--- a/Engine/Plugins/Experimental/DynamicMaterial/Source/DynamicMaterialEditor/Private/DMBlueprintFunctionLibrary.cpp
+++ b/Engine/Plugins/Experimental/DynamicMaterial/Source/DynamicMaterialEditor/Private/DMBlueprintFunctionLibrary.cpp
@@ -12,11 +12,13 @@
 #include "Components/DMMaterialStageThroughput.h"
 #include "Components/DMMaterialStageThroughputLayerBlend.h"
 #include "Components/DMMaterialSubStage.h"
+#include "Components/DMRenderTargetRenderer.h"
 #include "Components/MaterialStageBlends/DMMSBNormal.h"
 #include "Components/MaterialStageExpressions/DMMSETextureSample.h"
 #include "Components/MaterialStageInputs/DMMSIExpression.h"
 #include "Components/MaterialStageInputs/DMMSITextureUV.h"
 #include "Components/MaterialStageInputs/DMMSIValue.h"
+#include "Components/MaterialValues/DMMaterialValueRenderTarget.h"
 #include "Components/MaterialValues/DMMaterialValueTexture.h"
 #include "Components/PrimitiveComponent.h"
 #include "DMEDefs.h"
@@ -362,3 +364,49 @@ bool UDMBlueprintFunctionLibrary::ExportGeneratedMaterial(UDynamicMaterialModel*
 
 	return true;
 }
+
+void UDMBlueprintFunctionLibrary::SetStageInputToRenderer(UDMMaterialStage* InStage, TSubclassOf<UDMRenderTargetRenderer> InRendererClass, int32 InInputIndex)
+{
+	UDMMaterialStageInputExpression* InputExpression = UDMMaterialStageInputExpression::ChangeStageInput_Expression(
+		InStage,
+		UDMMaterialStageExpressionTextureSample::StaticClass(),
+		InInputIndex,
+		FDMMaterialStageConnectorChannel::WHOLE_CHANNEL,
+		0,
+		FDMMaterialStageConnectorChannel::THREE_CHANNELS
+	);
+
+	if (!ensure(InputExpression))
+	{
+		return;
+	}
+
+	UDMMaterialSubStage* SubStage = InputExpression->GetSubStage();
+
+	if (!ensure(SubStage))
+	{
+		return;
+	}
+
+	UDMMaterialStageInputValue* InputValue = UDMMaterialStageInputValue::ChangeStageInput_NewLocalValue(
+		SubStage,
+		0,
+		FDMMaterialStageConnectorChannel::WHOLE_CHANNEL,
+		UDMMaterialValueRenderTarget::StaticClass(),
+		FDMMaterialStageConnectorChannel::WHOLE_CHANNEL
+	);
+
+	if (!ensure(InputValue))
+	{
+		return;
+	}
+
+	UDMMaterialValueRenderTarget* RenderTargetValue = Cast<UDMMaterialValueRenderTarget>(InputValue->GetValue());
+
+	if (!RenderTargetValue)
+	{
+		return;
+	}
+
+	UDMRenderTargetRenderer::CreateRenderTargetRenderer(InRendererClass, RenderTargetValue);
+}
diff --git a/Engine/Plugins/Experimental/DynamicMaterial/Source/DynamicMaterialEditor/Private/DMBlueprintFunctionLibrary.h b/Engine/Plugins/Experimental/DynamicMaterial/Source/DynamicMaterialEditor/Private/DMBlueprintFunctionLibrary.h
--- a/Engine/Plugins/Experimental/DynamicMaterial/Source/DynamicMaterialEditor/Private/DMBlueprintFunctionLibrary.h
+++ b/Engine/Plugins/Experimental/DynamicMaterial/Source/DynamicMaterialEditor/Private/DMBlueprintFunctionLibrary.h
@@ -10,6 +10,7 @@
 class AActor;
 class UDMMaterialStage;
 class UDMMaterialStageInputValue;
+class UDMRenderTargetRenderer;
 class UDynamicMaterialModel;
 class UTexture;
 struct FDMObjectMaterialProperty;
@@ -34,4 +35,6 @@ public:
 	static bool ExportMaterialInstance(UDynamicMaterialModel* InMaterialModel, const FString& InSavePath);
 
 	static bool ExportGeneratedMaterial(UDynamicMaterialModel* InMaterialModel, const FString& InSavePath);
+
+	static void SetStageInputToRenderer(UDMMaterialStage* InStage, TSubclassOf<UDMRenderTargetRenderer> InRendererClass, int32 InInputIndex);
 };
diff --git a/Engine/Plugins/Experimental/DynamicMaterial/Source/DynamicMaterialEditor/Private/Menus/DMMaterialSlotLayerMenus.cpp b/Engine/Plugins/Experimental/DynamicMaterial/Source/DynamicMaterialEditor/Private/Menus/DMMaterialSlotLayerMenus.cpp
--- a/Engine/Plugins/Experimental/DynamicMaterial/Source/DynamicMaterialEditor/Private/Menus/DMMaterialSlotLayerMenus.cpp
+++ b/Engine/Plugins/Experimental/DynamicMaterial/Source/DynamicMaterialEditor/Private/Menus/DMMaterialSlotLayerMenus.cpp
@@ -7,9 +7,9 @@
 #include "Components/DMMaterialStageExpression.h"
 #include "Components/DMMaterialStageGradient.h"
 #include "Components/DMMaterialValue.h"
-#include "Components/MaterialStageExpressions/DMMSETextRenderer.h"
 #include "Components/MaterialStageExpressions/DMMSETextureSample.h"
 #include "Components/MaterialStageExpressions/DMMSETextureSampleEdgeColor.h"
+#include "Components/RenderTargetRenderers/DMRenderTargetTextRenderer.h"
 #include "DMDefs.h"
 #include "DMMaterialSlotLayerAddEffectMenus.h"
 #include "DMValueDefinition.h"
@@ -545,9 +545,8 @@ void FDMMaterialSlotLayerMenus::AddAddLayerSection(UToolMenu* InMenu)
 		FSlateIcon(),
 		FUIAction(FExecuteAction::CreateSP(
 			SlotWidget.Get(),
-			&SDMSlot::AddNewLayer_Expression,
-			TSubclassOf<UDMMaterialStageExpression>(UDMMaterialStageExpressionTextRenderer::StaticClass()),
-			EDMMaterialLayerStage::Base
+			&SDMSlot::AddNewLayer_Renderer,
+			TSubclassOf<UDMRenderTargetRenderer>(UDMRenderTargetTextRenderer::StaticClass())
 		))
 	);
 
diff --git a/Engine/Plugins/Experimental/DynamicMaterial/Source/DynamicMaterialEditor/Private/Menus/DMMaterialStageSourceMenus.cpp b/Engine/Plugins/Experimental/DynamicMaterial/Source/DynamicMaterialEditor/Private/Menus/DMMaterialStageSourceMenus.cpp
--- a/Engine/Plugins/Experimental/DynamicMaterial/Source/DynamicMaterialEditor/Private/Menus/DMMaterialStageSourceMenus.cpp
+++ b/Engine/Plugins/Experimental/DynamicMaterial/Source/DynamicMaterialEditor/Private/Menus/DMMaterialStageSourceMenus.cpp
@@ -1,6 +1,7 @@
 // Copyright Epic Games, Inc. All Rights Reserved.
 
 #include "DMMaterialStageSourceMenus.h"
+
 #include "Components/DMMaterialLayer.h"
 #include "Components/DMMaterialProperty.h"
 #include "Components/DMMaterialSlot.h"
@@ -14,7 +15,6 @@
 #include "Components/DMMaterialStageThroughputLayerBlend.h"
 #include "Components/DMMaterialValue.h"
 #include "Components/MaterialStageExpressions/DMMSESceneTexture.h"
-#include "Components/MaterialStageExpressions/DMMSETextRenderer.h"
 #include "Components/MaterialStageExpressions/DMMSETextureSample.h"
 #include "Components/MaterialStageExpressions/DMMSETextureSampleEdgeColor.h"
 #include "Components/MaterialStageInputs/DMMSIExpression.h"
@@ -22,6 +22,8 @@
 #include "Components/MaterialStageInputs/DMMSIGradient.h"
 #include "Components/MaterialStageInputs/DMMSISlot.h"
 #include "Components/MaterialStageInputs/DMMSIValue.h"
+#include "Components/RenderTargetRenderers/DMRenderTargetTextRenderer.h"
+#include "DMBlueprintFunctionLibrary.h"
 #include "DMDefs.h"
 #include "DMValueDefinition.h"
 #include "DynamicMaterialEditorModule.h"
@@ -127,9 +129,7 @@ namespace UE::DynamicMaterialEditor::Private
 							}
 							else
 							{
-								FScopedTransaction Transaction(LOCTEXT("SetStageSource", "Set Material Designer Stage Source"));
-								Stage->Modify();
-								UDMMaterialStageInputValue::ChangeStageSource_NewLocalValue(Stage, ValueType);
+								ensure(false);
 							}
 
 							if (TSharedPtr<SDMSlot> SlotWidget = MenuContext->GetSlotWidget().Pin())
@@ -231,9 +231,7 @@ namespace UE::DynamicMaterialEditor::Private
 						}
 						else
 						{
-							FScopedTransaction Transaction(LOCTEXT("SetStageSource", "Set Material Designer Stage Source"));
-							Stage->Modify();
-							UDMMaterialStageInputValue::ChangeStageSource_Value(Stage, Value);
+							ensure(false);
 						}
 
 						if (TSharedPtr<SDMSlot> SlotWidget = MenuContext->GetSlotWidget().Pin())
@@ -329,9 +327,7 @@ namespace UE::DynamicMaterialEditor::Private
 							}
 							else
 							{
-								FScopedTransaction Transaction(LOCTEXT("SetStageSource", "Set Material Designer Stage Source"));
-								Stage->Modify();
-								UDMMaterialStageInputValue::ChangeStageSource_NewValue(Stage, ValueType);
+								ensure(false);
 							}
 
 							if (TSharedPtr<SDMSlot> SlotWidget = MenuContext->GetSlotWidget().Pin())
@@ -456,9 +452,7 @@ namespace UE::DynamicMaterialEditor::Private
 							}
 							else
 							{
-								FScopedTransaction Transaction(LOCTEXT("SetStageSource", "Set Material Designer Stage Source"));
-								Stage->Modify();
-								UDMMaterialStageInputSlot::ChangeStageSource_Slot(Stage, Slot, Property);
+								ensure(false);
 							}
 
 							if (TSharedPtr<SDMSlot> SlotWidget = MenuContext->GetSlotWidget().Pin())
@@ -604,9 +598,7 @@ namespace UE::DynamicMaterialEditor::Private
 								}
 								else
 								{
-									FScopedTransaction Transaction(LOCTEXT("SetStageSource", "Set Material Designer Stage Source"));
-									Stage->Modify();
-									UDMMaterialStageInputSlot::ChangeStageSource_Slot(Stage, SlotIter, SlotProperty);
+									ensure(false);
 								}
 
 								if (TSharedPtr<SDMSlot> SlotWidget = MenuContext->GetSlotWidget().Pin())
@@ -709,9 +701,7 @@ namespace UE::DynamicMaterialEditor::Private
 							}
 							else
 							{
-								FScopedTransaction Transaction(LOCTEXT("SetStageSource", "Set Material Designer Stage Source"));
-								Stage->Modify();
-								UDMMaterialStageInputGradient::ChangeStageSource_Gradient(Stage, GradientClass.Get());
+								ensure(false);
 							}
 
 							if (TSharedPtr<SDMSlot> SlotWidget = MenuContext->GetSlotWidget().Pin())
@@ -776,13 +766,7 @@ namespace UE::DynamicMaterialEditor::Private
 		}
 		else
 		{
-			FScopedTransaction Transaction(LOCTEXT("SetStageSource", "Set Material Designer Stage Source"));
-			Stage->Modify();
-
-			UDMMaterialStageInputExpression::ChangeStageSource_Expression(
-				Stage, 
-				UDMMaterialStageExpressionTextureSample::StaticClass()
-			);
+			ensure(false);
 		}
 
 		if (TSharedPtr<SDMSlot> SlotWidget = InMenuContext->GetSlotWidget().Pin())
@@ -840,9 +824,7 @@ namespace UE::DynamicMaterialEditor::Private
 		}
 		else
 		{
-			FScopedTransaction Transaction(LOCTEXT("SetStageSource", "Set Material Designer Stage Source"));
-			Stage->Modify();
-			UDMMaterialStageInputValue::ChangeStageSource_NewLocalValue(Stage, EDMValueType::VT_Float3_RGB);
+			ensure(false);
 		}
 
 		if (TSharedPtr<SDMSlot> SlotWidget = InMenuContext->GetSlotWidget().Pin())
@@ -900,9 +882,7 @@ namespace UE::DynamicMaterialEditor::Private
 		}
 		else
 		{
-			FScopedTransaction Transaction(LOCTEXT("SetStageSource", "Set Material Designer Stage Source"));
-			Stage->Modify();
-			UDMMaterialStageInputValue::ChangeStageSource_NewLocalValue(Stage, EDMValueType::VT_ColorAtlas);
+			ensure(false);
 		}
 
 		if (TSharedPtr<SDMSlot> SlotWidget = InMenuContext->GetSlotWidget().Pin())
@@ -969,13 +949,7 @@ namespace UE::DynamicMaterialEditor::Private
 		}
 		else
 		{
-			FScopedTransaction Transaction(LOCTEXT("SetStageInput", "Set Material Designer Source"));
-			Stage->Modify();
-
-			UDMMaterialStageInputExpression::ChangeStageSource_Expression(
-				Stage,
-				UDMMaterialStageExpressionTextureSampleEdgeColor::StaticClass()
-			);
+			ensure(false);
 		}
 
 		if (UDMMaterialStage* MaskStage = Layer->GetStage(EDMMaterialLayerStage::Mask))
@@ -1064,13 +1038,7 @@ namespace UE::DynamicMaterialEditor::Private
 		}
 		else
 		{
-			FScopedTransaction Transaction(LOCTEXT("SetStageInput", "Set Material Designer Source"));
-			Stage->Modify();
-
-			UDMMaterialStageInputExpression::ChangeStageSource_Expression(
-				Stage,
-				UDMMaterialStageExpressionSceneTexture::StaticClass()
-			);
+			ensure(false);
 		}
 
 		if (TSharedPtr<SDMSlot> SlotWidget = InMenuContext->GetSlotWidget().Pin())
@@ -1155,9 +1123,7 @@ namespace UE::DynamicMaterialEditor::Private
 		}
 		else
 		{
-			FScopedTransaction Transaction(LOCTEXT("SetStageSource", "Set Material Designer Stage Source"));
-			Stage->Modify();
-			UDMMaterialStageInputFunction::ChangeStageSource_Function(Stage, UDMMaterialStageFunction::GetNoOpFunction());
+			ensure(false);
 		}
 
 		if (TSharedPtr<SDMSlot> SlotWidget = InMenuContext->GetSlotWidget().Pin())
@@ -1199,38 +1165,18 @@ namespace UE::DynamicMaterialEditor::Private
 			FScopedTransaction Transaction(LOCTEXT("SetStageInputBase", "Set Material Designer Base Source"));
 			Stage->Modify();
 
-			UDMMaterialStageInputExpression::ChangeStageInput_Expression(
-				Stage,
-				UDMMaterialStageExpressionTextRenderer::StaticClass(),
-				UDMMaterialStageBlend::InputB,
-				FDMMaterialStageConnectorChannel::WHOLE_CHANNEL,
-				0,
-				FDMMaterialStageConnectorChannel::THREE_CHANNELS
-			);
+			UDMBlueprintFunctionLibrary::SetStageInputToRenderer(Stage, UDMRenderTargetTextRenderer::StaticClass(), UDMMaterialStageBlend::InputB);
 		}
 		else if (StageSource->IsA<UDMMaterialStageThroughputLayerBlend>())
 		{
 			FScopedTransaction Transaction(LOCTEXT("SetStageInputMask", "Set Material Designer Mask Source"));
 			Stage->Modify();
 
-			UDMMaterialStageInputExpression::ChangeStageInput_Expression(
-				Stage,
-				UDMMaterialStageExpressionTextRenderer::StaticClass(),
-				UDMMaterialStageThroughputLayerBlend::InputMaskSource,
-				FDMMaterialStageConnectorChannel::WHOLE_CHANNEL,
-				0,
-				FDMMaterialStageConnectorChannel::WHOLE_CHANNEL
-			);
+			UDMBlueprintFunctionLibrary::SetStageInputToRenderer(Stage, UDMRenderTargetTextRenderer::StaticClass(), UDMMaterialStageThroughputLayerBlend::InputMaskSource);
 		}
 		else
 		{
-			FScopedTransaction Transaction(LOCTEXT("SetStageSource", "Set Material Designer Stage Source"));
-			Stage->Modify();
-
-			UDMMaterialStageInputExpression::ChangeStageSource_Expression(
-				Stage,
-				UDMMaterialStageExpressionTextRenderer::StaticClass()
-			);
+			ensure(false);
 		}
 
 		if (TSharedPtr<SDMSlot> SlotWidget = InMenuContext->GetSlotWidget().Pin())
diff --git a/Engine/Plugins/Experimental/DynamicMaterial/Source/DynamicMaterialEditor/Private/Slate/SDMSlot.cpp b/Engine/Plugins/Experimental/DynamicMaterial/Source/DynamicMaterialEditor/Private/Slate/SDMSlot.cpp
--- a/Engine/Plugins/Experimental/DynamicMaterial/Source/DynamicMaterialEditor/Private/Slate/SDMSlot.cpp
+++ b/Engine/Plugins/Experimental/DynamicMaterial/Source/DynamicMaterialEditor/Private/Slate/SDMSlot.cpp
@@ -9,7 +9,9 @@
 #include "Components/DMMaterialStageFunction.h"
 #include "Components/DMMaterialStageGradient.h"
 #include "Components/DMMaterialStageThroughputLayerBlend.h"
+#include "Components/DMMaterialSubStage.h"
 #include "Components/DMMaterialValue.h"
+#include "Components/DMRenderTargetRenderer.h"
 #include "Components/MaterialStageBlends/DMMSBNormal.h"
 #include "Components/MaterialStageExpressions/DMMSESceneTexture.h"
 #include "Components/MaterialStageExpressions/DMMSETextureSample.h"
@@ -20,6 +22,7 @@
 #include "Components/MaterialStageInputs/DMMSITextureUV.h"
 #include "Components/MaterialStageInputs/DMMSIValue.h"
 #include "Components/MaterialValues/DMMaterialValueFloat1.h"
+#include "Components/MaterialValues/DMMaterialValueRenderTarget.h"
 #include "DetailLayoutBuilder.h"
 #include "DMBlueprintFunctionLibrary.h"
 #include "DMPrivate.h"
@@ -892,6 +895,13 @@ void SDMSlot::InvalidateComponentEditWidget()
 	bInvalidateComponentEditWidget = true;
 }
 
+void SDMSlot::AddNewLayer_NewLocalValue(EDMValueType InValueType)
+{
+	AddNewLayer_NewLocalValue(
+		UDMValueDefinitionLibrary::GetValueDefinition(InValueType).GetValueClass()
+	);
+}
+
 void SDMSlot::RefreshHeaderPropertyListWidget()
 {
 	bInvalidateHeaderWidget = false;
@@ -1305,7 +1315,7 @@ UDMMaterialLayerObject* SDMSlot::AddNewLayer(UDMMaterialStage* InNewBaseStage, U
 	return Layer;
 }
 
-void SDMSlot::AddNewLayer_NewLocalValue(EDMValueType InType)
+void SDMSlot::AddNewLayer_NewLocalValue(TSubclassOf<UDMMaterialValue> InValueClass)
 {
 	UDMMaterialSlot* Slot = SlotWeak.Get();
 
@@ -1321,7 +1331,7 @@ void SDMSlot::AddNewLayer_NewLocalValue(EDMValueType InType)
 		NewBase, 
 		UDMMaterialStageBlend::InputB,
 		FDMMaterialStageConnectorChannel::WHOLE_CHANNEL, 
-		InType,
+		InValueClass,
 		FDMMaterialStageConnectorChannel::WHOLE_CHANNEL
 	);
 }
@@ -1358,7 +1368,14 @@ void SDMSlot::AddNewLayer_GlobalValue(UDMMaterialValue* InValue)
 	);
 }
 
-void SDMSlot::AddNewLayer_NewGlobalValue(EDMValueType InType)
+void SDMSlot::AddNewLayer_NewGlobalValue(EDMValueType InValueType)
+{
+	AddNewLayer_NewGlobalValue(
+		UDMValueDefinitionLibrary::GetValueDefinition(InValueType).GetValueClass()
+	);
+}
+
+void SDMSlot::AddNewLayer_NewGlobalValue(TSubclassOf<UDMMaterialValue> InValueClass)
 {
 	UDMMaterialSlot* Slot = SlotWeak.Get();
 
@@ -1374,7 +1391,7 @@ void SDMSlot::AddNewLayer_NewGlobalValue(EDMValueType InType)
 		NewStage, 
 		UDMMaterialStageBlend::InputB,
 		FDMMaterialStageConnectorChannel::WHOLE_CHANNEL,
-		InType, 
+		InValueClass,
 		FDMMaterialStageConnectorChannel::WHOLE_CHANNEL
 	);
 }
@@ -1589,4 +1606,19 @@ void SDMSlot::AddNewLayer_SceneTexture()
 	}
 }
 
+void SDMSlot::AddNewLayer_Renderer(TSubclassOf<UDMRenderTargetRenderer> InRendererClass)
+{
+	UDMMaterialSlot* Slot = SlotWeak.Get();
+
+	if (!ensure(Slot))
+	{
+		return;
+	}
+
+	UDMMaterialStage* NewBase = UDMMaterialStageBlend::CreateStage(UDMMaterialStageBlendNormal::StaticClass());
+	AddNewLayer(NewBase);
+
+	UDMBlueprintFunctionLibrary::SetStageInputToRenderer(NewBase, InRendererClass, UDMMaterialStageBlend::InputB);
+}
+
 #undef LOCTEXT_NAMESPACE
diff --git a/Engine/Plugins/Experimental/DynamicMaterial/Source/DynamicMaterialEditor/Private/Slate/SDMSlot.h b/Engine/Plugins/Experimental/DynamicMaterial/Source/DynamicMaterialEditor/Private/Slate/SDMSlot.h
--- a/Engine/Plugins/Experimental/DynamicMaterial/Source/DynamicMaterialEditor/Private/Slate/SDMSlot.h
+++ b/Engine/Plugins/Experimental/DynamicMaterial/Source/DynamicMaterialEditor/Private/Slate/SDMSlot.h
@@ -23,6 +23,7 @@ class UDMMaterialStageGradient;
 class UDMMaterialStageInputValue;
 class UDMMaterialValue;
 class UDMMaterialValueFloat1;
+class UDMRenderTargetRenderer;
 enum class EDMMaterialLayerStage : uint8;
 struct FDMScopedUITransaction;
 struct FPropertyChangedEvent;
@@ -61,9 +62,11 @@ public:
 	void InvalidateSlotSettingsRowWidget();
 	void InvalidateComponentEditWidget();
 
-	void AddNewLayer_NewLocalValue(EDMValueType InType);
+	void AddNewLayer_NewLocalValue(EDMValueType InValueType);
+	void AddNewLayer_NewLocalValue(TSubclassOf<UDMMaterialValue> InValueClass);
 	void AddNewLayer_GlobalValue(UDMMaterialValue* InValue);
-	void AddNewLayer_NewGlobalValue(EDMValueType InType);
+	void AddNewLayer_NewGlobalValue(EDMValueType InValueType);
+	void AddNewLayer_NewGlobalValue(TSubclassOf<UDMMaterialValue> InValueClass);
 	void AddNewLayer_Slot(UDMMaterialSlot* InSlot, EDMMaterialPropertyType InMaterialProperty);
 	void AddNewLayer_Expression(TSubclassOf<UDMMaterialStageExpression> InExpressionClass, EDMMaterialLayerStage LayerEnabledMask);
 	void AddNewLayer_Blend(TSubclassOf<UDMMaterialStageBlend> InBlendClass);
@@ -71,6 +74,7 @@ public:
 	void AddNewLayer_UV();
 	void AddNewLayer_MaterialFunction();
 	void AddNewLayer_SceneTexture();
+	void AddNewLayer_Renderer(TSubclassOf<UDMRenderTargetRenderer> InRendererClass);
 
 	void RemoveSlot();
 
diff --git a/Engine/Plugins/Experimental/DynamicMaterial/Source/DynamicMaterialEditor/Public/Components/DMMaterialEffect.h b/Engine/Plugins/Experimental/DynamicMaterial/Source/DynamicMaterialEditor/Public/Components/DMMaterialEffect.h
--- a/Engine/Plugins/Experimental/DynamicMaterial/Source/DynamicMaterialEditor/Public/Components/DMMaterialEffect.h
+++ b/Engine/Plugins/Experimental/DynamicMaterial/Source/DynamicMaterialEditor/Public/Components/DMMaterialEffect.h
@@ -24,7 +24,7 @@ enum class EDMMaterialEffectTarget : uint8
 	Slot      = 1 << 3
 };
 
-UCLASS(Abstract, BlueprintType, ClassGroup = "Material Designer", Meta = (DisplayName = "Material Designer Effect"))
+UCLASS(Abstract, BlueprintType, Blueprintable, ClassGroup = "Material Designer", meta = (DisplayName = "Material Designer Effect"))
 class DYNAMICMATERIALEDITOR_API UDMMaterialEffect : public UDMMaterialComponent
 {
 	GENERATED_BODY()
diff --git a/Engine/Plugins/Experimental/DynamicMaterial/Source/DynamicMaterialEditor/Public/Components/DMMaterialStageBlend.h b/Engine/Plugins/Experimental/DynamicMaterial/Source/DynamicMaterialEditor/Public/Components/DMMaterialStageBlend.h
--- a/Engine/Plugins/Experimental/DynamicMaterial/Source/DynamicMaterialEditor/Public/Components/DMMaterialStageBlend.h
+++ b/Engine/Plugins/Experimental/DynamicMaterial/Source/DynamicMaterialEditor/Public/Components/DMMaterialStageBlend.h
@@ -16,7 +16,7 @@ struct FDMMaterialBuildState;
 /**
  * A node which represents a blend operation.
  */
-UCLASS(Abstract, BlueprintType, ClassGroup = "Material Designer", meta = (DisplayName = "Material Designer Stage Blend"))
+UCLASS(Abstract, BlueprintType, Blueprintable, ClassGroup = "Material Designer", meta = (DisplayName = "Material Designer Stage Blend"))
 class DYNAMICMATERIALEDITOR_API UDMMaterialStageBlend : public UDMMaterialStageThroughput
 {
 	GENERATED_BODY()
diff --git a/Engine/Plugins/Experimental/DynamicMaterial/Source/DynamicMaterialEditor/Public/Components/DMMaterialStageBlendFunction.h b/Engine/Plugins/Experimental/DynamicMaterial/Source/DynamicMaterialEditor/Public/Components/DMMaterialStageBlendFunction.h
--- a/Engine/Plugins/Experimental/DynamicMaterial/Source/DynamicMaterialEditor/Public/Components/DMMaterialStageBlendFunction.h
+++ b/Engine/Plugins/Experimental/DynamicMaterial/Source/DynamicMaterialEditor/Public/Components/DMMaterialStageBlendFunction.h
@@ -9,7 +9,7 @@ class UMaterialExpression;
 class UMaterialFunctionInterface;
 struct FDMMaterialBuildState;
 
-UCLASS(Abstract, BlueprintType, ClassGroup = "Material Designer")
+UCLASS(Abstract, BlueprintType, Blueprintable, ClassGroup = "Material Designer")
 class DYNAMICMATERIALEDITOR_API UDMMaterialStageBlendFunction : public UDMMaterialStageBlend
 {
 	GENERATED_BODY()
diff --git a/Engine/Plugins/Experimental/DynamicMaterial/Source/DynamicMaterialEditor/Public/Components/DMMaterialStageExpression.h b/Engine/Plugins/Experimental/DynamicMaterial/Source/DynamicMaterialEditor/Public/Components/DMMaterialStageExpression.h
--- a/Engine/Plugins/Experimental/DynamicMaterial/Source/DynamicMaterialEditor/Public/Components/DMMaterialStageExpression.h
+++ b/Engine/Plugins/Experimental/DynamicMaterial/Source/DynamicMaterialEditor/Public/Components/DMMaterialStageExpression.h
@@ -32,7 +32,7 @@ enum class EDMExpressionMenu : uint8
 /**
  * A node which directly represents an material expression (or function).
  */
-UCLASS(Abstract, BlueprintType, ClassGroup = "Material Designer", meta = (DisplayName = "Material Designer Stage Expression"))
+UCLASS(Abstract, BlueprintType, Blueprintable, ClassGroup = "Material Designer", meta = (DisplayName = "Material Designer Stage Expression"))
 class DYNAMICMATERIALEDITOR_API UDMMaterialStageExpression : public UDMMaterialStageThroughput
 {
 	GENERATED_BODY()
diff --git a/Engine/Plugins/Experimental/DynamicMaterial/Source/DynamicMaterialEditor/Public/Components/DMMaterialStageGradient.h b/Engine/Plugins/Experimental/DynamicMaterial/Source/DynamicMaterialEditor/Public/Components/DMMaterialStageGradient.h
--- a/Engine/Plugins/Experimental/DynamicMaterial/Source/DynamicMaterialEditor/Public/Components/DMMaterialStageGradient.h
+++ b/Engine/Plugins/Experimental/DynamicMaterial/Source/DynamicMaterialEditor/Public/Components/DMMaterialStageGradient.h
@@ -14,7 +14,7 @@ struct FDMMaterialBuildState;
 /**
  * A node which represents UV-based gradient.
  */
-UCLASS(Abstract, BlueprintType, ClassGroup = "Material Designer", meta = (DisplayName = "Material Designer Stage Gradient"))
+UCLASS(Abstract, BlueprintType, Blueprintable, ClassGroup = "Material Designer", meta = (DisplayName = "Material Designer Stage Gradient"))
 class DYNAMICMATERIALEDITOR_API UDMMaterialStageGradient : public UDMMaterialStageThroughput
 {
 	GENERATED_BODY()
diff --git a/Engine/Plugins/Experimental/DynamicMaterial/Source/DynamicMaterialEditor/Public/Components/MaterialStageExpressions/DMMSETextRenderer.h b/Engine/Plugins/Experimental/DynamicMaterial/Source/DynamicMaterialEditor/Public/Components/MaterialStageExpressions/DMMSETextRenderer.h
deleted file mode 100644
diff --git a/Engine/Plugins/Experimental/DynamicMaterial/Source/DynamicMaterialEditor/Public/Components/MaterialStageExpressions/DMMSETextureSampleBase.h b/Engine/Plugins/Experimental/DynamicMaterial/Source/DynamicMaterialEditor/Public/Components/MaterialStageExpressions/DMMSETextureSampleBase.h
--- a/Engine/Plugins/Experimental/DynamicMaterial/Source/DynamicMaterialEditor/Public/Components/MaterialStageExpressions/DMMSETextureSampleBase.h
+++ b/Engine/Plugins/Experimental/DynamicMaterial/Source/DynamicMaterialEditor/Public/Components/MaterialStageExpressions/DMMSETextureSampleBase.h
@@ -5,7 +5,7 @@
 #include "Components/DMMaterialStageExpression.h"
 #include "DMMSETextureSampleBase.generated.h"
 
-UCLASS(BlueprintType, ClassGroup = "Material Designer")
+UCLASS(BlueprintType, Blueprintable, ClassGroup = "Material Designer")
 class DYNAMICMATERIALEDITOR_API UDMMaterialStageExpressionTextureSampleBase : public UDMMaterialStageExpression
 {
 	GENERATED_BODY()
@@ -34,8 +34,12 @@ public:
 	UFUNCTION(BlueprintCallable, Category = "Material Designer")
 	virtual void SetClampTextureEnabled(bool bInValue);
 
+	//~ Begin UObject
+	virtual void PostEditChangeProperty(FPropertyChangedEvent& InPropertyChangedEvent) override;
+	//~ End UObject
+
 protected:
-	UPROPERTY(EditInstanceOnly, BlueprintReadOnly, Setter=SetClampTextureEnabled, Category = "Material Designer", 
+	UPROPERTY(EditInstanceOnly, BlueprintReadWrite, Setter = SetClampTextureEnabled, BlueprintSetter = SetClampTextureEnabled, Category = "Material Designer",
 		meta=(ToolTip="Forces a material rebuild.", LowPriority, NotKeyframeable, AllowPrivateAccess = "true"))
 	bool bClampTexture;
 
diff --git a/Engine/Plugins/Experimental/DynamicMaterial/Source/DynamicMaterialEditor/Public/Components/MaterialStageGradients/DMMSGLinear.h b/Engine/Plugins/Experimental/DynamicMaterial/Source/DynamicMaterialEditor/Public/Components/MaterialStageGradients/DMMSGLinear.h
--- a/Engine/Plugins/Experimental/DynamicMaterial/Source/DynamicMaterialEditor/Public/Components/MaterialStageGradients/DMMSGLinear.h
+++ b/Engine/Plugins/Experimental/DynamicMaterial/Source/DynamicMaterialEditor/Public/Components/MaterialStageGradients/DMMSGLinear.h
@@ -33,6 +33,7 @@ public:
 	virtual void SetTilingType(ELinearGradientTileType InType);
 
 protected:
-	UPROPERTY(EditInstanceOnly, BlueprintReadOnly, Getter=GetTilingType, Setter=SetTilingType, Category = "Material Designer")
+	UPROPERTY(EditInstanceOnly, BlueprintReadWrite, Getter = GetTilingType, Setter=SetTilingType, BlueprintSetter = SetTilingType, 
+		Category = "Material Designer")
 	ELinearGradientTileType Tiling;
 };
diff --git a/Engine/Plugins/Experimental/DynamicMaterial/Source/DynamicMaterialEditor/Public/Components/MaterialStageInputs/DMMSIValue.h b/Engine/Plugins/Experimental/DynamicMaterial/Source/DynamicMaterialEditor/Public/Components/MaterialStageInputs/DMMSIValue.h
--- a/Engine/Plugins/Experimental/DynamicMaterial/Source/DynamicMaterialEditor/Public/Components/MaterialStageInputs/DMMSIValue.h
+++ b/Engine/Plugins/Experimental/DynamicMaterial/Source/DynamicMaterialEditor/Public/Components/MaterialStageInputs/DMMSIValue.h
@@ -22,26 +22,36 @@ public:
 
 	static UDMMaterialStage* CreateStage(UDMMaterialValue* InValue, UDMMaterialLayerObject* InLayer = nullptr);
 
+	static UDMMaterialStageInputValue* ChangeStageSource_NewLocalValue(UDMMaterialStage* InStage, EDMValueType InValueType);
+
 	UFUNCTION(BlueprintCallable, Category = "Material Designer")
-	static UDMMaterialStageInputValue* ChangeStageSource_NewLocalValue(UDMMaterialStage* InStage, EDMValueType InType);
+	static UDMMaterialStageInputValue* ChangeStageSource_NewLocalValue(UDMMaterialStage* InStage, TSubclassOf<UDMMaterialValue> InValueClass);
 
 	UFUNCTION(BlueprintCallable, Category = "Material Designer")
 	static UDMMaterialStageInputValue* ChangeStageSource_Value(UDMMaterialStage* InStage, UDMMaterialValue* InValue);
 
+	static UDMMaterialStageInputValue* ChangeStageSource_NewValue(UDMMaterialStage* InStage, EDMValueType InValueType);
+
 	UFUNCTION(BlueprintCallable, Category = "Material Designer")
-	static UDMMaterialStageInputValue* ChangeStageSource_NewValue(UDMMaterialStage* InStage, EDMValueType InType);
+	static UDMMaterialStageInputValue* ChangeStageSource_NewValue(UDMMaterialStage* InStage, TSubclassOf<UDMMaterialValue> InValueClass);
+
+	static UDMMaterialStageInputValue* ChangeStageInput_NewLocalValue(UDMMaterialStage* InStage, int32 InInputIdx, int32 InInputChannel,
+		EDMValueType InValueType, int32 InOutputChannel);
 
 	UFUNCTION(BlueprintCallable, Category = "Material Designer")
 	static UDMMaterialStageInputValue* ChangeStageInput_NewLocalValue(UDMMaterialStage* InStage, int32 InInputIdx, int32 InInputChannel,
-		EDMValueType InType, int32 InOutputChannel);
+		TSubclassOf<UDMMaterialValue> InValueClass, int32 InOutputChannel);
 
 	UFUNCTION(BlueprintCallable, Category = "Material Designer")
 	static UDMMaterialStageInputValue* ChangeStageInput_Value(UDMMaterialStage* InStage, int32 InInputIdx, int32 InInputChannel,
 		UDMMaterialValue* InValue, int32 InOutputChannel);
 
+	static UDMMaterialStageInputValue* ChangeStageInput_NewValue(UDMMaterialStage* InStage, int32 InInputIdx, int32 InInputChannel,
+		EDMValueType InValueType, int32 InOutputChannel);
+
 	UFUNCTION(BlueprintCallable, Category = "Material Designer")
 	static UDMMaterialStageInputValue* ChangeStageInput_NewValue(UDMMaterialStage* InStage, int32 InInputIdx, int32 InInputChannel,
-		EDMValueType InType, int32 InOutputChannel);
+		TSubclassOf<UDMMaterialValue> InValueClass, int32 InOutputChannel);
 
 	virtual FText GetComponentDescription() const override;
 	virtual FText GetChannelDescription(const FDMMaterialStageConnectorChannel& Channel) override;
