Material Designer
- Added more default slot textures.
- Made the default slot texture system a bit more flexible

#virtualized

diff --git a/Engine/Plugins/Experimental/DynamicMaterial/Source/DynamicMaterialEditor/Private/Components/DMMaterialProperty.cpp b/Engine/Plugins/Experimental/DynamicMaterial/Source/DynamicMaterialEditor/Private/Components/DMMaterialProperty.cpp
--- a/Engine/Plugins/Experimental/DynamicMaterial/Source/DynamicMaterialEditor/Private/Components/DMMaterialProperty.cpp
+++ b/Engine/Plugins/Experimental/DynamicMaterial/Source/DynamicMaterialEditor/Private/Components/DMMaterialProperty.cpp
@@ -2,8 +2,16 @@
 
 #include "Components/DMMaterialProperty.h"
 #include "Components/DMMaterialSlot.h"
+#include "Components/DMMaterialStageThroughputLayerBlend.h"
+#include "Components/DMMaterialSubStage.h"
+#include "Components/MaterialStageBlends/DMMSBNormal.h"
+#include "Components/MaterialStageExpressions/DMMSETextureSample.h"
+#include "Components/MaterialStageInputs/DMMSIExpression.h"
+#include "Components/MaterialStageInputs/DMMSIValue.h"
+#include "Components/MaterialValues/DMMaterialValueTexture.h"
 #include "DMPrivate.h"
 #include "DMValueDefinition.h"
+#include "DynamicMaterialEditorSettings.h"
 #include "Materials/MaterialExpressionConstant.h"
 #include "Materials/MaterialExpressionConstant2Vector.h"
 #include "Materials/MaterialExpressionConstant3Vector.h"
@@ -99,6 +107,93 @@ UMaterialExpression* UDMMaterialProperty::GetDefaultInput(const TSharedRef<FDMMa
 	return nullptr;
 }
 
+void UDMMaterialProperty::OnSlotAdded(UDMMaterialSlot* InSlot)
+{
+	UTexture* BaseTexture = nullptr;
+
+	if (const TSoftObjectPtr<UTexture>* TexturePtr = UDynamicMaterialEditorSettings::Get()->DefaultSlotTextures.Find(MaterialProperty))
+	{
+		if (UTexture* Texture = TexturePtr->LoadSynchronous())
+		{
+			BaseTexture = Texture;
+		}
+	}
+
+	if (!BaseTexture)
+	{
+		return;
+	}
+
+	UDMMaterialStage* DefaultStage = UDMMaterialStageBlend::CreateStage(UDMMaterialStageBlendNormal::StaticClass());
+	check(DefaultStage);
+
+	UDMMaterialStage* MaskStage = UDMMaterialStageThroughputLayerBlend::CreateStage();
+	check(MaskStage);
+
+	InSlot->AddLayerWithMask(MaterialProperty, DefaultStage, MaskStage);
+
+	UDMMaterialStageInputExpression* BaseInputExpression = UDMMaterialStageInputExpression::ChangeStageInput_Expression(
+		DefaultStage,
+		UDMMaterialStageExpressionTextureSample::StaticClass(),
+		UDMMaterialStageBlendNormal::InputB,
+		FDMMaterialStageConnectorChannel::WHOLE_CHANNEL,
+		0,
+		FDMMaterialStageConnectorChannel::WHOLE_CHANNEL
+	);
+
+	if (UDMMaterialStageExpressionTextureSample* BaseInputTextureSample = Cast<UDMMaterialStageExpressionTextureSample>(BaseInputExpression->GetMaterialStageExpression()))
+	{
+		if (UDMMaterialStage* BaseTextureInputStage = BaseInputExpression->GetSubStage())
+		{
+			const TArray<UDMMaterialStageInput*> BaseTextureStageInputs = BaseTextureInputStage->GetInputs();
+
+			for (UDMMaterialStageInput* BaseTextureStageInput : BaseTextureStageInputs)
+			{
+				if (UDMMaterialStageInputValue* BaseTextureInputValue = Cast<UDMMaterialStageInputValue>(BaseTextureStageInput))
+				{
+					if (UDMMaterialValueTexture* BaseTextureValue = Cast<UDMMaterialValueTexture>(BaseTextureInputValue->GetValue()))
+					{
+						BaseTextureValue->SetDefaultValue(BaseTexture);
+						BaseTextureValue->ApplyDefaultValue();
+						break;
+					}
+				}
+			}
+		}
+	}
+
+	if (UTexture* MaskTexture = UDynamicMaterialEditorSettings::Get()->DefaultMask.LoadSynchronous())
+	{
+		const TArray<UDMMaterialStageInput*> MaskStageInputs = MaskStage->GetInputs();
+
+		for (UDMMaterialStageInput* MaskStageInput : MaskStageInputs)
+		{
+			if (UDMMaterialStageInputExpression* MaskInputExpression = Cast<UDMMaterialStageInputExpression>(MaskStageInput))
+			{
+				if (UDMMaterialStageExpressionTextureSample* MaskInputTextureSample = Cast<UDMMaterialStageExpressionTextureSample>(MaskInputExpression->GetMaterialStageExpression()))
+				{
+					if (UDMMaterialStage* MaskTextureInputStage = MaskInputExpression->GetSubStage())
+					{
+						const TArray<UDMMaterialStageInput*> MaskTextureStageInputs = MaskTextureInputStage->GetInputs();
+
+						for (UDMMaterialStageInput* MaskTextureStageInput : MaskTextureStageInputs)
+						{
+							if (UDMMaterialStageInputValue* MaskTextureInputValue = Cast<UDMMaterialStageInputValue>(MaskTextureStageInput))
+							{
+								if (UDMMaterialValueTexture* MaskTextureValue = Cast<UDMMaterialValueTexture>(MaskTextureInputValue->GetValue()))
+								{
+									MaskTextureValue->SetDefaultValue(MaskTexture);
+									MaskTextureValue->ApplyDefaultValue();
+								}
+							}
+						}
+					}
+				}
+			}
+		}
+	}
+}
+
 void UDMMaterialProperty::Update(EDMUpdateType InUpdateType)
 {
 	if (!IsComponentValid())
diff --git a/Engine/Plugins/Experimental/DynamicMaterial/Source/DynamicMaterialEditor/Private/Components/DMMaterialStageThroughputLayerBlend.cpp b/Engine/Plugins/Experimental/DynamicMaterial/Source/DynamicMaterialEditor/Private/Components/DMMaterialStageThroughputLayerBlend.cpp
--- a/Engine/Plugins/Experimental/DynamicMaterial/Source/DynamicMaterialEditor/Private/Components/DMMaterialStageThroughputLayerBlend.cpp
+++ b/Engine/Plugins/Experimental/DynamicMaterial/Source/DynamicMaterialEditor/Private/Components/DMMaterialStageThroughputLayerBlend.cpp
@@ -145,7 +145,7 @@ void UDMMaterialStageThroughputLayerBlend::AddDefaultInput(int32 InInputIndex) c
 			UDMMaterialValueTexture* InputTexture = Cast<UDMMaterialValueTexture>(InputValue->GetValue());
 			check(InputTexture);
 
-			InputTexture->SetDefaultValue(UDynamicMaterialEditorSettings::Get()->DefaultOpaqueTexture.LoadSynchronous());
+			InputTexture->SetDefaultValue(UDynamicMaterialEditorSettings::Get()->DefaultMask.LoadSynchronous());
 			InputTexture->ApplyDefaultValue();
 			break;
 		}
diff --git a/Engine/Plugins/Experimental/DynamicMaterial/Source/DynamicMaterialEditor/Private/Components/MaterialProperties/DMMPBaseColor.cpp b/Engine/Plugins/Experimental/DynamicMaterial/Source/DynamicMaterialEditor/Private/Components/MaterialProperties/DMMPBaseColor.cpp
--- a/Engine/Plugins/Experimental/DynamicMaterial/Source/DynamicMaterialEditor/Private/Components/MaterialProperties/DMMPBaseColor.cpp
+++ b/Engine/Plugins/Experimental/DynamicMaterial/Source/DynamicMaterialEditor/Private/Components/MaterialProperties/DMMPBaseColor.cpp
@@ -1,32 +1,9 @@
 // Copyright Epic Games, Inc. All Rights Reserved.
 
 #include "Components/MaterialProperties/DMMPBaseColor.h"
-
 #include "Components/DMMaterialSlot.h"
-#include "Components/DMMaterialStage.h"
-#include "Components/DMMaterialStageBlend.h"
-#include "Components/DMMaterialStageThroughputLayerBlend.h"
-#include "Components/MaterialStageBlends/DMMSBNormal.h"
-#include "Components/MaterialStageExpressions/DMMSETextureSample.h"
-#include "Components/MaterialStageInputs/DMMSIExpression.h"
 #include "Model/DynamicMaterialModelEditorOnlyData.h"
 
-void UDMMaterialPropertyBaseColor::AddDefaultRGBLayer(EDMMaterialPropertyType InMaterialProperty, UDMMaterialSlot* InSlot)
-{
-	UDMMaterialStage* DefaultStage = UDMMaterialStageBlend::CreateStage(UDMMaterialStageBlendNormal::StaticClass());
-	check(DefaultStage);
-
-	UDMMaterialStage* MaskStage = UDMMaterialStageThroughputLayerBlend::CreateStage();
-	check(MaskStage);
-
-	InSlot->AddLayerWithMask(InMaterialProperty, DefaultStage, MaskStage);
-
-	UDMMaterialStageInputExpression* BaseInputExpression = UDMMaterialStageInputExpression::ChangeStageInput_Expression(DefaultStage,
-		UDMMaterialStageExpressionTextureSample::StaticClass(), UDMMaterialStageBlendNormal::InputB,
-		FDMMaterialStageConnectorChannel::WHOLE_CHANNEL, 0,
-		FDMMaterialStageConnectorChannel::WHOLE_CHANNEL);
-}
-
 UDMMaterialPropertyBaseColor::UDMMaterialPropertyBaseColor()
 	: UDMMaterialProperty(
 		EDMMaterialPropertyType(EDMMaterialPropertyType::BaseColor),
@@ -43,8 +20,3 @@ UMaterialExpression* UDMMaterialPropertyBaseColor::GetDefaultInput(const TShared
 {
 	return CreateConstant(InBuildState, FVector::OneVector);
 }
-
-void UDMMaterialPropertyBaseColor::OnSlotAdded(UDMMaterialSlot* InSlot)
-{
-	AddDefaultRGBLayer(EDMMaterialPropertyType::BaseColor, InSlot);
-}
diff --git a/Engine/Plugins/Experimental/DynamicMaterial/Source/DynamicMaterialEditor/Private/Components/MaterialProperties/DMMPEmissiveColor.cpp b/Engine/Plugins/Experimental/DynamicMaterial/Source/DynamicMaterialEditor/Private/Components/MaterialProperties/DMMPEmissiveColor.cpp
--- a/Engine/Plugins/Experimental/DynamicMaterial/Source/DynamicMaterialEditor/Private/Components/MaterialProperties/DMMPEmissiveColor.cpp
+++ b/Engine/Plugins/Experimental/DynamicMaterial/Source/DynamicMaterialEditor/Private/Components/MaterialProperties/DMMPEmissiveColor.cpp
@@ -29,8 +29,3 @@ UMaterialExpression* UDMMaterialPropertyEmissiveColor::GetDefaultInput(
 {
 	return CreateConstant(InBuildState, FVector::OneVector);
 }
-
-void UDMMaterialPropertyEmissiveColor::OnSlotAdded(UDMMaterialSlot* InSlot)
-{
-	UDMMaterialPropertyBaseColor::AddDefaultRGBLayer(EDMMaterialPropertyType::EmissiveColor, InSlot);
-}
diff --git a/Engine/Plugins/Experimental/DynamicMaterial/Source/DynamicMaterialEditor/Private/Components/MaterialProperties/DMMPOpacity.cpp b/Engine/Plugins/Experimental/DynamicMaterial/Source/DynamicMaterialEditor/Private/Components/MaterialProperties/DMMPOpacity.cpp
--- a/Engine/Plugins/Experimental/DynamicMaterial/Source/DynamicMaterialEditor/Private/Components/MaterialProperties/DMMPOpacity.cpp
+++ b/Engine/Plugins/Experimental/DynamicMaterial/Source/DynamicMaterialEditor/Private/Components/MaterialProperties/DMMPOpacity.cpp
@@ -1,93 +1,9 @@
 // Copyright Epic Games, Inc. All Rights Reserved.
 
 #include "Components/MaterialProperties/DMMPOpacity.h"
-
 #include "Components/DMMaterialLayer.h"
-#include "Components/DMMaterialSlot.h"
-#include "Components/DMMaterialStage.h"
-#include "Components/DMMaterialStageThroughputLayerBlend.h"
-#include "Components/DMMaterialSubStage.h"
-#include "Components/MaterialStageBlends/DMMSBNormal.h"
-#include "Components/MaterialStageExpressions/DMMSETextureSample.h"
-#include "Components/MaterialStageInputs/DMMSIExpression.h"
-#include "Components/MaterialStageInputs/DMMSIValue.h"
-#include "Components/MaterialValues/DMMaterialValueTexture.h"
-#include "DynamicMaterialEditorSettings.h"
 #include "Model/DynamicMaterialModelEditorOnlyData.h"
 
-void UDMMaterialPropertyOpacity::AddDefaultOpacityLayer(EDMMaterialPropertyType InMaterialProperty, UDMMaterialSlot* InSlot)
-{
-	UDMMaterialStage* DefaultStage = UDMMaterialStageBlend::CreateStage(UDMMaterialStageBlendNormal::StaticClass());
-	check(DefaultStage);
-
-	UDMMaterialStage* MaskStage = UDMMaterialStageThroughputLayerBlend::CreateStage();
-	check(MaskStage);
-
-	InSlot->AddLayerWithMask(InMaterialProperty, DefaultStage, MaskStage);
-
-	UDMMaterialStageInputExpression* BaseInputExpression = UDMMaterialStageInputExpression::ChangeStageInput_Expression(
-		DefaultStage,
-		UDMMaterialStageExpressionTextureSample::StaticClass(),
-		UDMMaterialStageBlendNormal::InputB,
-		FDMMaterialStageConnectorChannel::WHOLE_CHANNEL, 
-		0,
-		FDMMaterialStageConnectorChannel::WHOLE_CHANNEL
-	);
-
-	if (UDMMaterialStageExpressionTextureSample* BaseInputTextureSample = Cast<UDMMaterialStageExpressionTextureSample>(BaseInputExpression->GetMaterialStageExpression()))
-	{
-		BaseInputTextureSample->SetClampTextureEnabled(true);
-
-		if (UDMMaterialStage* BaseTextureInputStage = BaseInputExpression->GetSubStage())
-		{
-			const TArray<UDMMaterialStageInput*> BaseTextureStageInputs = BaseTextureInputStage->GetInputs();
-
-			for (UDMMaterialStageInput* BaseTextureStageInput : BaseTextureStageInputs)
-			{
-				if (UDMMaterialStageInputValue* BaseTextureInputValue = Cast<UDMMaterialStageInputValue>(BaseTextureStageInput))
-				{
-					if (UDMMaterialValueTexture* BaseTextureValue = Cast<UDMMaterialValueTexture>(BaseTextureInputValue->GetValue()))
-					{
-						BaseTextureValue->SetDefaultValue(UDynamicMaterialEditorSettings::Get()->DefaultOpacitySlotMask.LoadSynchronous());
-						BaseTextureValue->ApplyDefaultValue();
-						break;
-					}
-				}
-			}
-		}
-	}
-
-	const TArray<UDMMaterialStageInput*> MaskStageInputs = MaskStage->GetInputs();
-
-	for (UDMMaterialStageInput* MaskStageInput : MaskStageInputs)
-	{
-		if (UDMMaterialStageInputExpression* MaskInputExpression = Cast<UDMMaterialStageInputExpression>(MaskStageInput))
-		{
-			if (UDMMaterialStageExpressionTextureSample* MaskInputTextureSample = Cast<UDMMaterialStageExpressionTextureSample>(MaskInputExpression->GetMaterialStageExpression()))
-			{
-				MaskInputTextureSample->SetClampTextureEnabled(true);
-
-				if (UDMMaterialStage* MaskTextureInputStage = MaskInputExpression->GetSubStage())
-				{
-					const TArray<UDMMaterialStageInput*> MaskTextureStageInputs = MaskTextureInputStage->GetInputs();
-
-					for (UDMMaterialStageInput* MaskTextureStageInput : MaskTextureStageInputs)
-					{
-						if (UDMMaterialStageInputValue* MaskTextureInputValue = Cast<UDMMaterialStageInputValue>(MaskTextureStageInput))
-						{
-							if (UDMMaterialValueTexture* MaskTextureValue = Cast<UDMMaterialValueTexture>(MaskTextureInputValue->GetValue()))
-							{
-								MaskTextureValue->SetDefaultValue(UDynamicMaterialEditorSettings::Get()->DefaultOpaqueTexture.LoadSynchronous());
-								MaskTextureValue->ApplyDefaultValue();
-							}
-						}
-					}
-				}
-			}
-		}
-	}
-}
-
 UDMMaterialPropertyOpacity::UDMMaterialPropertyOpacity()
 	: UDMMaterialProperty(
 		EDMMaterialPropertyType(EDMMaterialPropertyType::Opacity),
@@ -114,8 +30,3 @@ UMaterialExpression* UDMMaterialPropertyOpacity::GetDefaultInput(const TSharedRe
 {
 	return CreateConstant(InBuildState, 1.f);
 }
-
-void UDMMaterialPropertyOpacity::OnSlotAdded(UDMMaterialSlot* InSlot)
-{
-	AddDefaultOpacityLayer(EDMMaterialPropertyType::Opacity, InSlot);
-}
diff --git a/Engine/Plugins/Experimental/DynamicMaterial/Source/DynamicMaterialEditor/Private/Components/MaterialProperties/DMMPOpacityMask.cpp b/Engine/Plugins/Experimental/DynamicMaterial/Source/DynamicMaterialEditor/Private/Components/MaterialProperties/DMMPOpacityMask.cpp
--- a/Engine/Plugins/Experimental/DynamicMaterial/Source/DynamicMaterialEditor/Private/Components/MaterialProperties/DMMPOpacityMask.cpp
+++ b/Engine/Plugins/Experimental/DynamicMaterial/Source/DynamicMaterialEditor/Private/Components/MaterialProperties/DMMPOpacityMask.cpp
@@ -22,8 +22,3 @@ UMaterialExpression* UDMMaterialPropertyOpacityMask::GetDefaultInput(
 {
 	return CreateConstant(InBuildState, 1.f);
 }
-
-void UDMMaterialPropertyOpacityMask::OnSlotAdded(UDMMaterialSlot* InSlot)
-{
-	UDMMaterialPropertyOpacity::AddDefaultOpacityLayer(EDMMaterialPropertyType::OpacityMask, InSlot);
-}
diff --git a/Engine/Plugins/Experimental/DynamicMaterial/Source/DynamicMaterialEditor/Private/DynamicMaterialEditorModule.cpp b/Engine/Plugins/Experimental/DynamicMaterial/Source/DynamicMaterialEditor/Private/DynamicMaterialEditorModule.cpp
--- a/Engine/Plugins/Experimental/DynamicMaterial/Source/DynamicMaterialEditor/Private/DynamicMaterialEditorModule.cpp
+++ b/Engine/Plugins/Experimental/DynamicMaterial/Source/DynamicMaterialEditor/Private/DynamicMaterialEditorModule.cpp
@@ -303,7 +303,15 @@ void FDynamicMaterialEditorModule::StartupModule()
 
 	UDMMaterialValueTexture::GetDefaultRGBTexture.BindLambda([]()
 		{
-			return UDynamicMaterialEditorSettings::Get()->DefaultRGBTexture.LoadSynchronous();
+			if (const TSoftObjectPtr<UTexture>* TexturePtr = UDynamicMaterialEditorSettings::Get()->DefaultSlotTextures.Find(EDMMaterialPropertyType::BaseColor))
+			{
+				if (UTexture* Texture = TexturePtr->LoadSynchronous())
+				{
+					return Texture;
+				}
+			}
+
+			return UDynamicMaterialEditorSettings::Get()->DefaultMask.LoadSynchronous();
 		});
 
 	FDMValueDetailsRowExtensions::Get().RegisterRowExtensions();
diff --git a/Engine/Plugins/Experimental/DynamicMaterial/Source/DynamicMaterialEditor/Private/DynamicMaterialEditorSettings.cpp b/Engine/Plugins/Experimental/DynamicMaterial/Source/DynamicMaterialEditor/Private/DynamicMaterialEditorSettings.cpp
--- a/Engine/Plugins/Experimental/DynamicMaterial/Source/DynamicMaterialEditor/Private/DynamicMaterialEditorSettings.cpp
+++ b/Engine/Plugins/Experimental/DynamicMaterial/Source/DynamicMaterialEditor/Private/DynamicMaterialEditorSettings.cpp
@@ -70,10 +70,20 @@ UDynamicMaterialEditorSettings::UDynamicMaterialEditorSettings()
 
 	ResetAllLayoutSettings();
 
-	DefaultRGBTexture = TSoftObjectPtr<UTexture>(FSoftObjectPath(TEXT("Texture2D'/DynamicMaterial/T_Default_Texture.T_Default_Texture'")));
-	DefaultOpaqueTexture = TSoftObjectPtr<UTexture>(FSoftObjectPath(TEXT("Texture2D'/Engine/EngineResources/WhiteSquareTexture.WhiteSquareTexture'")));
-	DefaultPreviewCanvasTexture = nullptr;
-	DefaultOpacitySlotMask = TSoftObjectPtr<UTexture>(FSoftObjectPath(TEXT("/Script/Engine.Texture2D'/DynamicMaterial/Textures/T_DM_HorizontalGradient.T_DM_HorizontalGradient'")));
+	DefaultMask = TSoftObjectPtr<UTexture>(FSoftObjectPath(TEXT("Texture2D'/Engine/EngineResources/WhiteSquareTexture.WhiteSquareTexture'")));
+
+	DefaultSlotTextures.Emplace(EDMMaterialPropertyType::BaseColor,           FSoftObjectPath(TEXT("Texture2D'/DynamicMaterial/T_Default_Texture.T_Default_Texture'")));
+	DefaultSlotTextures.Emplace(EDMMaterialPropertyType::Opacity,             FSoftObjectPath(TEXT("/Script/Engine.Texture2D'/DynamicMaterial/Textures/T_DM_HorizontalGradient.T_DM_HorizontalGradient'")));
+	DefaultSlotTextures.Emplace(EDMMaterialPropertyType::Metallic,            FSoftObjectPath(TEXT("/Script/Engine.Texture2D'/DynamicMaterial/Textures/SlotDefaults/T_MD_Metallic.T_MD_Metallic'")));
+	DefaultSlotTextures.Emplace(EDMMaterialPropertyType::Specular,            FSoftObjectPath(TEXT("")));
+	DefaultSlotTextures.Emplace(EDMMaterialPropertyType::Roughness,           FSoftObjectPath(TEXT("/Script/Engine.Texture2D'/DynamicMaterial/Textures/SlotDefaults/T_MD_Roughness.T_MD_Roughness'")));
+	DefaultSlotTextures.Emplace(EDMMaterialPropertyType::Anisotropy,          FSoftObjectPath(TEXT("")));
+	DefaultSlotTextures.Emplace(EDMMaterialPropertyType::Normal,              FSoftObjectPath(TEXT("/Script/Engine.Texture2D'/DynamicMaterial/Textures/SlotDefaults/T_MD_Normal.T_MD_Normal'")));
+	DefaultSlotTextures.Emplace(EDMMaterialPropertyType::Tangent,             FSoftObjectPath(TEXT("")));
+	DefaultSlotTextures.Emplace(EDMMaterialPropertyType::WorldPositionOffset, FSoftObjectPath(TEXT("")));
+	DefaultSlotTextures.Emplace(EDMMaterialPropertyType::AmbientOcclusion,    FSoftObjectPath(TEXT("/Script/Engine.Texture2D'/DynamicMaterial/Textures/SlotDefaults/T_MD_AmbientOcclusion.T_MD_AmbientOcclusion'")));
+	DefaultSlotTextures.Emplace(EDMMaterialPropertyType::Refraction,          FSoftObjectPath(TEXT("")));
+	DefaultSlotTextures.Emplace(EDMMaterialPropertyType::PixelDepthOffset,    FSoftObjectPath(TEXT("")));
 
 	FDMMaterialChannelListPreset Opaque;
 	Opaque.bRGB = true;
diff --git a/Engine/Plugins/Experimental/DynamicMaterial/Source/DynamicMaterialEditor/Private/DynamicMaterialEditorSettings.h b/Engine/Plugins/Experimental/DynamicMaterial/Source/DynamicMaterialEditor/Private/DynamicMaterialEditorSettings.h
--- a/Engine/Plugins/Experimental/DynamicMaterial/Source/DynamicMaterialEditor/Private/DynamicMaterialEditorSettings.h
+++ b/Engine/Plugins/Experimental/DynamicMaterial/Source/DynamicMaterialEditor/Private/DynamicMaterialEditorSettings.h
@@ -166,20 +166,14 @@ public:
 	bool bShowTooltipPreview;
 
 	/** The size of the preview image when displayed in the tooltip. */
-	UPROPERTY(Config, EditAnywhere, BlueprintReadWrite, Category = "ToolTips", meta = (EditCondition = "bShowTooltipPreview == true", ClampMin = "128", UIMin = "128", ClampMax = "2048", UIMax = "2048"))
+	UPROPERTY(Config, EditAnywhere, BlueprintReadWrite, Category = "ToolTips", meta = (EditCondition = "bShowTooltipPreview", ClampMin = "128", UIMin = "128", ClampMax = "2048", UIMax = "2048"))
 	int32 TooltipTextureSize;
 
-	UPROPERTY(Config, EditAnywhere, BlueprintReadWrite, Category = "Preview", meta=(DisplayName = "Default RGB Texture"))
-	TSoftObjectPtr<UTexture> DefaultRGBTexture;
-
-	UPROPERTY(Config, EditAnywhere, BlueprintReadWrite, Category = "Preview")
-	TSoftObjectPtr<UTexture> DefaultOpaqueTexture;
-
 	UPROPERTY(Config, EditAnywhere, BlueprintReadWrite, Category = "Preview")
-	TSoftObjectPtr<UTexture> DefaultPreviewCanvasTexture;
+	TSoftObjectPtr<UTexture> DefaultMask;
 
-	UPROPERTY(Config, EditAnywhere, BlueprintReadWrite, Category = "Preview")
-	TSoftObjectPtr<UTexture> DefaultOpacitySlotMask;
+	UPROPERTY(Config, EditAnywhere, BlueprintReadWrite, Category = "Channels", meta = (EditFixedSize))
+	TMap<EDMMaterialPropertyType, TSoftObjectPtr<UTexture>> DefaultSlotTextures;
 
 	/*
 	 * Add paths to search for custom effects.
diff --git a/Engine/Plugins/Experimental/DynamicMaterial/Source/DynamicMaterialEditor/Public/Components/DMMaterialProperty.h b/Engine/Plugins/Experimental/DynamicMaterial/Source/DynamicMaterialEditor/Public/Components/DMMaterialProperty.h
--- a/Engine/Plugins/Experimental/DynamicMaterial/Source/DynamicMaterialEditor/Public/Components/DMMaterialProperty.h
+++ b/Engine/Plugins/Experimental/DynamicMaterial/Source/DynamicMaterialEditor/Public/Components/DMMaterialProperty.h
@@ -59,7 +59,7 @@ public:
 
 	virtual UMaterialExpression* GetDefaultInput(const TSharedRef<FDMMaterialBuildState>& InBuildState) const;
 
-	virtual void OnSlotAdded(UDMMaterialSlot* InSlot) {}
+	virtual void OnSlotAdded(UDMMaterialSlot* InSlot);
 
 	//~ Begin UDMMaterialComponent
 	virtual void Update(EDMUpdateType InUpdateType) override;
diff --git a/Engine/Plugins/Experimental/DynamicMaterial/Source/DynamicMaterialEditor/Public/Components/MaterialProperties/DMMPBaseColor.h b/Engine/Plugins/Experimental/DynamicMaterial/Source/DynamicMaterialEditor/Public/Components/MaterialProperties/DMMPBaseColor.h
--- a/Engine/Plugins/Experimental/DynamicMaterial/Source/DynamicMaterialEditor/Public/Components/MaterialProperties/DMMPBaseColor.h
+++ b/Engine/Plugins/Experimental/DynamicMaterial/Source/DynamicMaterialEditor/Public/Components/MaterialProperties/DMMPBaseColor.h
@@ -13,13 +13,10 @@ class DYNAMICMATERIALEDITOR_API UDMMaterialPropertyBaseColor : public UDMMateria
 	GENERATED_BODY()
 
 public:
-	static void AddDefaultRGBLayer(EDMMaterialPropertyType InMaterialProperty, UDMMaterialSlot* InSlot);
-
 	UDMMaterialPropertyBaseColor();
 
 	//~ Begin UDMMaterialProperty
 	virtual bool IsValidForModel(UDynamicMaterialModelEditorOnlyData& InModelEditorOnlyData) const override;
 	virtual UMaterialExpression* GetDefaultInput(const TSharedRef<FDMMaterialBuildState>& InBuildState) const override;
-	virtual void OnSlotAdded(UDMMaterialSlot* InSlot) override;
 	//~ End UDMMaterialProperty
 };
diff --git a/Engine/Plugins/Experimental/DynamicMaterial/Source/DynamicMaterialEditor/Public/Components/MaterialProperties/DMMPEmissiveColor.h b/Engine/Plugins/Experimental/DynamicMaterial/Source/DynamicMaterialEditor/Public/Components/MaterialProperties/DMMPEmissiveColor.h
--- a/Engine/Plugins/Experimental/DynamicMaterial/Source/DynamicMaterialEditor/Public/Components/MaterialProperties/DMMPEmissiveColor.h
+++ b/Engine/Plugins/Experimental/DynamicMaterial/Source/DynamicMaterialEditor/Public/Components/MaterialProperties/DMMPEmissiveColor.h
@@ -18,6 +18,5 @@ public:
 	//~ Begin UDMMaterialProperty
 	virtual bool IsValidForModel(UDynamicMaterialModelEditorOnlyData& InModelEditorOnlyData) const override;
 	virtual UMaterialExpression* GetDefaultInput(const TSharedRef<FDMMaterialBuildState>& InBuildState) const override;
-	virtual void OnSlotAdded(UDMMaterialSlot* InSlot) override;
 	//~ End UDMMaterialProperty
 };
diff --git a/Engine/Plugins/Experimental/DynamicMaterial/Source/DynamicMaterialEditor/Public/Components/MaterialProperties/DMMPOpacity.h b/Engine/Plugins/Experimental/DynamicMaterial/Source/DynamicMaterialEditor/Public/Components/MaterialProperties/DMMPOpacity.h
--- a/Engine/Plugins/Experimental/DynamicMaterial/Source/DynamicMaterialEditor/Public/Components/MaterialProperties/DMMPOpacity.h
+++ b/Engine/Plugins/Experimental/DynamicMaterial/Source/DynamicMaterialEditor/Public/Components/MaterialProperties/DMMPOpacity.h
@@ -13,13 +13,10 @@ class DYNAMICMATERIALEDITOR_API UDMMaterialPropertyOpacity : public UDMMaterialP
 	GENERATED_BODY()
 
 public:
-	static void AddDefaultOpacityLayer(EDMMaterialPropertyType InMaterialProperty, UDMMaterialSlot* InSlot);
-
 	UDMMaterialPropertyOpacity();
 
 	//~ Begin UDMMaterialProperty
 	virtual bool IsValidForModel(UDynamicMaterialModelEditorOnlyData& InModelEditorOnlyData) const override;
 	virtual UMaterialExpression* GetDefaultInput(const TSharedRef<FDMMaterialBuildState>& InBuildState) const override;
-	virtual void OnSlotAdded(UDMMaterialSlot* InSlot) override;
 	//~ End UDMMaterialProperty
 };
diff --git a/Engine/Plugins/Experimental/DynamicMaterial/Source/DynamicMaterialEditor/Public/Components/MaterialProperties/DMMPOpacityMask.h b/Engine/Plugins/Experimental/DynamicMaterial/Source/DynamicMaterialEditor/Public/Components/MaterialProperties/DMMPOpacityMask.h
--- a/Engine/Plugins/Experimental/DynamicMaterial/Source/DynamicMaterialEditor/Public/Components/MaterialProperties/DMMPOpacityMask.h
+++ b/Engine/Plugins/Experimental/DynamicMaterial/Source/DynamicMaterialEditor/Public/Components/MaterialProperties/DMMPOpacityMask.h
@@ -18,6 +18,5 @@ public:
 	//~ Begin UDMMaterialProperty
 	virtual bool IsValidForModel(UDynamicMaterialModelEditorOnlyData& InModelEditorOnlyData) const override;
 	virtual UMaterialExpression* GetDefaultInput(const TSharedRef<FDMMaterialBuildState>& InBuildState) const override;
-	virtual void OnSlotAdded(UDMMaterialSlot* InSlot) override;
 	//~ End UDMMaterialProperty
 };
