Material Designer
- Horde issue 455218

diff --git a/Engine/Plugins/Experimental/DynamicMaterial/Source/DynamicMaterial/Public/Model/DynamicMaterialModel.h b/Engine/Plugins/Experimental/DynamicMaterial/Source/DynamicMaterial/Public/Model/DynamicMaterialModel.h
--- a/Engine/Plugins/Experimental/DynamicMaterial/Source/DynamicMaterial/Public/Model/DynamicMaterialModel.h
+++ b/Engine/Plugins/Experimental/DynamicMaterial/Source/DynamicMaterial/Public/Model/DynamicMaterialModel.h
@@ -5,6 +5,7 @@
 #include "UObject/Object.h"
 #include "Containers/Array.h"
 #include "Containers/Set.h"
+#include "Templates/SubclassOf.h"
 
 #if WITH_EDITOR
 #include "UObject/ScriptInterface.h"
