Material Designer
- Doubled the size of the UV handle.
- Fixed bugs with the UV handle axis highlights not matching the hover areas.
- Made the UV handles translucent.
- Added missing pre- and post-edit change calls.

diff --git a/Engine/Plugins/Experimental/DynamicMaterial/Source/DynamicMaterialEditor/Private/Slate/Properties/SDMTextureUVVisualizer.cpp b/Engine/Plugins/Experimental/DynamicMaterial/Source/DynamicMaterialEditor/Private/Slate/Properties/SDMTextureUVVisualizer.cpp
--- a/Engine/Plugins/Experimental/DynamicMaterial/Source/DynamicMaterialEditor/Private/Slate/Properties/SDMTextureUVVisualizer.cpp
+++ b/Engine/Plugins/Experimental/DynamicMaterial/Source/DynamicMaterialEditor/Private/Slate/Properties/SDMTextureUVVisualizer.cpp
@@ -6,6 +6,7 @@
 #include "Components/DMMaterialStage.h"
 #include "Components/DMTextureUV.h"
 #include "DynamicMaterialEditorModule.h"
+#include "DynamicMaterialEditorSettings.h"
 #include "Editor.h"
 #include "Framework/Application/SlateApplication.h"
 #include "Materials/MaterialInterface.h"
@@ -26,16 +27,19 @@ namespace UE::DynamicMaterialEditor::Private
 	const FVector2D TextureUVVisualizerPopoutImageSize = FVector2D(1024.f, 768.f);
 
 	/** Outer size of the position handle. */
-	constexpr float TextureUVVisualizerLargeRadius = 10.f;
+	constexpr float TextureUVVisualizerLargeRadius = 25.f;
 
 	/** Inner size of the position handle. */
-	constexpr float TextureUVVisualizerSmallRadius = 5.f;
+	constexpr float TextureUVVisualizerSmallRadius = 10.f;
 
 	/** "Radius" circle handle is for mouse interaction. */
 	constexpr float TextureUVVisualizerCircleHandleRadius = 5.f;
 
 	/** Based distance the circle handle is from the center compared to the width of the image. */
-	constexpr float TextureUVVisualizerCircleHandleBaseRadiusMultiplier = 0.25f;
+	constexpr float TextureUVVisualizerCircleHandleBaseRadiusMultiplier = 0.4f;
+
+	/** Size of the border around the center square */
+	constexpr float TextureUVVisualizerBorder = 8.f;
 }
 
 SDMTextureUVVisualizer::SDMTextureUVVisualizer()
@@ -246,9 +250,9 @@ int32 SDMTextureUVVisualizer::OnPaint(const FPaintArgs& InArgs, const FGeometry&
 
 	using namespace UE::DynamicMaterialEditor::Private;
 
-	static const FSlateColorBrush WhiteBrush = FSlateColorBrush(FColor(255, 255, 255, 255));
-	static const FLinearColor BorderColor = FLinearColor::Black;
-	static const FLinearColor NormalColor = FLinearColor::White;
+	static const FSlateColorBrush WhiteBrush = FSlateColorBrush(FLinearColor(1.f, 1.f, 1.f, 1.f));
+	static const FLinearColor BorderColor = FLinearColor(0.f, 0.f, 0.f, 0.5f);
+	static const FLinearColor NormalColor = FLinearColor(1.f, 1.f, 1.f, 0.25f);
 	static const FLinearColor HighlightColor = FStyleColors::Primary.GetSpecifiedColor();
 
 	const float Rotation = TextureUV->GetRotation();
@@ -285,7 +289,7 @@ int32 SDMTextureUVVisualizer::OnPaint(const FPaintArgs& InArgs, const FGeometry&
 	auto DrawRotatedBorderBox = [&DrawRotatedBox](const FVector2f& InLocation, const FVector2f& InSize, const FLinearColor& InInnerColor, float InRotationRadians)
 		{
 			DrawRotatedBox(InLocation, InSize, BorderColor, InRotationRadians);
-			return DrawRotatedBox(InLocation, InSize - FVector2f(4.f, 4.f), InInnerColor, InRotationRadians);
+			return DrawRotatedBox(InLocation, InSize - FVector2f(TextureUVVisualizerBorder, TextureUVVisualizerBorder), InInnerColor, InRotationRadians);
 		};
 
 	/** Center Handle */
@@ -302,14 +306,14 @@ int32 SDMTextureUVVisualizer::OnPaint(const FPaintArgs& InArgs, const FGeometry&
 		{
 			DrawRotatedBox(
 				FVector2f(TextureUVVisualizerLargeRadius - TextureUVVisualizerSmallRadius * 0.5f - 2.f, 0.f),
-				FVector2f(TextureUVVisualizerSmallRadius, TextureUVVisualizerLargeRadius * 2.f - 4.f),
+				FVector2f(TextureUVVisualizerSmallRadius, TextureUVVisualizerLargeRadius * 2.f - TextureUVVisualizerBorder * 0.5f),
 				HighlightColor,
 				bPivotEditMode ? 0 : RotationRadians
 			);
 
 			DrawRotatedBox(
 				FVector2f(-TextureUVVisualizerLargeRadius + TextureUVVisualizerSmallRadius * 0.5f + 2.f, 0.f),
-				FVector2f(TextureUVVisualizerSmallRadius, TextureUVVisualizerLargeRadius * 2.f - 4.f),
+				FVector2f(TextureUVVisualizerSmallRadius, TextureUVVisualizerLargeRadius * 2.f - TextureUVVisualizerBorder * 0.5f),
 				HighlightColor,
 				bPivotEditMode ? 0 : RotationRadians
 			);
@@ -319,14 +323,14 @@ int32 SDMTextureUVVisualizer::OnPaint(const FPaintArgs& InArgs, const FGeometry&
 		{
 			DrawRotatedBox(
 				FVector2f(0.f, TextureUVVisualizerLargeRadius - TextureUVVisualizerSmallRadius * 0.5f - 2.f),
-				FVector2f(TextureUVVisualizerLargeRadius * 2.f - 4.f, TextureUVVisualizerSmallRadius),
+				FVector2f(TextureUVVisualizerLargeRadius * 2.f - TextureUVVisualizerBorder * 0.5f, TextureUVVisualizerSmallRadius),
 				HighlightColor,
 				bPivotEditMode ? 0 : RotationRadians
 			);
 
 			DrawRotatedBox(
 				FVector2f(0.f, -TextureUVVisualizerLargeRadius + TextureUVVisualizerSmallRadius * 0.5f + 2.f),
-				FVector2f(TextureUVVisualizerLargeRadius * 2.f - 4.f, TextureUVVisualizerSmallRadius),
+				FVector2f(TextureUVVisualizerLargeRadius * 2.f - TextureUVVisualizerBorder * 0.5f, TextureUVVisualizerSmallRadius),
 				HighlightColor,
 				bPivotEditMode ? 0 : RotationRadians
 			);
@@ -591,17 +595,19 @@ SDMTextureUVVisualizer::EHandleAxis SDMTextureUVVisualizer::GetCenterHandleAxis(
 	HandleOffset.X = FMath::Abs(HandleOffset.X);
 	HandleOffset.Y = FMath::Abs(HandleOffset.Y);
 
-	if (HandleOffset.X <= TextureUVVisualizerSmallRadius && HandleOffset.Y <= TextureUVVisualizerSmallRadius)
+	static constexpr float LargeMinusSmallRadius = TextureUVVisualizerLargeRadius - TextureUVVisualizerSmallRadius;
+
+	if (HandleOffset.X <= LargeMinusSmallRadius && HandleOffset.Y <= LargeMinusSmallRadius)
 	{
 		return EHandleAxis::XY;
 	}
 
-	if (HandleOffset.X <= TextureUVVisualizerSmallRadius && HandleOffset.Y <= TextureUVVisualizerLargeRadius)
+	if (HandleOffset.X <= LargeMinusSmallRadius && HandleOffset.Y <= TextureUVVisualizerLargeRadius)
 	{
 		return EHandleAxis::Y;
 	}
 
-	if (HandleOffset.X <= TextureUVVisualizerLargeRadius && HandleOffset.Y <= TextureUVVisualizerSmallRadius)
+	if (HandleOffset.X <= TextureUVVisualizerLargeRadius && HandleOffset.Y <= LargeMinusSmallRadius)
 	{
 		return EHandleAxis::X;
 	}
@@ -902,11 +908,13 @@ void SDMTextureUVVisualizer::SetScrubbingMode(EScrubbingMode InMode, EHandleAxis
 
 	ScrubbingStartAbsoluteCenter = CurrentAbsoluteCenter;
 	ScrubbingStartAbsoluteMouse = FSlateApplication::Get().GetCursorPos();
+	FProperty* TextureProperty = nullptr;
 
 	switch (ScrubbingMode)
 	{
 		case EScrubbingMode::Offset:
 			ValueStart = TextureUV->GetOffset();
+			TextureProperty = UDMTextureUV::StaticClass()->FindPropertyByName(UDMTextureUV::NAME_Offset);
 			UE_LOG(LogDynamicMaterialEditor, Verbose, TEXT("Started Offset mode"));
 			break;
 
@@ -918,17 +926,21 @@ void SDMTextureUVVisualizer::SetScrubbingMode(EScrubbingMode InMode, EHandleAxis
 			const FVector2f MouseOffset = ScrubbingStartAbsoluteMouse - GetAbsolutePivotLocation();
 			ValueStart.Y = FMath::RadiansToDegrees(FMath::Atan2(MouseOffset.Y, MouseOffset.X));
 
+			TextureProperty = UDMTextureUV::StaticClass()->FindPropertyByName(UDMTextureUV::NAME_Rotation);
+
 			UE_LOG(LogDynamicMaterialEditor, Verbose, TEXT("Started Rotation mode"));
 			break;
 		}
 
 		case EScrubbingMode::Scale:
 			ValueStart = TextureUV->GetScale();
+			TextureProperty = UDMTextureUV::StaticClass()->FindPropertyByName(UDMTextureUV::NAME_Scale);
 			UE_LOG(LogDynamicMaterialEditor, Verbose, TEXT("Started Scale mode"));
 			break;
 
 		case EScrubbingMode::Pivot:
 			ValueStart = TextureUV->GetPivot();
+			TextureProperty = UDMTextureUV::StaticClass()->FindPropertyByName(UDMTextureUV::NAME_Pivot);
 			UE_LOG(LogDynamicMaterialEditor, Verbose, TEXT("Started Pivot mode"));
 			break;
 
@@ -938,6 +950,8 @@ void SDMTextureUVVisualizer::SetScrubbingMode(EScrubbingMode InMode, EHandleAxis
 
 	ScrubbingTransaction = MakeShared<FScopedTransaction>(LOCTEXT("VisualizerUVScrubbingTransaction", "UV Visualizer Scrub"));
 	TextureUV->Modify();
+
+	TextureUV->PreEditChange(TextureProperty);
 }
 
 FVector2f SDMTextureUVVisualizer::ToPopoutLocation(const FVector2f& InSize, FVector2f&& InLocation) const
@@ -1048,6 +1062,14 @@ void SDMTextureUVVisualizer::UpdateScrub_Offset()
 	}
 
 	TextureUV->SetOffset(ValueStart + OffsetChange);
+
+	FPropertyChangedEvent ChangedEvent(
+		UDMTextureUV::StaticClass()->FindPropertyByName(UDMTextureUV::NAME_Offset),
+		EPropertyChangeType::Interactive,
+		{TextureUV}
+	);
+
+	TextureUV->PostEditChangeProperty(ChangedEvent);
 }
 
 void SDMTextureUVVisualizer::UpdateScrub_Rotation()
@@ -1065,6 +1087,14 @@ void SDMTextureUVVisualizer::UpdateScrub_Rotation()
 	const float NewValue = UE::Math::TRotator<float>::ClampAxis(ValueStart.X + Angle - ValueStart.Y);
 
 	TextureUV->SetRotation(NewValue);
+
+	FPropertyChangedEvent ChangedEvent(
+		UDMTextureUV::StaticClass()->FindPropertyByName(UDMTextureUV::NAME_Rotation),
+		EPropertyChangeType::Interactive,
+		{TextureUV}
+	);
+
+	TextureUV->PostEditChangeProperty(ChangedEvent);
 }
 
 void SDMTextureUVVisualizer::UpdateScrub_Scale()
@@ -1136,6 +1166,14 @@ void SDMTextureUVVisualizer::UpdateScrub_Scale()
 	}
 
 	TextureUV->SetScale(NewScale);
+
+	FPropertyChangedEvent ChangedEvent(
+		UDMTextureUV::StaticClass()->FindPropertyByName(UDMTextureUV::NAME_Scale),
+		EPropertyChangeType::Interactive,
+		{TextureUV}
+	);
+
+	TextureUV->PostEditChangeProperty(ChangedEvent);
 }
 
 void SDMTextureUVVisualizer::UpdateScrub_Pivot()
@@ -1164,6 +1202,14 @@ void SDMTextureUVVisualizer::UpdateScrub_Pivot()
 	}
 
 	TextureUV->SetPivot(ValueStart + PivotChange);
+
+	FPropertyChangedEvent ChangedEvent(
+		UDMTextureUV::StaticClass()->FindPropertyByName(UDMTextureUV::NAME_Pivot),
+		EPropertyChangeType::Interactive,
+		{TextureUV}
+	);
+
+	TextureUV->PostEditChangeProperty(ChangedEvent);
 }
 
 #undef LOCTEXT_NAMESPACE
