Material Designer
- Added the ability for a material property to override the texture sampler type. E.g. The normal property will set the texture sampler type to normal.

diff --git a/Engine/Plugins/Experimental/DynamicMaterial/Source/DynamicMaterialEditor/Private/Components/DMMaterialProperty.cpp b/Engine/Plugins/Experimental/DynamicMaterial/Source/DynamicMaterialEditor/Private/Components/DMMaterialProperty.cpp
--- a/Engine/Plugins/Experimental/DynamicMaterial/Source/DynamicMaterialEditor/Private/Components/DMMaterialProperty.cpp
+++ b/Engine/Plugins/Experimental/DynamicMaterial/Source/DynamicMaterialEditor/Private/Components/DMMaterialProperty.cpp
@@ -107,6 +107,11 @@ UMaterialExpression* UDMMaterialProperty::GetDefaultInput(const TSharedRef<FDMMa
 	return nullptr;
 }
 
+TEnumAsByte<EMaterialSamplerType> UDMMaterialProperty::GetTextureSamplerType() const
+{
+	return EMaterialSamplerType::SAMPLERTYPE_Color;
+}
+
 void UDMMaterialProperty::OnSlotAdded(UDMMaterialSlot* InSlot)
 {
 	UTexture* BaseTexture = nullptr;
diff --git a/Engine/Plugins/Experimental/DynamicMaterial/Source/DynamicMaterialEditor/Private/Components/MaterialProperties/DMMPNormal.cpp b/Engine/Plugins/Experimental/DynamicMaterial/Source/DynamicMaterialEditor/Private/Components/MaterialProperties/DMMPNormal.cpp
--- a/Engine/Plugins/Experimental/DynamicMaterial/Source/DynamicMaterialEditor/Private/Components/MaterialProperties/DMMPNormal.cpp
+++ b/Engine/Plugins/Experimental/DynamicMaterial/Source/DynamicMaterialEditor/Private/Components/MaterialProperties/DMMPNormal.cpp
@@ -32,3 +32,8 @@ UMaterialExpression* UDMMaterialPropertyNormal::GetDefaultInput(const TSharedRef
 {
 	return CreateConstant(InBuildState, FVector::OneVector);
 }
+
+TEnumAsByte<EMaterialSamplerType> UDMMaterialPropertyNormal::GetTextureSamplerType() const
+{
+	return EMaterialSamplerType::SAMPLERTYPE_Normal;
+}
diff --git a/Engine/Plugins/Experimental/DynamicMaterial/Source/DynamicMaterialEditor/Private/Components/MaterialProperties/DMMPOpacity.cpp b/Engine/Plugins/Experimental/DynamicMaterial/Source/DynamicMaterialEditor/Private/Components/MaterialProperties/DMMPOpacity.cpp
--- a/Engine/Plugins/Experimental/DynamicMaterial/Source/DynamicMaterialEditor/Private/Components/MaterialProperties/DMMPOpacity.cpp
+++ b/Engine/Plugins/Experimental/DynamicMaterial/Source/DynamicMaterialEditor/Private/Components/MaterialProperties/DMMPOpacity.cpp
@@ -30,3 +30,8 @@ UMaterialExpression* UDMMaterialPropertyOpacity::GetDefaultInput(const TSharedRe
 {
 	return CreateConstant(InBuildState, 1.f);
 }
+
+TEnumAsByte<EMaterialSamplerType> UDMMaterialPropertyOpacity::GetTextureSamplerType() const
+{
+	return EMaterialSamplerType::SAMPLERTYPE_Alpha;
+}
diff --git a/Engine/Plugins/Experimental/DynamicMaterial/Source/DynamicMaterialEditor/Private/Components/MaterialProperties/DMMPOpacityMask.cpp b/Engine/Plugins/Experimental/DynamicMaterial/Source/DynamicMaterialEditor/Private/Components/MaterialProperties/DMMPOpacityMask.cpp
--- a/Engine/Plugins/Experimental/DynamicMaterial/Source/DynamicMaterialEditor/Private/Components/MaterialProperties/DMMPOpacityMask.cpp
+++ b/Engine/Plugins/Experimental/DynamicMaterial/Source/DynamicMaterialEditor/Private/Components/MaterialProperties/DMMPOpacityMask.cpp
@@ -22,3 +22,8 @@ UMaterialExpression* UDMMaterialPropertyOpacityMask::GetDefaultInput(
 {
 	return CreateConstant(InBuildState, 1.f);
 }
+
+TEnumAsByte<EMaterialSamplerType> UDMMaterialPropertyOpacityMask::GetTextureSamplerType() const
+{
+	return EMaterialSamplerType::SAMPLERTYPE_Masks;
+}
diff --git a/Engine/Plugins/Experimental/DynamicMaterial/Source/DynamicMaterialEditor/Private/Components/MaterialStageExpressions/DMMSETextureSampleBase.cpp b/Engine/Plugins/Experimental/DynamicMaterial/Source/DynamicMaterialEditor/Private/Components/MaterialStageExpressions/DMMSETextureSampleBase.cpp
--- a/Engine/Plugins/Experimental/DynamicMaterial/Source/DynamicMaterialEditor/Private/Components/MaterialStageExpressions/DMMSETextureSampleBase.cpp
+++ b/Engine/Plugins/Experimental/DynamicMaterial/Source/DynamicMaterialEditor/Private/Components/MaterialStageExpressions/DMMSETextureSampleBase.cpp
@@ -3,6 +3,7 @@
 #include "Components/MaterialStageExpressions/DMMSETextureSampleBase.h"
 
 #include "Components/DMMaterialLayer.h"
+#include "Components/DMMaterialProperty.h"
 #include "Components/DMMaterialSlot.h"
 #include "Components/DMMaterialStageBlend.h"
 #include "Components/DMMaterialStageThroughputLayerBlend.h"
@@ -15,6 +16,7 @@
 #include "CoreGlobals.h"
 #include "Engine/TextureCube.h"
 #include "Materials/MaterialExpressionTextureSample.h"
+#include "Model/DMMaterialBuildState.h"
 #include "Model/DynamicMaterialModelEditorOnlyData.h"
 
 #define LOCTEXT_NAMESPACE "DMMaterialStageExpressionTextureSample"
@@ -44,6 +46,42 @@ UDMMaterialStageExpressionTextureSampleBase::UDMMaterialStageExpressionTextureSa
 	EditableProperties.Add(GET_MEMBER_NAME_CHECKED(UDMMaterialStageExpressionTextureSampleBase, bClampTexture));
 }
 
+void UDMMaterialStageExpressionTextureSampleBase::GenerateExpressions(const TSharedRef<FDMMaterialBuildState>& InBuildState) const
+{
+	if (!IsComponentValid() || !IsComponentAdded())
+	{
+		return;
+	}
+
+	check(MaterialExpressionClass.Get());
+
+	if (InBuildState->HasStageSource(this))
+	{
+		return;
+	}
+
+	UMaterialExpression* NewExpression = InBuildState->GetBuildUtils().CreateExpression(MaterialExpressionClass.Get(), UE_DM_NodeComment_Default);
+	AddExpressionProperties({NewExpression});
+
+	InBuildState->AddStageSourceExpressions(this, {NewExpression});
+
+	UMaterialExpressionTextureSample* TextureSample = Cast<UMaterialExpressionTextureSample>(NewExpression);
+
+	if (!TextureSample)
+	{
+		return;
+	}
+
+	UDMMaterialProperty* CurrentProperty = InBuildState->GetCurrentMaterialProperty();
+
+	if (!CurrentProperty)
+	{
+		return;
+	}
+
+	TextureSample->SamplerType = CurrentProperty->GetTextureSamplerType();
+}
+
 void UDMMaterialStageExpressionTextureSampleBase::OnComponentAdded()
 {
 	Super::OnComponentAdded();
diff --git a/Engine/Plugins/Experimental/DynamicMaterial/Source/DynamicMaterialEditor/Private/Model/DMMaterialBuildState.cpp b/Engine/Plugins/Experimental/DynamicMaterial/Source/DynamicMaterialEditor/Private/Model/DMMaterialBuildState.cpp
--- a/Engine/Plugins/Experimental/DynamicMaterial/Source/DynamicMaterialEditor/Private/Model/DMMaterialBuildState.cpp
+++ b/Engine/Plugins/Experimental/DynamicMaterial/Source/DynamicMaterialEditor/Private/Model/DMMaterialBuildState.cpp
@@ -11,7 +11,6 @@
 #include "MaterialDomain.h"
 #include "Materials/Material.h"
 #include "Materials/MaterialExpression.h"
-#include "Materials/MaterialExpressionAppendVector.h"
 #include "Model/DMMaterialNodeArranger.h"
 #include "Model/DMMaterialBuildUtils.h"
 
@@ -61,6 +60,16 @@ UDynamicMaterialModel* FDMMaterialBuildState::GetMaterialModel() const
 	return MaterialModel;
 }
 
+UDMMaterialProperty* FDMMaterialBuildState::GetCurrentMaterialProperty() const
+{
+	return CurrentProperty;
+}
+
+void FDMMaterialBuildState::SetCurrentMaterialProperty(UDMMaterialProperty* InProperty)
+{
+	CurrentProperty = InProperty;
+}
+
 void FDMMaterialBuildState::SetIgnoreUVs()
 {
 	bIgnoreUVs = true;
diff --git a/Engine/Plugins/Experimental/DynamicMaterial/Source/DynamicMaterialEditor/Private/Model/DynamicMaterialModelEditorOnlyData.cpp b/Engine/Plugins/Experimental/DynamicMaterial/Source/DynamicMaterialEditor/Private/Model/DynamicMaterialModelEditorOnlyData.cpp
--- a/Engine/Plugins/Experimental/DynamicMaterial/Source/DynamicMaterialEditor/Private/Model/DynamicMaterialModelEditorOnlyData.cpp
+++ b/Engine/Plugins/Experimental/DynamicMaterial/Source/DynamicMaterialEditor/Private/Model/DynamicMaterialModelEditorOnlyData.cpp
@@ -330,6 +330,8 @@ void UDynamicMaterialModelEditorOnlyData::BuildMaterial(bool bInDirtyAssets)
 			continue;
 		}
 
+		BuildState->SetCurrentMaterialProperty(Pair.Value);
+
 		Slot->GenerateExpressions(BuildState);
 
 		if (BuildState->GetSlotExpressions(Slot).IsEmpty())
diff --git a/Engine/Plugins/Experimental/DynamicMaterial/Source/DynamicMaterialEditor/Public/Components/DMMaterialProperty.h b/Engine/Plugins/Experimental/DynamicMaterial/Source/DynamicMaterialEditor/Public/Components/DMMaterialProperty.h
--- a/Engine/Plugins/Experimental/DynamicMaterial/Source/DynamicMaterialEditor/Public/Components/DMMaterialProperty.h
+++ b/Engine/Plugins/Experimental/DynamicMaterial/Source/DynamicMaterialEditor/Public/Components/DMMaterialProperty.h
@@ -59,6 +59,8 @@ public:
 
 	virtual UMaterialExpression* GetDefaultInput(const TSharedRef<FDMMaterialBuildState>& InBuildState) const;
 
+	virtual TEnumAsByte<EMaterialSamplerType> GetTextureSamplerType() const;
+
 	virtual void OnSlotAdded(UDMMaterialSlot* InSlot);
 
 	//~ Begin UDMMaterialComponent
diff --git a/Engine/Plugins/Experimental/DynamicMaterial/Source/DynamicMaterialEditor/Public/Components/MaterialProperties/DMMPNormal.h b/Engine/Plugins/Experimental/DynamicMaterial/Source/DynamicMaterialEditor/Public/Components/MaterialProperties/DMMPNormal.h
--- a/Engine/Plugins/Experimental/DynamicMaterial/Source/DynamicMaterialEditor/Public/Components/MaterialProperties/DMMPNormal.h
+++ b/Engine/Plugins/Experimental/DynamicMaterial/Source/DynamicMaterialEditor/Public/Components/MaterialProperties/DMMPNormal.h
@@ -18,5 +18,6 @@ public:
 	//~ Begin UDMMaterialProperty
 	virtual bool IsValidForModel(UDynamicMaterialModelEditorOnlyData& InModelEditorOnlyData) const override;
 	virtual UMaterialExpression* GetDefaultInput(const TSharedRef<FDMMaterialBuildState>& InBuildState) const override;
+	virtual TEnumAsByte<EMaterialSamplerType> GetTextureSamplerType() const override;
 	//~ End UDMMaterialProperty
 };
diff --git a/Engine/Plugins/Experimental/DynamicMaterial/Source/DynamicMaterialEditor/Public/Components/MaterialProperties/DMMPOpacity.h b/Engine/Plugins/Experimental/DynamicMaterial/Source/DynamicMaterialEditor/Public/Components/MaterialProperties/DMMPOpacity.h
--- a/Engine/Plugins/Experimental/DynamicMaterial/Source/DynamicMaterialEditor/Public/Components/MaterialProperties/DMMPOpacity.h
+++ b/Engine/Plugins/Experimental/DynamicMaterial/Source/DynamicMaterialEditor/Public/Components/MaterialProperties/DMMPOpacity.h
@@ -18,5 +18,6 @@ public:
 	//~ Begin UDMMaterialProperty
 	virtual bool IsValidForModel(UDynamicMaterialModelEditorOnlyData& InModelEditorOnlyData) const override;
 	virtual UMaterialExpression* GetDefaultInput(const TSharedRef<FDMMaterialBuildState>& InBuildState) const override;
+	virtual TEnumAsByte<EMaterialSamplerType> GetTextureSamplerType() const override;
 	//~ End UDMMaterialProperty
 };
diff --git a/Engine/Plugins/Experimental/DynamicMaterial/Source/DynamicMaterialEditor/Public/Components/MaterialProperties/DMMPOpacityMask.h b/Engine/Plugins/Experimental/DynamicMaterial/Source/DynamicMaterialEditor/Public/Components/MaterialProperties/DMMPOpacityMask.h
--- a/Engine/Plugins/Experimental/DynamicMaterial/Source/DynamicMaterialEditor/Public/Components/MaterialProperties/DMMPOpacityMask.h
+++ b/Engine/Plugins/Experimental/DynamicMaterial/Source/DynamicMaterialEditor/Public/Components/MaterialProperties/DMMPOpacityMask.h
@@ -18,5 +18,6 @@ public:
 	//~ Begin UDMMaterialProperty
 	virtual bool IsValidForModel(UDynamicMaterialModelEditorOnlyData& InModelEditorOnlyData) const override;
 	virtual UMaterialExpression* GetDefaultInput(const TSharedRef<FDMMaterialBuildState>& InBuildState) const override;
+	virtual TEnumAsByte<EMaterialSamplerType> GetTextureSamplerType() const override;
 	//~ End UDMMaterialProperty
 };
diff --git a/Engine/Plugins/Experimental/DynamicMaterial/Source/DynamicMaterialEditor/Public/Components/MaterialStageExpressions/DMMSETextureSampleBase.h b/Engine/Plugins/Experimental/DynamicMaterial/Source/DynamicMaterialEditor/Public/Components/MaterialStageExpressions/DMMSETextureSampleBase.h
--- a/Engine/Plugins/Experimental/DynamicMaterial/Source/DynamicMaterialEditor/Public/Components/MaterialStageExpressions/DMMSETextureSampleBase.h
+++ b/Engine/Plugins/Experimental/DynamicMaterial/Source/DynamicMaterialEditor/Public/Components/MaterialStageExpressions/DMMSETextureSampleBase.h
@@ -13,6 +13,10 @@ class DYNAMICMATERIALEDITOR_API UDMMaterialStageExpressionTextureSampleBase : pu
 public:
 	UDMMaterialStageExpressionTextureSampleBase(const FText& InName, TSubclassOf<UMaterialExpression> InClass);
 
+	//~ Begin UDMMaterialStageExpression
+	virtual void GenerateExpressions(const TSharedRef<FDMMaterialBuildState>& InBuildState) const override;
+	//~ End UDMMaterialStageExpression
+
 	//~ Begin UDMMaterialStageSource
 	virtual void OnComponentAdded() override;
 	virtual bool IsPropertyVisible(FName Property) const override;
diff --git a/Engine/Plugins/Experimental/DynamicMaterial/Source/DynamicMaterialEditor/Public/Model/DMMaterialBuildState.h b/Engine/Plugins/Experimental/DynamicMaterial/Source/DynamicMaterialEditor/Public/Model/DMMaterialBuildState.h
--- a/Engine/Plugins/Experimental/DynamicMaterial/Source/DynamicMaterialEditor/Public/Model/DMMaterialBuildState.h
+++ b/Engine/Plugins/Experimental/DynamicMaterial/Source/DynamicMaterialEditor/Public/Model/DMMaterialBuildState.h
@@ -9,6 +9,7 @@
 #include "Model/DMMaterialBuildUtils.h"
 
 class UDMMaterialLayerObject;
+class UDMMaterialProperty;
 class UDMMaterialSlot;
 class UDMMaterialStage;
 class UDMMaterialStageInput;
@@ -39,6 +40,10 @@ struct DYNAMICMATERIALEDITOR_API FDMMaterialBuildState : public IDMMaterialBuild
 
 	virtual UDynamicMaterialModel* GetMaterialModel() const override;
 
+	UDMMaterialProperty* GetCurrentMaterialProperty() const;
+
+	void SetCurrentMaterialProperty(UDMMaterialProperty* InProperty);
+
 	bool ShouldDirtyAssets() const { return bDirtyAssets; }
 
 	void SetIgnoreUVs();
@@ -147,6 +152,7 @@ struct DYNAMICMATERIALEDITOR_API FDMMaterialBuildState : public IDMMaterialBuild
 private:
 	UMaterial* DynamicMaterial = nullptr;
 	UDynamicMaterialModel* MaterialModel = nullptr;
+	UDMMaterialProperty* CurrentProperty = nullptr;
 	bool bDirtyAssets;
 	bool bIgnoreUVs;
 	bool bIsPreviewMaterial;
