Material Designer
- Fixed missing default textures for material properties.
- Fixed a crash when adding a new layer.

diff --git a/Engine/Plugins/Experimental/DynamicMaterial/Source/DynamicMaterialEditor/Private/Components/DMMaterialLayer.cpp b/Engine/Plugins/Experimental/DynamicMaterial/Source/DynamicMaterialEditor/Private/Components/DMMaterialLayer.cpp
--- a/Engine/Plugins/Experimental/DynamicMaterial/Source/DynamicMaterialEditor/Private/Components/DMMaterialLayer.cpp
+++ b/Engine/Plugins/Experimental/DynamicMaterial/Source/DynamicMaterialEditor/Private/Components/DMMaterialLayer.cpp
@@ -58,6 +58,7 @@ UDMMaterialLayerObject* UDMMaterialLayerObject::CreateLayer(UDMMaterialSlot* InS
 	const TArray<UDMMaterialStage*>& InStages)
 {
 	check(InSlot);
+	check(InMaterialProperty != EDMMaterialPropertyType::None);
 
 	UDMMaterialLayerObject* NewLayer = NewObject<UDMMaterialLayerObject>(InSlot, NAME_None, RF_Transactional);
 	NewLayer->MaterialProperty = InMaterialProperty;
diff --git a/Engine/Plugins/Experimental/DynamicMaterial/Source/DynamicMaterialEditor/Private/Components/DMMaterialProperty.cpp b/Engine/Plugins/Experimental/DynamicMaterial/Source/DynamicMaterialEditor/Private/Components/DMMaterialProperty.cpp
--- a/Engine/Plugins/Experimental/DynamicMaterial/Source/DynamicMaterialEditor/Private/Components/DMMaterialProperty.cpp
+++ b/Engine/Plugins/Experimental/DynamicMaterial/Source/DynamicMaterialEditor/Private/Components/DMMaterialProperty.cpp
@@ -114,21 +114,6 @@ TEnumAsByte<EMaterialSamplerType> UDMMaterialProperty::GetTextureSamplerType() c
 
 void UDMMaterialProperty::OnSlotAdded(UDMMaterialSlot* InSlot)
 {
-	UTexture* BaseTexture = nullptr;
-
-	if (const TSoftObjectPtr<UTexture>* TexturePtr = UDynamicMaterialEditorSettings::Get()->DefaultSlotTextures.Find(MaterialProperty))
-	{
-		if (UTexture* Texture = TexturePtr->LoadSynchronous())
-		{
-			BaseTexture = Texture;
-		}
-	}
-
-	if (!BaseTexture)
-	{
-		return;
-	}
-
 	UDMMaterialStage* DefaultStage = UDMMaterialStageBlend::CreateStage(UDMMaterialStageBlendNormal::StaticClass());
 	check(DefaultStage);
 
@@ -137,35 +122,48 @@ void UDMMaterialProperty::OnSlotAdded(UDMMaterialSlot* InSlot)
 
 	InSlot->AddLayerWithMask(MaterialProperty, DefaultStage, MaskStage);
 
-	UDMMaterialStageInputExpression* BaseInputExpression = UDMMaterialStageInputExpression::ChangeStageInput_Expression(
-		DefaultStage,
-		UDMMaterialStageExpressionTextureSample::StaticClass(),
-		UDMMaterialStageBlendNormal::InputB,
-		FDMMaterialStageConnectorChannel::WHOLE_CHANNEL,
-		0,
-		FDMMaterialStageConnectorChannel::WHOLE_CHANNEL
-	);
-
-	if (UDMMaterialStageExpressionTextureSample* BaseInputTextureSample = Cast<UDMMaterialStageExpressionTextureSample>(BaseInputExpression->GetMaterialStageExpression()))
+	if (UTexture* BaseTexture = UDynamicMaterialEditorSettings::Get()->GetDefaultTextureForSlot(MaterialProperty))
 	{
-		if (UDMMaterialStage* BaseTextureInputStage = BaseInputExpression->GetSubStage())
+		UDMMaterialStageInputExpression* BaseInputExpression = UDMMaterialStageInputExpression::ChangeStageInput_Expression(
+			DefaultStage,
+			UDMMaterialStageExpressionTextureSample::StaticClass(),
+			UDMMaterialStageBlendNormal::InputB,
+			FDMMaterialStageConnectorChannel::WHOLE_CHANNEL,
+			0,
+			FDMMaterialStageConnectorChannel::WHOLE_CHANNEL
+		);
+
+		if (UDMMaterialStageExpressionTextureSample* BaseInputTextureSample = Cast<UDMMaterialStageExpressionTextureSample>(BaseInputExpression->GetMaterialStageExpression()))
 		{
-			const TArray<UDMMaterialStageInput*> BaseTextureStageInputs = BaseTextureInputStage->GetInputs();
-
-			for (UDMMaterialStageInput* BaseTextureStageInput : BaseTextureStageInputs)
+			if (UDMMaterialStage* BaseTextureInputStage = BaseInputExpression->GetSubStage())
 			{
-				if (UDMMaterialStageInputValue* BaseTextureInputValue = Cast<UDMMaterialStageInputValue>(BaseTextureStageInput))
+				const TArray<UDMMaterialStageInput*> BaseTextureStageInputs = BaseTextureInputStage->GetInputs();
+
+				for (UDMMaterialStageInput* BaseTextureStageInput : BaseTextureStageInputs)
 				{
-					if (UDMMaterialValueTexture* BaseTextureValue = Cast<UDMMaterialValueTexture>(BaseTextureInputValue->GetValue()))
+					if (UDMMaterialStageInputValue* BaseTextureInputValue = Cast<UDMMaterialStageInputValue>(BaseTextureStageInput))
 					{
-						BaseTextureValue->SetDefaultValue(BaseTexture);
-						BaseTextureValue->ApplyDefaultValue();
-						break;
+						if (UDMMaterialValueTexture* BaseTextureValue = Cast<UDMMaterialValueTexture>(BaseTextureInputValue->GetValue()))
+						{
+							BaseTextureValue->SetDefaultValue(BaseTexture);
+							BaseTextureValue->ApplyDefaultValue();
+							break;
+						}
 					}
 				}
 			}
 		}
 	}
+	else
+	{
+		UDMMaterialStageInputValue::ChangeStageInput_NewLocalValue(
+			DefaultStage,
+			UDMMaterialStageBlendNormal::InputB,
+			FDMMaterialStageConnectorChannel::WHOLE_CHANNEL,
+			EDMValueType::VT_Float3_RGB,
+			0
+		);
+	}
 
 	if (UTexture* MaskTexture = UDynamicMaterialEditorSettings::Get()->DefaultMask.LoadSynchronous())
 	{
diff --git a/Engine/Plugins/Experimental/DynamicMaterial/Source/DynamicMaterialEditor/Private/Components/DMMaterialSlot.cpp b/Engine/Plugins/Experimental/DynamicMaterial/Source/DynamicMaterialEditor/Private/Components/DMMaterialSlot.cpp
--- a/Engine/Plugins/Experimental/DynamicMaterial/Source/DynamicMaterialEditor/Private/Components/DMMaterialSlot.cpp
+++ b/Engine/Plugins/Experimental/DynamicMaterial/Source/DynamicMaterialEditor/Private/Components/DMMaterialSlot.cpp
@@ -518,23 +518,19 @@ bool UDMMaterialSlot::CanRemoveLayer(const UDMMaterialLayerObject* InLayer) cons
 	check(InLayer->GetSlot() == this);
 
 	const EDMMaterialPropertyType LayerProperty = InLayer->GetMaterialProperty();
+	int32 LayerPropertyCount = 0;
 
-	if (LayerProperty == EDMMaterialPropertyType::BaseColor || LayerProperty == EDMMaterialPropertyType::EmissiveColor)
+	for (UDMMaterialLayerObject* Layer : LayerObjects)
 	{
-		int32 LayerPropertyCount = 0;
-
-		for (UDMMaterialLayerObject* Layer : LayerObjects)
+		if (Layer->GetMaterialProperty() == LayerProperty)
 		{
-			if (Layer->GetMaterialProperty() == LayerProperty)
-			{
-				++LayerPropertyCount;
-			}
+			++LayerPropertyCount;
 		}
+	}
 
-		if (LayerPropertyCount == 1)
-		{
-			return false;
-		}
+	if (LayerPropertyCount == 1)
+	{
+		return false;
 	}
 
 	return InLayer->FindIndex() != INDEX_NONE;
diff --git a/Engine/Plugins/Experimental/DynamicMaterial/Source/DynamicMaterialEditor/Private/DynamicMaterialEditorModule.cpp b/Engine/Plugins/Experimental/DynamicMaterial/Source/DynamicMaterialEditor/Private/DynamicMaterialEditorModule.cpp
--- a/Engine/Plugins/Experimental/DynamicMaterial/Source/DynamicMaterialEditor/Private/DynamicMaterialEditorModule.cpp
+++ b/Engine/Plugins/Experimental/DynamicMaterial/Source/DynamicMaterialEditor/Private/DynamicMaterialEditorModule.cpp
@@ -303,12 +303,9 @@ void FDynamicMaterialEditorModule::StartupModule()
 
 	UDMMaterialValueTexture::GetDefaultRGBTexture.BindLambda([]()
 		{
-			if (const TSoftObjectPtr<UTexture>* TexturePtr = UDynamicMaterialEditorSettings::Get()->DefaultSlotTextures.Find(EDMMaterialPropertyType::BaseColor))
+			if (UTexture* Texture = UDynamicMaterialEditorSettings::Get()->GetDefaultTextureForSlot(EDMMaterialPropertyType::BaseColor))
 			{
-				if (UTexture* Texture = TexturePtr->LoadSynchronous())
-				{
-					return Texture;
-				}
+				return Texture;
 			}
 
 			return UDynamicMaterialEditorSettings::Get()->DefaultMask.LoadSynchronous();
diff --git a/Engine/Plugins/Experimental/DynamicMaterial/Source/DynamicMaterialEditor/Private/DynamicMaterialEditorSettings.cpp b/Engine/Plugins/Experimental/DynamicMaterial/Source/DynamicMaterialEditor/Private/DynamicMaterialEditorSettings.cpp
--- a/Engine/Plugins/Experimental/DynamicMaterial/Source/DynamicMaterialEditor/Private/DynamicMaterialEditorSettings.cpp
+++ b/Engine/Plugins/Experimental/DynamicMaterial/Source/DynamicMaterialEditor/Private/DynamicMaterialEditorSettings.cpp
@@ -14,6 +14,20 @@
 
 #define LOCTEXT_NAMESPACE "MaterialDesignerSettings"
 
+namespace UE::DynamicMaterialEditor::Private
+{
+	const TMap<EDMMaterialPropertyType, TSoftObjectPtr<UTexture>> BaseDefaultSlotProperties = {
+		{EDMMaterialPropertyType::BaseColor,           TSoftObjectPtr<UTexture>(FSoftObjectPath(TEXT("Texture2D'/DynamicMaterial/T_Default_Texture.T_Default_Texture'")))},
+		{EDMMaterialPropertyType::EmissiveColor,       TSoftObjectPtr<UTexture>(FSoftObjectPath(TEXT("Texture2D'/DynamicMaterial/T_Default_Texture.T_Default_Texture'")))},
+		{EDMMaterialPropertyType::Opacity,             TSoftObjectPtr<UTexture>(FSoftObjectPath(TEXT("/Script/Engine.Texture2D'/DynamicMaterial/Textures/T_DM_HorizontalGradient.T_DM_HorizontalGradient'")))},
+		{EDMMaterialPropertyType::OpacityMask,         TSoftObjectPtr<UTexture>(FSoftObjectPath(TEXT("/Script/Engine.Texture2D'/DynamicMaterial/Textures/T_DM_HorizontalGradient.T_DM_HorizontalGradient'")))},
+		{EDMMaterialPropertyType::Metallic,            TSoftObjectPtr<UTexture>(FSoftObjectPath(TEXT("/Script/Engine.Texture2D'/DynamicMaterial/Textures/SlotDefaults/T_MD_Metallic.T_MD_Metallic'")))},
+		{EDMMaterialPropertyType::Roughness,           TSoftObjectPtr<UTexture>(FSoftObjectPath(TEXT("/Script/Engine.Texture2D'/DynamicMaterial/Textures/SlotDefaults/T_MD_Roughness.T_MD_Roughness'")))},
+		{EDMMaterialPropertyType::Normal,              TSoftObjectPtr<UTexture>(FSoftObjectPath(TEXT("/Script/Engine.Texture2D'/DynamicMaterial/Textures/SlotDefaults/T_MD_Normal.T_MD_Normal'")))},
+		{EDMMaterialPropertyType::AmbientOcclusion,    TSoftObjectPtr<UTexture>(FSoftObjectPath(TEXT("/Script/Engine.Texture2D'/DynamicMaterial/Textures/SlotDefaults/T_MD_AmbientOcclusion.T_MD_AmbientOcclusion'")))},
+	};
+}
+
 bool FDMMaterialChannelListPreset::IsPropertyEnabled(EDMMaterialPropertyType InProperty) const
 {
 	switch (InProperty)
@@ -72,18 +86,7 @@ UDynamicMaterialEditorSettings::UDynamicMaterialEditorSettings()
 
 	DefaultMask = TSoftObjectPtr<UTexture>(FSoftObjectPath(TEXT("Texture2D'/Engine/EngineResources/WhiteSquareTexture.WhiteSquareTexture'")));
 
-	DefaultSlotTextures.Emplace(EDMMaterialPropertyType::BaseColor,           FSoftObjectPath(TEXT("Texture2D'/DynamicMaterial/T_Default_Texture.T_Default_Texture'")));
-	DefaultSlotTextures.Emplace(EDMMaterialPropertyType::Opacity,             FSoftObjectPath(TEXT("/Script/Engine.Texture2D'/DynamicMaterial/Textures/T_DM_HorizontalGradient.T_DM_HorizontalGradient'")));
-	DefaultSlotTextures.Emplace(EDMMaterialPropertyType::Metallic,            FSoftObjectPath(TEXT("/Script/Engine.Texture2D'/DynamicMaterial/Textures/SlotDefaults/T_MD_Metallic.T_MD_Metallic'")));
-	DefaultSlotTextures.Emplace(EDMMaterialPropertyType::Specular,            FSoftObjectPath(TEXT("")));
-	DefaultSlotTextures.Emplace(EDMMaterialPropertyType::Roughness,           FSoftObjectPath(TEXT("/Script/Engine.Texture2D'/DynamicMaterial/Textures/SlotDefaults/T_MD_Roughness.T_MD_Roughness'")));
-	DefaultSlotTextures.Emplace(EDMMaterialPropertyType::Anisotropy,          FSoftObjectPath(TEXT("")));
-	DefaultSlotTextures.Emplace(EDMMaterialPropertyType::Normal,              FSoftObjectPath(TEXT("/Script/Engine.Texture2D'/DynamicMaterial/Textures/SlotDefaults/T_MD_Normal.T_MD_Normal'")));
-	DefaultSlotTextures.Emplace(EDMMaterialPropertyType::Tangent,             FSoftObjectPath(TEXT("")));
-	DefaultSlotTextures.Emplace(EDMMaterialPropertyType::WorldPositionOffset, FSoftObjectPath(TEXT("")));
-	DefaultSlotTextures.Emplace(EDMMaterialPropertyType::AmbientOcclusion,    FSoftObjectPath(TEXT("/Script/Engine.Texture2D'/DynamicMaterial/Textures/SlotDefaults/T_MD_AmbientOcclusion.T_MD_AmbientOcclusion'")));
-	DefaultSlotTextures.Emplace(EDMMaterialPropertyType::Refraction,          FSoftObjectPath(TEXT("")));
-	DefaultSlotTextures.Emplace(EDMMaterialPropertyType::PixelDepthOffset,    FSoftObjectPath(TEXT("")));
+	OverriddenDefaultSlotTextures = {};
 
 	FDMMaterialChannelListPreset Opaque;
 	Opaque.bRGB = true;
@@ -334,4 +337,25 @@ TArray<FDMMaterialEffectList> UDynamicMaterialEditorSettings::GetEffectList() co
 	return Effects;
 }
 
+UTexture* UDynamicMaterialEditorSettings::GetDefaultTextureForSlot(EDMMaterialPropertyType InProperty) const
+{
+	if (const TSoftObjectPtr<UTexture>* OverridePtr = OverriddenDefaultSlotTextures.Find(InProperty))
+	{
+		if (UTexture* OverrideTexture = OverridePtr->LoadSynchronous())
+		{
+			return OverrideTexture;
+		}
+	}
+	
+	if (const TSoftObjectPtr<UTexture>* OverridePtr = UE::DynamicMaterialEditor::Private::BaseDefaultSlotProperties.Find(InProperty))
+	{
+		if (UTexture* OverrideTexture = OverridePtr->LoadSynchronous())
+		{
+			return OverrideTexture;
+		}
+	}
+
+	return nullptr;
+}
+
 #undef LOCTEXT_NAMESPACE
diff --git a/Engine/Plugins/Experimental/DynamicMaterial/Source/DynamicMaterialEditor/Private/DynamicMaterialEditorSettings.h b/Engine/Plugins/Experimental/DynamicMaterial/Source/DynamicMaterialEditor/Private/DynamicMaterialEditorSettings.h
--- a/Engine/Plugins/Experimental/DynamicMaterial/Source/DynamicMaterialEditor/Private/DynamicMaterialEditorSettings.h
+++ b/Engine/Plugins/Experimental/DynamicMaterial/Source/DynamicMaterialEditor/Private/DynamicMaterialEditorSettings.h
@@ -172,8 +172,8 @@ public:
 	UPROPERTY(Config, EditAnywhere, BlueprintReadWrite, Category = "Preview")
 	TSoftObjectPtr<UTexture> DefaultMask;
 
-	UPROPERTY(Config, EditAnywhere, BlueprintReadWrite, Category = "Channels", meta = (EditFixedSize))
-	TMap<EDMMaterialPropertyType, TSoftObjectPtr<UTexture>> DefaultSlotTextures;
+	UPROPERTY(Config, EditAnywhere, BlueprintReadWrite, Category = "Channels")
+	TMap<EDMMaterialPropertyType, TSoftObjectPtr<UTexture>> OverriddenDefaultSlotTextures;
 
 	/*
 	 * Add paths to search for custom effects.
@@ -208,4 +208,6 @@ public:
 	float GetSpinboxValueChangeMultiplier(const FModifierKeysState& InModifierKeys) const;
 
 	TArray<FDMMaterialEffectList> GetEffectList() const;
+
+	UTexture* GetDefaultTextureForSlot(EDMMaterialPropertyType InProperty) const;
 };
diff --git a/Engine/Plugins/Experimental/DynamicMaterial/Source/DynamicMaterialEditor/Private/Model/DynamicMaterialModelEditorOnlyData.cpp b/Engine/Plugins/Experimental/DynamicMaterial/Source/DynamicMaterialEditor/Private/Model/DynamicMaterialModelEditorOnlyData.cpp
--- a/Engine/Plugins/Experimental/DynamicMaterial/Source/DynamicMaterialEditor/Private/Model/DynamicMaterialModelEditorOnlyData.cpp
+++ b/Engine/Plugins/Experimental/DynamicMaterial/Source/DynamicMaterialEditor/Private/Model/DynamicMaterialModelEditorOnlyData.cpp
@@ -1166,6 +1166,41 @@ UDMMaterialSlot* UDynamicMaterialModelEditorOnlyData::AddSlotForMaterialProperty
 		return ExistingSlot;
 	}
 
+	switch (InType)
+	{
+		case EDMMaterialPropertyType::BaseColor:
+			if (UDMMaterialSlot* ExistingSlot = GetSlotForMaterialProperty(EDMMaterialPropertyType::EmissiveColor))
+			{
+				SwapSlotMaterialProperty(EDMMaterialPropertyType::EmissiveColor, EDMMaterialPropertyType::BaseColor);
+				return ExistingSlot;
+			}
+			break;
+
+		case EDMMaterialPropertyType::EmissiveColor:
+			if (UDMMaterialSlot* ExistingSlot = GetSlotForMaterialProperty(EDMMaterialPropertyType::BaseColor))
+			{
+				SwapSlotMaterialProperty(EDMMaterialPropertyType::BaseColor, EDMMaterialPropertyType::EmissiveColor);
+				return ExistingSlot;
+			}
+			break;
+
+		case EDMMaterialPropertyType::Opacity:
+			if (UDMMaterialSlot* ExistingSlot = GetSlotForMaterialProperty(EDMMaterialPropertyType::OpacityMask))
+			{
+				SwapSlotMaterialProperty(EDMMaterialPropertyType::OpacityMask, EDMMaterialPropertyType::Opacity);
+				return ExistingSlot;
+			}
+			break;
+
+		case EDMMaterialPropertyType::OpacityMask:
+			if (UDMMaterialSlot* ExistingSlot = GetSlotForMaterialProperty(EDMMaterialPropertyType::Opacity))
+			{
+				SwapSlotMaterialProperty(EDMMaterialPropertyType::Opacity, EDMMaterialPropertyType::OpacityMask);
+				return ExistingSlot;
+			}
+			break;
+	}
+
 	UDMMaterialSlot* NewSlot = NewObject<UDMMaterialSlot>(this, NAME_None, RF_Transactional);
 	check(NewSlot);
 
diff --git a/Engine/Plugins/Experimental/DynamicMaterial/Source/DynamicMaterialEditor/Private/Slate/SDMSlot.cpp b/Engine/Plugins/Experimental/DynamicMaterial/Source/DynamicMaterialEditor/Private/Slate/SDMSlot.cpp
--- a/Engine/Plugins/Experimental/DynamicMaterial/Source/DynamicMaterialEditor/Private/Slate/SDMSlot.cpp
+++ b/Engine/Plugins/Experimental/DynamicMaterial/Source/DynamicMaterialEditor/Private/Slate/SDMSlot.cpp
@@ -1271,13 +1271,13 @@ UDMMaterialLayerObject* SDMSlot::AddNewLayer(UDMMaterialStage* InNewBaseStage, U
 		{
 			TArray<EDMMaterialPropertyType> SlotProperties = ModelEditorOnlyData->GetMaterialPropertiesForSlot(Slot);
 
-			if (!SlotProperties.IsEmpty())
+			if (ensureMsgf(!SlotProperties.IsEmpty(), TEXT("Cannot find material property.")))
 			{
 				MaterialProperty = SlotProperties[0];
 			}
 			else
 			{
-				check("Cannot find material property.");
+				return nullptr;
 			}
 		}
 		else
@@ -1285,6 +1285,11 @@ UDMMaterialLayerObject* SDMSlot::AddNewLayer(UDMMaterialStage* InNewBaseStage, U
 			MaterialProperty = Slot->GetLayers().Last()->GetMaterialProperty();
 		}
 
+		if (!ensureMsgf(MaterialProperty != EDMMaterialPropertyType::None, TEXT("Could not find material property.")))
+		{
+			return nullptr;
+		}
+
 		if (!InNewMaskStage)
 		{
 			Slot->AddLayer(MaterialProperty, InNewBaseStage);
