Material Designer
- Added an export button to the material designer ui.

diff --git a/Engine/Plugins/Experimental/DynamicMaterial/Source/DynamicMaterialEditor/Private/Menus/DMToolBarMenus.cpp b/Engine/Plugins/Experimental/DynamicMaterial/Source/DynamicMaterialEditor/Private/Menus/DMToolBarMenus.cpp
--- a/Engine/Plugins/Experimental/DynamicMaterial/Source/DynamicMaterialEditor/Private/Menus/DMToolBarMenus.cpp
+++ b/Engine/Plugins/Experimental/DynamicMaterial/Source/DynamicMaterialEditor/Private/Menus/DMToolBarMenus.cpp
@@ -556,11 +556,14 @@ namespace UE::DynamicMaterialEditor::Private
 			FContentBrowserModule& ContentBrowserModule = FModuleManager::LoadModuleChecked<FContentBrowserModule>("ContentBrowser");
 			FString SaveObjectPath = ContentBrowserModule.Get().CreateModalSaveAssetDialog(SaveAssetDialogConfig);
 
-			UDMBlueprintFunctionLibrary::ExportMaterialInstance(MaterialInstance->GetMaterialModel(), SaveObjectPath);
-
-			if (FEngineAnalytics::IsAvailable())
+			if (!SaveObjectPath.IsEmpty())
 			{
-				FEngineAnalytics::GetProvider().RecordEvent(TEXT("Editor.Usage.MaterialDesigner.ExportedMaterialInstance"));
+				UDMBlueprintFunctionLibrary::ExportMaterialInstance(MaterialInstance->GetMaterialModel(), SaveObjectPath);
+
+				if (FEngineAnalytics::IsAvailable())
+				{
+					FEngineAnalytics::GetProvider().RecordEvent(TEXT("Editor.Usage.MaterialDesigner.ExportedMaterialInstance"));
+				}
 			}
 		}
 	}
diff --git a/Engine/Plugins/Experimental/DynamicMaterial/Source/DynamicMaterialEditor/Private/Slate/SDMToolBar.cpp b/Engine/Plugins/Experimental/DynamicMaterial/Source/DynamicMaterialEditor/Private/Slate/SDMToolBar.cpp
--- a/Engine/Plugins/Experimental/DynamicMaterial/Source/DynamicMaterialEditor/Private/Slate/SDMToolBar.cpp
+++ b/Engine/Plugins/Experimental/DynamicMaterial/Source/DynamicMaterialEditor/Private/Slate/SDMToolBar.cpp
@@ -2,16 +2,20 @@
 
 #include "Slate/SDMToolBar.h"
 #include "AssetToolsModule.h"
+#include "ContentBrowserModule.h"
 #include "DMBlueprintFunctionLibrary.h"
 #include "DMEDefs.h"
 #include "DynamicMaterialEditorModule.h"
 #include "DynamicMaterialEditorSettings.h"
 #include "DynamicMaterialEditorStyle.h"
 #include "DynamicMaterialModule.h"
+#include "EngineAnalytics.h"
 #include "GameFramework/Actor.h"
 #include "IAssetTools.h"
+#include "IContentBrowserSingleton.h"
 #include "Model/DynamicMaterialModel.h"
 #include "SlateOptMacros.h"
+#include "Material/DynamicMaterialInstance.h"
 #include "Styling/AppStyle.h"
 #include "Styling/StyleColors.h"
 #include "Widgets/Colors/SColorBlock.h"
@@ -124,7 +128,7 @@ TSharedRef<SWidget> SDMToolBar::CreateToolBarEntries()
 					CreateSlotsComboBoxWidget()
 				]
 			]
-		]		
+		]
 		
 		+ SHorizontalBox::Slot()
 		.AutoWidth()
@@ -135,7 +139,7 @@ TSharedRef<SWidget> SDMToolBar::CreateToolBarEntries()
 			SNew(SButton)
 			.ContentPadding(GetDefaultToolBarButtonContentPadding())
 			.ButtonStyle(FDynamicMaterialEditorStyle::Get(), "HoverHintOnly")
-			.ToolTipText(LOCTEXT("MaterialDesignerFollowSelectionTooltip", "Toggles whether the Material Designer display will change selecting new objects and actors."))
+			.ToolTipText(LOCTEXT("MaterialDesignerFollowSelectionTooltip", "Toggles whether the Material Designer display will change when selecting new objects and actors."))
 			.OnClicked(this, &SDMToolBar::OnFollowSelectionButtonClicked)
 			[
 				SNew(SImage)
@@ -145,6 +149,25 @@ TSharedRef<SWidget> SDMToolBar::CreateToolBarEntries()
 			]
 		]
 		
+		+ SHorizontalBox::Slot()
+		.AutoWidth()
+		.HAlign(HAlign_Right)
+		.VAlign(VAlign_Top)
+		.Padding(5.0f, 0.0f, 0.0f, 0.0f)
+		[
+			SNew(SButton)
+			.ContentPadding(GetDefaultToolBarButtonContentPadding())
+			.ButtonStyle(FDynamicMaterialEditorStyle::Get(), "HoverHintOnly")
+			.ToolTipText(LOCTEXT("ExportMaterialInstance", "Export Material Designer Instance"))
+			.Visibility(this, &SDMToolBar::GetExportMaterialInstanceButtonVisibility)
+			.OnClicked(this, &SDMToolBar::OnExportMaterialInstanceButtonClicked)
+			[
+				SNew(SImage)
+				.Image(FAppStyle::Get().GetBrush(TEXT("Icons.Toolbar.Export")))
+				.DesiredSizeOverride(GetDefaultToolBarButtonSize())
+			]
+		]
+		
 		+ SHorizontalBox::Slot()
 		.AutoWidth()
 		.HAlign(HAlign_Right)
@@ -375,4 +398,61 @@ FReply SDMToolBar::OnFollowSelectionButtonClicked()
 	return FReply::Handled();
 }
 
+EVisibility SDMToolBar::GetExportMaterialInstanceButtonVisibility() const
+{
+	UDynamicMaterialModel* MaterialModel = MaterialModelWeak.Get();
+
+	if (!MaterialModel)
+	{
+		return EVisibility::Collapsed;
+	}
+
+	UDynamicMaterialInstance* MaterialInstance = MaterialModel->GetDynamicMaterialInstance();
+
+	if (!MaterialInstance)
+	{
+		return EVisibility::Collapsed;
+	}
+
+	return EVisibility::Visible;
+}
+
+FReply SDMToolBar::OnExportMaterialInstanceButtonClicked()
+{
+	UDynamicMaterialModel* MaterialModel = MaterialModelWeak.Get();
+
+	if (!MaterialModel)
+	{
+		return FReply::Handled();
+	}
+
+	UDynamicMaterialInstance* MaterialInstance = MaterialModel->GetDynamicMaterialInstance();
+
+	if (!MaterialInstance)
+	{
+		return FReply::Handled();
+	}
+
+	FSaveAssetDialogConfig SaveAssetDialogConfig;
+	SaveAssetDialogConfig.DialogTitleOverride = LOCTEXT("SaveAssetDialogTitle", "Save Asset As");
+	SaveAssetDialogConfig.DefaultPath = "/Game";
+	SaveAssetDialogConfig.DefaultAssetName = MaterialInstance->GetName();
+	SaveAssetDialogConfig.ExistingAssetPolicy = ESaveAssetDialogExistingAssetPolicy::AllowButWarn;
+
+	FContentBrowserModule& ContentBrowserModule = FModuleManager::LoadModuleChecked<FContentBrowserModule>("ContentBrowser");
+	FString SaveObjectPath = ContentBrowserModule.Get().CreateModalSaveAssetDialog(SaveAssetDialogConfig);
+
+	if (!SaveObjectPath.IsEmpty())
+	{
+		UDMBlueprintFunctionLibrary::ExportMaterialInstance(MaterialInstance->GetMaterialModel(), SaveObjectPath);
+
+		if (FEngineAnalytics::IsAvailable())
+		{
+			FEngineAnalytics::GetProvider().RecordEvent(TEXT("Editor.Usage.MaterialDesigner.ExportedMaterialInstance"));
+		}
+	}
+
+	return FReply::Handled();
+}
+
 #undef LOCTEXT_NAMESPACE
diff --git a/Engine/Plugins/Experimental/DynamicMaterial/Source/DynamicMaterialEditor/Private/Slate/SDMToolBar.h b/Engine/Plugins/Experimental/DynamicMaterial/Source/DynamicMaterialEditor/Private/Slate/SDMToolBar.h
--- a/Engine/Plugins/Experimental/DynamicMaterial/Source/DynamicMaterialEditor/Private/Slate/SDMToolBar.h
+++ b/Engine/Plugins/Experimental/DynamicMaterial/Source/DynamicMaterialEditor/Private/Slate/SDMToolBar.h
@@ -71,4 +71,7 @@ protected:
 	const FSlateBrush* GetFollowSelectionBrush() const;
 	FSlateColor GetFollowSelectionColor() const;
 	FReply OnFollowSelectionButtonClicked();
+
+	EVisibility GetExportMaterialInstanceButtonVisibility() const;
+	FReply OnExportMaterialInstanceButtonClicked();
 };
