Material Designer
- Added "global scale" that affects all texture uvs.

#virtualized

diff --git a/Engine/Plugins/Experimental/DynamicMaterial/Source/DynamicMaterial/Private/Model/DynamicMaterialModel.cpp b/Engine/Plugins/Experimental/DynamicMaterial/Source/DynamicMaterial/Private/Model/DynamicMaterialModel.cpp
--- a/Engine/Plugins/Experimental/DynamicMaterial/Source/DynamicMaterial/Private/Model/DynamicMaterialModel.cpp
+++ b/Engine/Plugins/Experimental/DynamicMaterial/Source/DynamicMaterial/Private/Model/DynamicMaterialModel.cpp
@@ -4,6 +4,7 @@
 #include "Components/DMMaterialParameter.h"
 #include "Components/DMMaterialValue.h"
 #include "Components/MaterialValues/DMMaterialValueFloat1.h"
+#include "Components/MaterialValues/DMMaterialValueFloat2.h"
 #include "DMComponentPath.h"
 #include "DMDefs.h"
 #include "DMValueDefinition.h"
@@ -18,7 +19,10 @@
 
 const FString UDynamicMaterialModel::ValuesPathToken          = FString(TEXT("Values"));
 const FString UDynamicMaterialModel::ParametersPathToken      = FString(TEXT("Parameters"));
+const FName UDynamicMaterialModel::GlobalOpacityPropertyName  = FName("GlobalOpacityValue");
 const FName UDynamicMaterialModel::GlobalOpacityParameterName = FName("GlobalOpacity");
+const FName UDynamicMaterialModel::GlobalScalePropertyName    = FName("GlobalScaleValue");
+const FName UDynamicMaterialModel::GlobalScaleParameterName   = FName("GlobalScale");
 
 UDynamicMaterialModel::UDynamicMaterialModel()
 {
@@ -32,6 +36,7 @@ UDynamicMaterialModel::UDynamicMaterialModel()
 	GlobalOpacityValue = CreateDefaultSubobject<UDMMaterialValueFloat1>(GlobalOpacityParameterName);
 
 #if WITH_EDITOR
+	GlobalOpacityValue->SetValueRange({0.f, 1.f});
 	GlobalOpacityValue->SetDefaultValue(1.f);
 	GlobalOpacityValue->ApplyDefaultValue();
 #endif
@@ -41,6 +46,19 @@ UDynamicMaterialModel::UDynamicMaterialModel()
 	GlobalOpacityValue->Parameter = GlobalOpacityParameter;
 
 	ParameterMap.Add(GlobalOpacityParameterName, GlobalOpacityParameter);
+
+	GlobalScaleValue = CreateDefaultSubobject<UDMMaterialValueFloat2>(GlobalScaleParameterName);
+
+#if WITH_EDITOR
+	GlobalScaleValue->SetDefaultValue(FVector2D(1.0, 1.0));
+	GlobalScaleValue->ApplyDefaultValue();
+#endif
+
+	GlobalScaleParameter = CreateDefaultSubobject<UDMMaterialParameter>("GlobalScaleParameter");
+	GlobalScaleParameter->ParameterName = GlobalScaleParameterName;
+	GlobalScaleValue->Parameter = GlobalScaleParameter;
+
+	ParameterMap.Add(GlobalScaleParameterName, GlobalScaleParameter);
 }
 
 void UDynamicMaterialModel::SetDynamicMaterialInstance(UDynamicMaterialInstance* InDynamicMaterialInstance)
@@ -350,7 +368,7 @@ void UDynamicMaterialModel::PostLoad()
 
 	SetFlags(RF_Transactional);
 
-	FixGlobalOpacityVars();
+	FixGlobalVars();
 
 	IDynamicMaterialModelEditorOnlyDataInterface* ModelEditorOnlyData = GetEditorOnlyData();
 
@@ -377,7 +395,7 @@ void UDynamicMaterialModel::PostEditImport()
 {
 	Super::PostEditImport();
 
-	FixGlobalOpacityVars();
+	FixGlobalVars();
 	PostEditorDuplicate();
 	ReinitComponents();
 
@@ -391,7 +409,7 @@ void UDynamicMaterialModel::PostDuplicate(bool bDuplicateForPIE)
 {
 	Super::PostDuplicate(bDuplicateForPIE);
 
-	FixGlobalOpacityVars();
+	FixGlobalVars();
 	PostEditorDuplicate();
 	ReinitComponents();
 
@@ -493,7 +511,7 @@ void UDynamicMaterialModel::ReinitComponents()
 	}
 }
 
-void UDynamicMaterialModel::FixGlobalOpacityVars()
+void UDynamicMaterialModel::FixGlobalVars()
 {
 	if (const TWeakObjectPtr<UDMMaterialParameter>* ParameterPtr = ParameterMap.Find(GlobalOpacityParameterName))
 	{
@@ -513,19 +531,47 @@ void UDynamicMaterialModel::FixGlobalOpacityVars()
 		}
 	}
 
-	if (!GlobalOpacityValue || GlobalOpacityParameter == GlobalOpacityValue->Parameter)
+	if (GlobalOpacityValue && GlobalOpacityParameter != GlobalOpacityValue->Parameter)
 	{
-		return;
+		if (GUndo)
+		{
+			GlobalOpacityParameter->Modify();
+			GlobalOpacityValue->Modify();
+		}
+
+		GlobalOpacityParameter->ParameterName = GlobalOpacityParameterName;
+		GlobalOpacityValue->Parameter = GlobalOpacityParameter;
 	}
 
-	if (GUndo)
+	if (const TWeakObjectPtr<UDMMaterialParameter>* ParameterPtr = ParameterMap.Find(GlobalScaleParameterName))
 	{
-		GlobalOpacityParameter->Modify();
-		GlobalOpacityValue->Modify();
+		UDMMaterialParameter* Parameter = ParameterPtr->Get();
+
+		if (!Parameter || Parameter->HasAnyFlags(RF_ArchetypeObject)
+			|| (GlobalScaleParameter && Parameter != GlobalScaleParameter))
+		{
+			if (GlobalScaleParameter)
+			{
+				ParameterMap[GlobalScaleParameterName] = GlobalScaleParameter;
+			}
+			else
+			{
+				ParameterMap.Remove(GlobalScaleParameterName);
+			}
+		}
 	}
 
-	GlobalOpacityParameter->ParameterName = GlobalOpacityParameterName;
-	GlobalOpacityValue->Parameter = GlobalOpacityParameter;
+	if (GlobalScaleValue && GlobalScaleParameter != GlobalScaleValue->Parameter)
+	{
+		if (GUndo)
+		{
+			GlobalScaleParameter->Modify();
+			GlobalScaleValue->Modify();
+		}
+
+		GlobalScaleParameter->ParameterName = GlobalScaleParameterName;
+		GlobalScaleValue->Parameter = GlobalScaleParameter;
+	}
 }
 
 FName UDynamicMaterialModel::CreateUniqueParameterName(FName InBaseName)
diff --git a/Engine/Plugins/Experimental/DynamicMaterial/Source/DynamicMaterial/Public/Model/DynamicMaterialModel.h b/Engine/Plugins/Experimental/DynamicMaterial/Source/DynamicMaterial/Public/Model/DynamicMaterialModel.h
--- a/Engine/Plugins/Experimental/DynamicMaterial/Source/DynamicMaterial/Public/Model/DynamicMaterialModel.h
+++ b/Engine/Plugins/Experimental/DynamicMaterial/Source/DynamicMaterial/Public/Model/DynamicMaterialModel.h
@@ -17,6 +17,7 @@ class UDMMaterialComponent;
 class UDMMaterialParameter;
 class UDMMaterialValue;
 class UDMMaterialValueFloat1;
+class UDMMaterialValueFloat2;
 class UDMTextureUV;
 class UDynamicMaterialInstance;
 enum class EDMMaterialPropertyType : uint8;
@@ -42,7 +43,10 @@ class DYNAMICMATERIAL_API UDynamicMaterialModel : public UObject
 public:
 	static const FString ValuesPathToken;
 	static const FString ParametersPathToken;
+	static const FName GlobalOpacityPropertyName;
 	static const FName GlobalOpacityParameterName;
+	static const FName GlobalScalePropertyName;
+	static const FName GlobalScaleParameterName;
 
 	UDynamicMaterialModel();
 
@@ -57,6 +61,9 @@ public:
 	UFUNCTION(BlueprintPure, Category = "Material Designer")
 	UDMMaterialValueFloat1* GetGlobalOpacityValue() const { return GlobalOpacityValue; }
 
+	UFUNCTION(BlueprintPure, Category = "Material Designer")
+	UDMMaterialValueFloat2* GetGlobalScaleValue() const { return GlobalScaleValue; }
+
 	UFUNCTION(BlueprintCallable, Category = "Material Designer")
 	UDMMaterialComponent* GetComponentByPath(const FString& InPath) const;
 
@@ -153,6 +160,12 @@ protected:
 	UPROPERTY(VisibleInstanceOnly, BlueprintReadOnly, Category = "Material Designer")
 	TObjectPtr<UDMMaterialParameter> GlobalOpacityParameter;
 
+	UPROPERTY(VisibleInstanceOnly, BlueprintReadOnly, Category = "Material Designer")
+	TObjectPtr<UDMMaterialValueFloat2> GlobalScaleValue;
+
+	UPROPERTY(VisibleInstanceOnly, BlueprintReadOnly, Category = "Material Designer")
+	TObjectPtr<UDMMaterialParameter> GlobalScaleParameter;
+
 	UPROPERTY(VisibleInstanceOnly, DuplicateTransient, TextExportTransient, Category = "Material Designer")
 	TMap<FName, TWeakObjectPtr<UDMMaterialParameter>> ParameterMap;
 
@@ -172,7 +185,7 @@ protected:
 
 	void ReinitComponents();
 
-	void FixGlobalOpacityVars();
+	void FixGlobalVars();
 #endif
 
 	/**
diff --git a/Engine/Plugins/Experimental/DynamicMaterial/Source/DynamicMaterialEditor/Private/Menus/DMToolBarMenus.cpp b/Engine/Plugins/Experimental/DynamicMaterial/Source/DynamicMaterialEditor/Private/Menus/DMToolBarMenus.cpp
--- a/Engine/Plugins/Experimental/DynamicMaterial/Source/DynamicMaterialEditor/Private/Menus/DMToolBarMenus.cpp
+++ b/Engine/Plugins/Experimental/DynamicMaterial/Source/DynamicMaterialEditor/Private/Menus/DMToolBarMenus.cpp
@@ -56,6 +56,7 @@ namespace UE::DynamicMaterialEditor::Private
 
 					if (EditorWidget.IsValid())
 					{
+						EditorWidget->RefreshGlobalDetailsView();
 						EditorWidget->RefreshSlotPickerList();
 					}
 				}
@@ -122,6 +123,7 @@ namespace UE::DynamicMaterialEditor::Private
 
 					if (EditorWidget.IsValid())
 					{
+						EditorWidget->RefreshGlobalDetailsView();
 						EditorWidget->RefreshSlotPickerList();
 					}
 				}
@@ -204,6 +206,7 @@ namespace UE::DynamicMaterialEditor::Private
 
 					if (EditorWidget.IsValid())
 					{
+						EditorWidget->RefreshGlobalDetailsView();
 						EditorWidget->RefreshSlotPickerList();
 					}
 				}
diff --git a/Engine/Plugins/Experimental/DynamicMaterial/Source/DynamicMaterialEditor/Private/Model/DynamicMaterialModelEditorOnlyData.cpp b/Engine/Plugins/Experimental/DynamicMaterial/Source/DynamicMaterialEditor/Private/Model/DynamicMaterialModelEditorOnlyData.cpp
--- a/Engine/Plugins/Experimental/DynamicMaterial/Source/DynamicMaterialEditor/Private/Model/DynamicMaterialModelEditorOnlyData.cpp
+++ b/Engine/Plugins/Experimental/DynamicMaterial/Source/DynamicMaterialEditor/Private/Model/DynamicMaterialModelEditorOnlyData.cpp
@@ -1166,6 +1166,8 @@ UDMMaterialSlot* UDynamicMaterialModelEditorOnlyData::AddSlotForMaterialProperty
 		return ExistingSlot;
 	}
 
+	// BaseColor and Emissive are mutually exclusive so if something tries to add one of them, the other must
+	// be checked. If it is found, it is converted and returned. The same goes for Opacity and OpacityMask.
 	switch (InType)
 	{
 		case EDMMaterialPropertyType::BaseColor:
diff --git a/Engine/Plugins/Experimental/DynamicMaterial/Source/DynamicMaterialEditor/Private/Slate/SDMEditor.cpp b/Engine/Plugins/Experimental/DynamicMaterial/Source/DynamicMaterialEditor/Private/Slate/SDMEditor.cpp
--- a/Engine/Plugins/Experimental/DynamicMaterial/Source/DynamicMaterialEditor/Private/Slate/SDMEditor.cpp
+++ b/Engine/Plugins/Experimental/DynamicMaterial/Source/DynamicMaterialEditor/Private/Slate/SDMEditor.cpp
@@ -8,7 +8,10 @@
 #include "Components/DMMaterialValue.h"
 #include "Components/MaterialStageExpressions/DMMSETextureSample.h"
 #include "Components/MaterialValues/DMMaterialValueFloat1.h"
+#include "Components/MaterialValues/DMMaterialValueFloat2.h"
 #include "Components/PrimitiveComponent.h"
+#include "CustomDetailsViewArgs.h"
+#include "CustomDetailsViewModule.h"
 #include "DetailLayoutBuilder.h"
 #include "DMBlueprintFunctionLibrary.h"
 #include "DMPrivate.h"
@@ -22,13 +25,16 @@
 #include "Framework/Commands/GenericCommands.h"
 #include "Framework/Commands/InputChord.h"
 #include "HAL/PlatformApplicationMisc.h"
+#include "ICustomDetailsView.h"
 #include "IDetailTreeNode.h"
 #include "IPropertyRowGenerator.h"
+#include "Items/ICustomDetailsViewItem.h"
 #include "MaterialDomain.h"
 #include "Menus/DMToolBarMenus.h"
 #include "Misc/CoreDelegates.h"
 #include "Model/DynamicMaterialModel.h"
 #include "Model/DynamicMaterialModelEditorOnlyData.h"
+#include "Properties/Editors/SDMPropertyEditFloat2Value.h"
 #include "PropertyHandle.h"
 #include "Slate/Properties/Editors/SDMPropertyEditOpacity.h"
 #include "Slate/Properties/SDMMaterialParameters.h"
@@ -204,9 +210,9 @@ void SDMEditor::SetActiveSlotIndex(int InSlotIndex)
 	}
 }
 
-FDMPropertyHandle SDMEditor::GetPropertyHandle(const SWidget* InOwner, UDMMaterialComponent* InComponent, const FName& InPropertyName)
+FDMPropertyHandle SDMEditor::GetPropertyHandle(const SWidget* InOwningWidget, UObject* InObject, const FName& InPropertyName)
 {
-	TArray<FDMPropertyHandle>& PropertyHandles = PropertyHandleMap.FindOrAdd(InOwner);
+	TArray<FDMPropertyHandle>& PropertyHandles = PropertyHandleMap.FindOrAdd(InOwningWidget);
 
 	for (const FDMPropertyHandle& ExistingHandle : PropertyHandles)
 	{
@@ -215,14 +221,14 @@ FDMPropertyHandle SDMEditor::GetPropertyHandle(const SWidget* InOwner, UDMMateri
 			TArray<UObject*> Outers;
 			ExistingHandle.PropertyHandle->GetOuterObjects(Outers);
 
-			if (Outers.IsEmpty() == false && Outers[0] == InComponent)
+			if (Outers.IsEmpty() == false && Outers[0] == InObject)
 			{
 				return ExistingHandle;
 			}
 		}
 	}
 
-	if (TSharedPtr<IPropertyRowGenerator> PropertyRowGenerator = UE::DynamicMaterialEditor::Private::SearchForGenerator(PropertyHandles, InComponent))
+	if (TSharedPtr<IPropertyRowGenerator> PropertyRowGenerator = UE::DynamicMaterialEditor::Private::SearchForGenerator(PropertyHandles, InObject))
 	{
 		FDMPropertyHandle PropertyHandle;
 		PropertyHandle.PropertyRowGenerator = PropertyRowGenerator;
@@ -237,29 +243,32 @@ FDMPropertyHandle SDMEditor::GetPropertyHandle(const SWidget* InOwner, UDMMateri
 		return PropertyHandle;
 	}
 
-	FDMPropertyHandle NewHandle = CreatePropertyHandle(InOwner, InComponent, InPropertyName);
+	FDMPropertyHandle NewHandle = CreatePropertyHandle(InOwningWidget, InObject, InPropertyName);
 	PropertyHandles.Add(NewHandle);
 
 	return NewHandle;
 }
 
-void SDMEditor::ClearPropertyHandles(const SWidget* InOwner)
+void SDMEditor::ClearPropertyHandles(const SWidget* InOwningWidget)
 {
-	PropertyHandleMap.Remove(InOwner);
+	PropertyHandleMap.Remove(InOwningWidget);
 }
 
-FDMPropertyHandle SDMEditor::CreatePropertyHandle(const void* InOwner, UDMMaterialComponent* InComponent,
-	const FName& InPropertyName)
+FDMPropertyHandle SDMEditor::CreatePropertyHandle(const void* InOwningWidget, UObject* InObject, const FName& InPropertyName)
 {
 	FDMPropertyHandle PropertyHandle;
 
 	FPropertyEditorModule& PropertyEditor = FModuleManager::Get().LoadModuleChecked<FPropertyEditorModule>(TEXT("PropertyEditor"));
 
 	FPropertyRowGeneratorArgs RowGeneratorArgs;
-	RowGeneratorArgs.NotifyHook = InComponent;
+
+	if (UDMMaterialComponent* Component = Cast<UDMMaterialComponent>(InObject))
+	{
+		RowGeneratorArgs.NotifyHook = Component;
+	}
 
 	PropertyHandle.PropertyRowGenerator = PropertyEditor.CreatePropertyRowGenerator(RowGeneratorArgs);
-	PropertyHandle.PropertyRowGenerator->SetObjects({InComponent});
+	PropertyHandle.PropertyRowGenerator->SetObjects({InObject});
 
 	if (const TSharedPtr<IDetailTreeNode> FoundTreeNode = UE::DynamicMaterialEditor::Private::SearchGeneratorForNode(
 		PropertyHandle.PropertyRowGenerator.ToSharedRef(), InPropertyName))
@@ -562,9 +571,8 @@ TSharedRef<SWidget> SDMEditor::CreateMainLayout()
 			SAssignNew(GlobalOpacityContainer, SBox)
 			.HAlign(HAlign_Fill)
 			.VAlign(VAlign_Fill)
-			.Visibility(this, &SDMEditor::GetGlobalOpacityVisibility)
 			[
-				CreateGlobalOpacityWidget()
+				CreateGlobalDetailsView()
 			]
 		]
 		+ SVerticalBox::Slot()
@@ -599,59 +607,109 @@ TSharedRef<SWidget> SDMEditor::CreateMainLayout()
 		];
 }
 
-TSharedRef<SWidget> SDMEditor::CreateGlobalOpacityWidget()
+TSharedRef<SWidget> SDMEditor::CreateGlobalDetailsView()
 {
-	UDMMaterialValueFloat1* OpacityValue = nullptr;
+	UDynamicMaterialModel* MaterialModel = MaterialModelWeak.Get();
 
-	if (UDynamicMaterialModel* MaterialModel = MaterialModelWeak.Get())
+	if (!MaterialModel)
 	{
-		OpacityValue = MaterialModel->GetGlobalOpacityValue();
+		return SNullWidget::NullWidget;
 	}
 
-	if (!OpacityValue)
+	UDMMaterialValueFloat1* OpacityValue = MaterialModel->GetGlobalOpacityValue();
+	UDMMaterialValueFloat2* ScaleValue = MaterialModel->GetGlobalScaleValue();
+
+	if (!OpacityValue || !ScaleValue)
 	{
 		return SNullWidget::NullWidget;
 	}
 
-	TSharedRef<SDMPropertyEdit> GlobalOpacityWidget = SNew(SDMPropertyEditOpacity, SharedThis(this), OpacityValue);
+	FCustomDetailsViewArgs Args;
+	Args.KeyframeHandler = nullptr;
+	Args.bAllowGlobalExtensions = true;
+	Args.bAllowResetToDefault = true;
+	Args.bShowCategories = false;
 
-	TSharedRef<SWidget> GlobalOpacityButtons = SDMComponentEdit::CreateExtensionButtons(SharedThis(this), OpacityValue, UDMMaterialValue::ValueName, true, FSimpleDelegate());
+	if (const UWorld* const World = MaterialModel->GetWorld())
+	{
+		if (const UDMWorldSubsystem* const WorldSubsystem = World->GetSubsystem<UDMWorldSubsystem>())
+		{
+			Args.KeyframeHandler = WorldSubsystem->GetKeyframeHandler();
+		}
+	}
 
-	TSharedRef<SWidget> RowLabel = SNew(SBox)
-		.Padding(0.f, 0.f, 5.f, 0.f)
-		[
-			SNew(STextBlock)
-				.Text(LOCTEXT("GlobalOpacity", "Global Opacity"))
-		];
+	TSharedRef<ICustomDetailsView> DetailsView = ICustomDetailsViewModule::Get().CreateCustomDetailsView(Args);
+	FCustomDetailsViewItemId RootId = DetailsView->GetRootItem()->GetItemId();
 
-	RowLabel->SetOnMouseButtonDown(FPointerEventHandler::CreateStatic(&SDMPropertyEdit::CreateRightClickDetailsMenu, GlobalOpacityWidget.ToWeakPtr()));
+	static const FName UIMin = FName("UIMin");
+	static const FName UIMax = FName("UIMax");
+	static const FName ClampMin = FName("ClampMin");
+	static const FName ClampMax = FName("ClampMax");
 
-	TSharedRef<SWidget> Row =
-		SNew(SHorizontalBox)
-		+ SHorizontalBox::Slot()
-		.AutoWidth()
-		.HAlign(HAlign_Left)
-		.VAlign(VAlign_Top)
-		.Padding(10.0f, 5.0f)
-		[
-			RowLabel
-		]
-		+ SHorizontalBox::Slot()
-		.FillWidth(1.0f)
-		.HAlign(HAlign_Right)
-		.VAlign(VAlign_Top)
-		[
-			GlobalOpacityWidget
-		]
-		+ SHorizontalBox::Slot()
-		.AutoWidth()
-		.HAlign(HAlign_Left)
-		.VAlign(VAlign_Top)
-		[
-			GlobalOpacityButtons
-		];
+	TArray<FDMPropertyHandle> EditRows;
+
+	if (GetGlobalOpacityVisibility() == EVisibility::Visible)
+	{
+		FDMPropertyHandle& GlobalOpacityHandle = EditRows.Add_GetRef(SDMEditor::GetPropertyHandle(this, OpacityValue, UDMMaterialValue::ValueName));
+
+		GlobalOpacityHandle.PropertyHandle->SetInstanceMetaData(UIMin, FString::SanitizeFloat(OpacityValue->GetValueRange().Min));
+		GlobalOpacityHandle.PropertyHandle->SetInstanceMetaData(ClampMin, FString::SanitizeFloat(OpacityValue->GetValueRange().Min));
+		GlobalOpacityHandle.PropertyHandle->SetInstanceMetaData(UIMax, FString::SanitizeFloat(OpacityValue->GetValueRange().Max));
+		GlobalOpacityHandle.PropertyHandle->SetInstanceMetaData(ClampMax, FString::SanitizeFloat(OpacityValue->GetValueRange().Max));
+
+		GlobalOpacityHandle.NameOverride = LOCTEXT("GlobalOpacity", "Global Opacity");
+	}
+
+	FDMPropertyHandle& GlobalScaleHandle = EditRows.Add_GetRef(SDMEditor::GetPropertyHandle(this, ScaleValue, UDMMaterialValue::ValueName));
+	GlobalScaleHandle.NameOverride = LOCTEXT("GlobalScale", "Global Scale");
+
+	for (const FDMPropertyHandle& EditRow : EditRows)
+	{
+		if (!EditRow.DetailTreeNode)
+		{
+			continue;
+		}
+
+		ECustomDetailsTreeInsertPosition Position = ECustomDetailsTreeInsertPosition::Child;
+
+		if (EditRow.DetailTreeNode->CreatePropertyHandle()->HasMetaData("HighPriority"))
+		{
+			Position = ECustomDetailsTreeInsertPosition::FirstChild;
+		}
+		else if (EditRow.DetailTreeNode->CreatePropertyHandle()->HasMetaData("LowPriority"))
+		{
+			Position = ECustomDetailsTreeInsertPosition::LastChild;
+		}
+
+		TSharedRef<ICustomDetailsViewItem> Item = DetailsView->CreateDetailTreeItem(EditRow.DetailTreeNode.ToSharedRef());
+
+		if (EditRow.NameOverride.IsSet())
+		{
+			Item->SetOverrideWidget(
+				ECustomDetailsViewWidgetType::Name,
+				SNew(STextBlock)
+				.Font(IDetailLayoutBuilder::GetDetailFont())
+				.Text(EditRow.NameOverride.GetValue())
+				.ToolTipText(EditRow.NameToolTipOverride.Get(FText::GetEmpty()))
+			);
+		}
+
+		if (EditRow.DetailTreeNode->CreatePropertyHandle()->HasMetaData("NotKeyframeable"))
+		{
+			Item->SetKeyframeEnabled(false);
+		}
+
+		if (EditRow.ResetToDefaultOverride.IsSet())
+		{
+			Item->SetResetToDefaultOverride(EditRow.ResetToDefaultOverride.GetValue());
+		}
+
+		DetailsView->ExtendTree(RootId, Position, Item);
+	}
+
+	DetailsView->RebuildTree(ECustomDetailsViewBuildType::InstantBuild);
 
-	return Row;
+	return DetailsView;
 }
 
 TSharedRef<SWidget> SDMEditor::CreateParametersArea()
@@ -763,12 +821,12 @@ TSharedPtr<SDMSlot> SDMEditor::GetSlotWidget(UDMMaterialSlot* Slot) const
 	return nullptr;
 }
 
-void SDMEditor::RefreshGlobalOpacitySlider()
+void SDMEditor::RefreshGlobalDetailsView()
 {
 	if (GlobalOpacityContainer.IsValid())
 	{
 		GlobalOpacityContainer->SetContent(SNullWidget::NullWidget);
-		GlobalOpacityContainer->SetContent(CreateGlobalOpacityWidget());
+		GlobalOpacityContainer->SetContent(CreateGlobalDetailsView());
 	}
 }
 
diff --git a/Engine/Plugins/Experimental/DynamicMaterial/Source/DynamicMaterialEditor/Private/Slate/SDMEditor.h b/Engine/Plugins/Experimental/DynamicMaterial/Source/DynamicMaterialEditor/Private/Slate/SDMEditor.h
--- a/Engine/Plugins/Experimental/DynamicMaterial/Source/DynamicMaterialEditor/Private/Slate/SDMEditor.h
+++ b/Engine/Plugins/Experimental/DynamicMaterial/Source/DynamicMaterialEditor/Private/Slate/SDMEditor.h
@@ -44,8 +44,8 @@ public:
 
 	static TSharedRef<SWidget> GetEmptyContent();
 
-	static FDMPropertyHandle GetPropertyHandle(const SWidget* InOwner, UDMMaterialComponent* InComponent, const FName& InPropertyName);
-	static void ClearPropertyHandles(const SWidget* InOwner);
+	static FDMPropertyHandle GetPropertyHandle(const SWidget* InOwningWidget, UObject* InObject, const FName& InPropertyName);
+	static void ClearPropertyHandles(const SWidget* InOwningWidget);
 
 	void Construct(const FArguments& InArgs, TWeakObjectPtr<UDynamicMaterialModel> InModelWeak);
 	virtual ~SDMEditor() override;
@@ -61,7 +61,7 @@ public:
 	const TArray<TSharedRef<SDMSlot>>& GetSlotWidgets() const { return SlotWidgets; }
 	TSharedPtr<SDMSlot> GetSlotWidget(UDMMaterialSlot* Slot) const;
 
-	void RefreshGlobalOpacitySlider();
+	void RefreshGlobalDetailsView();
 	void RefreshParametersList();
 	void RefreshSlotPickerList();
 	void RefreshSlotsList();
@@ -113,7 +113,7 @@ protected:
 
 	static TMap<const SWidget*, TArray<FDMPropertyHandle>> PropertyHandleMap;
 
-	static FDMPropertyHandle CreatePropertyHandle(const void* InOwner, UDMMaterialComponent* InComponent, const FName& InPropertyName);
+	static FDMPropertyHandle CreatePropertyHandle(const void* InOwningWidget, UObject* InObject, const FName& InPropertyName);
 
 	TSharedPtr<FUICommandList> CommandList;
 
@@ -134,7 +134,7 @@ protected:
 	void BindCommands();
 
 	TSharedRef<SWidget> CreateMainLayout();
-	TSharedRef<SWidget> CreateGlobalOpacityWidget();
+	TSharedRef<SWidget> CreateGlobalDetailsView();
 	TSharedRef<SWidget> CreateParametersArea();
 
 	TSharedRef<SWidget> CreateSlotPickerWidget();
