Material Designer
- More material properties have had their texture sampler type set.

diff --git a/Engine/Plugins/Experimental/DynamicMaterial/Source/DynamicMaterialEditor/Private/Components/DMMaterialProperty.cpp b/Engine/Plugins/Experimental/DynamicMaterial/Source/DynamicMaterialEditor/Private/Components/DMMaterialProperty.cpp
--- a/Engine/Plugins/Experimental/DynamicMaterial/Source/DynamicMaterialEditor/Private/Components/DMMaterialProperty.cpp
+++ b/Engine/Plugins/Experimental/DynamicMaterial/Source/DynamicMaterialEditor/Private/Components/DMMaterialProperty.cpp
@@ -109,7 +109,7 @@ UMaterialExpression* UDMMaterialProperty::GetDefaultInput(const TSharedRef<FDMMa
 
 TEnumAsByte<EMaterialSamplerType> UDMMaterialProperty::GetTextureSamplerType() const
 {
-	return EMaterialSamplerType::SAMPLERTYPE_Color;
+	return EMaterialSamplerType::SAMPLERTYPE_LinearColor;
 }
 
 void UDMMaterialProperty::OnSlotAdded(UDMMaterialSlot* InSlot)
diff --git a/Engine/Plugins/Experimental/DynamicMaterial/Source/DynamicMaterialEditor/Private/Components/MaterialProperties/DMMPBaseColor.cpp b/Engine/Plugins/Experimental/DynamicMaterial/Source/DynamicMaterialEditor/Private/Components/MaterialProperties/DMMPBaseColor.cpp
--- a/Engine/Plugins/Experimental/DynamicMaterial/Source/DynamicMaterialEditor/Private/Components/MaterialProperties/DMMPBaseColor.cpp
+++ b/Engine/Plugins/Experimental/DynamicMaterial/Source/DynamicMaterialEditor/Private/Components/MaterialProperties/DMMPBaseColor.cpp
@@ -20,3 +20,8 @@ UMaterialExpression* UDMMaterialPropertyBaseColor::GetDefaultInput(const TShared
 {
 	return CreateConstant(InBuildState, FVector::OneVector);
 }
+
+TEnumAsByte<EMaterialSamplerType> UDMMaterialPropertyBaseColor::GetTextureSamplerType() const
+{
+	return EMaterialSamplerType::SAMPLERTYPE_Color;
+}
diff --git a/Engine/Plugins/Experimental/DynamicMaterial/Source/DynamicMaterialEditor/Private/Components/MaterialProperties/DMMPEmissiveColor.cpp b/Engine/Plugins/Experimental/DynamicMaterial/Source/DynamicMaterialEditor/Private/Components/MaterialProperties/DMMPEmissiveColor.cpp
--- a/Engine/Plugins/Experimental/DynamicMaterial/Source/DynamicMaterialEditor/Private/Components/MaterialProperties/DMMPEmissiveColor.cpp
+++ b/Engine/Plugins/Experimental/DynamicMaterial/Source/DynamicMaterialEditor/Private/Components/MaterialProperties/DMMPEmissiveColor.cpp
@@ -29,3 +29,8 @@ UMaterialExpression* UDMMaterialPropertyEmissiveColor::GetDefaultInput(
 {
 	return CreateConstant(InBuildState, FVector::OneVector);
 }
+
+TEnumAsByte<EMaterialSamplerType> UDMMaterialPropertyEmissiveColor::GetTextureSamplerType() const
+{
+	return EMaterialSamplerType::SAMPLERTYPE_Color;
+}
diff --git a/Engine/Plugins/Experimental/DynamicMaterial/Source/DynamicMaterialEditor/Private/Components/MaterialProperties/DMMPOpacity.cpp b/Engine/Plugins/Experimental/DynamicMaterial/Source/DynamicMaterialEditor/Private/Components/MaterialProperties/DMMPOpacity.cpp
--- a/Engine/Plugins/Experimental/DynamicMaterial/Source/DynamicMaterialEditor/Private/Components/MaterialProperties/DMMPOpacity.cpp
+++ b/Engine/Plugins/Experimental/DynamicMaterial/Source/DynamicMaterialEditor/Private/Components/MaterialProperties/DMMPOpacity.cpp
@@ -30,8 +30,3 @@ UMaterialExpression* UDMMaterialPropertyOpacity::GetDefaultInput(const TSharedRe
 {
 	return CreateConstant(InBuildState, 1.f);
 }
-
-TEnumAsByte<EMaterialSamplerType> UDMMaterialPropertyOpacity::GetTextureSamplerType() const
-{
-	return EMaterialSamplerType::SAMPLERTYPE_Alpha;
-}
diff --git a/Engine/Plugins/Experimental/DynamicMaterial/Source/DynamicMaterialEditor/Public/Components/MaterialProperties/DMMPBaseColor.h b/Engine/Plugins/Experimental/DynamicMaterial/Source/DynamicMaterialEditor/Public/Components/MaterialProperties/DMMPBaseColor.h
--- a/Engine/Plugins/Experimental/DynamicMaterial/Source/DynamicMaterialEditor/Public/Components/MaterialProperties/DMMPBaseColor.h
+++ b/Engine/Plugins/Experimental/DynamicMaterial/Source/DynamicMaterialEditor/Public/Components/MaterialProperties/DMMPBaseColor.h
@@ -18,5 +18,6 @@ public:
 	//~ Begin UDMMaterialProperty
 	virtual bool IsValidForModel(UDynamicMaterialModelEditorOnlyData& InModelEditorOnlyData) const override;
 	virtual UMaterialExpression* GetDefaultInput(const TSharedRef<FDMMaterialBuildState>& InBuildState) const override;
+	virtual TEnumAsByte<EMaterialSamplerType> GetTextureSamplerType() const override;
 	//~ End UDMMaterialProperty
 };
diff --git a/Engine/Plugins/Experimental/DynamicMaterial/Source/DynamicMaterialEditor/Public/Components/MaterialProperties/DMMPEmissiveColor.h b/Engine/Plugins/Experimental/DynamicMaterial/Source/DynamicMaterialEditor/Public/Components/MaterialProperties/DMMPEmissiveColor.h
--- a/Engine/Plugins/Experimental/DynamicMaterial/Source/DynamicMaterialEditor/Public/Components/MaterialProperties/DMMPEmissiveColor.h
+++ b/Engine/Plugins/Experimental/DynamicMaterial/Source/DynamicMaterialEditor/Public/Components/MaterialProperties/DMMPEmissiveColor.h
@@ -18,5 +18,6 @@ public:
 	//~ Begin UDMMaterialProperty
 	virtual bool IsValidForModel(UDynamicMaterialModelEditorOnlyData& InModelEditorOnlyData) const override;
 	virtual UMaterialExpression* GetDefaultInput(const TSharedRef<FDMMaterialBuildState>& InBuildState) const override;
+	virtual TEnumAsByte<EMaterialSamplerType> GetTextureSamplerType() const override;
 	//~ End UDMMaterialProperty
 };
diff --git a/Engine/Plugins/Experimental/DynamicMaterial/Source/DynamicMaterialEditor/Public/Components/MaterialProperties/DMMPOpacity.h b/Engine/Plugins/Experimental/DynamicMaterial/Source/DynamicMaterialEditor/Public/Components/MaterialProperties/DMMPOpacity.h
--- a/Engine/Plugins/Experimental/DynamicMaterial/Source/DynamicMaterialEditor/Public/Components/MaterialProperties/DMMPOpacity.h
+++ b/Engine/Plugins/Experimental/DynamicMaterial/Source/DynamicMaterialEditor/Public/Components/MaterialProperties/DMMPOpacity.h
@@ -18,6 +18,5 @@ public:
 	//~ Begin UDMMaterialProperty
 	virtual bool IsValidForModel(UDynamicMaterialModelEditorOnlyData& InModelEditorOnlyData) const override;
 	virtual UMaterialExpression* GetDefaultInput(const TSharedRef<FDMMaterialBuildState>& InBuildState) const override;
-	virtual TEnumAsByte<EMaterialSamplerType> GetTextureSamplerType() const override;
 	//~ End UDMMaterialProperty
 };
