Material Designer
- Changed edit widget to be part of the main widget, rather than part of the slot widget.
- Created a global settings bubble.
- Moved global settings into the edit widget when nothing is selected.

diff --git a/Engine/Plugins/Experimental/DynamicMaterial/Source/DynamicMaterialEditor/Private/Menus/DMToolBarMenus.cpp b/Engine/Plugins/Experimental/DynamicMaterial/Source/DynamicMaterialEditor/Private/Menus/DMToolBarMenus.cpp
--- a/Engine/Plugins/Experimental/DynamicMaterial/Source/DynamicMaterialEditor/Private/Menus/DMToolBarMenus.cpp
+++ b/Engine/Plugins/Experimental/DynamicMaterial/Source/DynamicMaterialEditor/Private/Menus/DMToolBarMenus.cpp
@@ -56,8 +56,8 @@ namespace UE::DynamicMaterialEditor::Private
 
 					if (EditorWidget.IsValid())
 					{
-						EditorWidget->RefreshGlobalDetailsView();
 						EditorWidget->RefreshSlotPickerList();
+						EditorWidget->RefreshSlotWidget();
 					}
 				}
 			}
@@ -123,8 +123,8 @@ namespace UE::DynamicMaterialEditor::Private
 
 					if (EditorWidget.IsValid())
 					{
-						EditorWidget->RefreshGlobalDetailsView();
 						EditorWidget->RefreshSlotPickerList();
+						EditorWidget->RefreshSlotWidget();
 					}
 				}
 			}
@@ -206,8 +206,8 @@ namespace UE::DynamicMaterialEditor::Private
 
 					if (EditorWidget.IsValid())
 					{
-						EditorWidget->RefreshGlobalDetailsView();
 						EditorWidget->RefreshSlotPickerList();
+						EditorWidget->RefreshSlotWidget();
 					}
 				}
 			}
@@ -1004,18 +1004,11 @@ TSharedRef<SWidget> FDMToolBarMenus::MakeEditorLayoutMenu(const TSharedPtr<SDMEd
 		);
 	}
 
-	TSharedPtr<SDMSlot> SlotWidget;
-	if (InEditorWidget.IsValid())
-	{
-		const TArray<TSharedRef<SDMSlot>>& SlotWidgets = InEditorWidget->GetSlotWidgets();
-
-		if (SlotWidgets.IsEmpty() == false)
-		{
-			SlotWidget = SlotWidgets[0];
-		}
-	}
-
-	FToolMenuContext MenuContext(FDynamicMaterialEditorModule::Get().GetCommandList(), TSharedPtr<FExtender>(), UDMMenuContext::CreateSlot(SlotWidget));
+	FToolMenuContext MenuContext(
+		FDynamicMaterialEditorModule::Get().GetCommandList(),
+		TSharedPtr<FExtender>(),
+		UDMMenuContext::CreateSlot(InEditorWidget->GetActiveSlotWidget())
+	);
 
 	return UToolMenus::Get()->GenerateWidget(ToolBarEditorLayoutMenuName, MenuContext);
 }
diff --git a/Engine/Plugins/Experimental/DynamicMaterial/Source/DynamicMaterialEditor/Private/Slate/Layers/SDMLayerEffectsView.cpp b/Engine/Plugins/Experimental/DynamicMaterial/Source/DynamicMaterialEditor/Private/Slate/Layers/SDMLayerEffectsView.cpp
--- a/Engine/Plugins/Experimental/DynamicMaterial/Source/DynamicMaterialEditor/Private/Slate/Layers/SDMLayerEffectsView.cpp
+++ b/Engine/Plugins/Experimental/DynamicMaterial/Source/DynamicMaterialEditor/Private/Slate/Layers/SDMLayerEffectsView.cpp
@@ -49,30 +49,43 @@ TSharedRef<ITableRow> SDMLayerEffectsView::OnGenerateLayerItemWidget(TSharedPtr<
 
 void SDMLayerEffectsView::OnLayerEffectSelected(const bool bInSelected, const TSharedRef<SDMLayerEffectsItem>& InEffectsItemWidget)
 {
+	TSharedPtr<SDMSlot> SlotWidget = SlotWidgetWeak.Pin();
+
+	if (!SlotWidget.IsValid())
+	{
+		return;
+	}
+
+	TSharedPtr<SDMEditor> EditorWidget = SlotWidget->GetEditorWidget();
+
 	if (!bInSelected)
 	{
 		SelectedEffectWidgets.Remove(InEffectsItemWidget);
 		InEffectsItemWidget->SetEffectSelected(false);
 
-		if (TSharedPtr<SDMSlot> SlotWidget = SlotWidgetWeak.Pin())
+		if (EditorWidget.IsValid())
 		{
-			SlotWidget->SetEditedComponent(nullptr);
-			SlotWidget->SetSelectedLayer(nullptr);
+			EditorWidget->SetEditedComponent(nullptr);
+		}
+
+		SlotWidget->SetSelectedLayer(nullptr);
 
-			if (TSharedPtr<FDMEffectsLayerItem> SelectedEffectLayerItem = SelectedEffectWidgets.Last()->GetLayerItem())
+		if (TSharedPtr<FDMEffectsLayerItem> SelectedEffectLayerItem = SelectedEffectWidgets.Last()->GetLayerItem())
+		{
+			if (UDMMaterialEffect* Effect = SelectedEffectLayerItem->MaterialEffectWeak.Get())
 			{
-				if (UDMMaterialEffect* Effect = SelectedEffectLayerItem->MaterialEffectWeak.Get())
+				if (EditorWidget.IsValid())
 				{
-					SlotWidget->SetEditedComponent(Effect);
+					EditorWidget->SetEditedComponent(Effect);
+				}
 
-					if (UDMMaterialEffectStack* EffectStack = Effect->GetEffectStack())
+				if (UDMMaterialEffectStack* EffectStack = Effect->GetEffectStack())
+				{
+					if (UDMMaterialLayerObject* SelectedLayer = EffectStack->GetLayer())
 					{
-						if (UDMMaterialLayerObject* SelectedLayer = EffectStack->GetLayer())
-						{
-							SlotWidget->SetSelectedLayer(nullptr);
-						}
-					}					
-				}
+						SlotWidget->SetSelectedLayer(nullptr);
+					}
+				}					
 			}
 		}
 	}
@@ -81,27 +94,28 @@ void SDMLayerEffectsView::OnLayerEffectSelected(const bool bInSelected, const TS
 		SelectedEffectWidgets.AddUnique(InEffectsItemWidget);
 		InEffectsItemWidget->SetEffectSelected(true);
 
-		if (TSharedPtr<SDMSlot> SlotWidget = SlotWidgetWeak.Pin())
+		UDMMaterialLayerObject* Layer = nullptr;
+		UDMMaterialEffect* Effect = nullptr;
+
+		if (TSharedPtr<FDMEffectsLayerItem> EffectViewItem = InEffectsItemWidget->GetLayerItem())
 		{
-			UDMMaterialLayerObject* Layer = nullptr;
-			UDMMaterialEffect* Effect = nullptr;
+			Effect = EffectViewItem->MaterialEffectWeak.Get();
 
-			if (TSharedPtr<FDMEffectsLayerItem> EffectViewItem = InEffectsItemWidget->GetLayerItem())
+			if (Effect)
 			{
-				Effect = EffectViewItem->MaterialEffectWeak.Get();
-
-				if (Effect)
+				if (UDMMaterialEffectStack* EffectStack = Effect->GetEffectStack())
 				{
-					if (UDMMaterialEffectStack* EffectStack = Effect->GetEffectStack())
-					{
-						Layer = EffectStack->GetLayer();
-					}
+					Layer = EffectStack->GetLayer();
 				}
 			}
+		}
 
-			SlotWidget->SetEditedComponent(Effect);
-			SlotWidget->SetSelectedLayer(Layer);
+		if (EditorWidget.IsValid())
+		{
+			EditorWidget->SetEditedComponent(Effect);
 		}
+
+		SlotWidget->SetSelectedLayer(Layer);
 	}
 }
 
diff --git a/Engine/Plugins/Experimental/DynamicMaterial/Source/DynamicMaterialEditor/Private/Slate/SDMComponentEdit.cpp b/Engine/Plugins/Experimental/DynamicMaterial/Source/DynamicMaterialEditor/Private/Slate/SDMComponentEdit.cpp
--- a/Engine/Plugins/Experimental/DynamicMaterial/Source/DynamicMaterialEditor/Private/Slate/SDMComponentEdit.cpp
+++ b/Engine/Plugins/Experimental/DynamicMaterial/Source/DynamicMaterialEditor/Private/Slate/SDMComponentEdit.cpp
@@ -12,12 +12,14 @@
 #include "Components/DMMaterialStageThroughput.h"
 #include "Components/DMMaterialValue.h"
 #include "Components/DMTextureUV.h"
+#include "Components/MaterialValues/DMMaterialValueFloat1.h"
+#include "Components/MaterialValues/DMMaterialValueFloat2.h"
 #include "CustomDetailsViewArgs.h"
 #include "CustomDetailsViewModule.h"
 #include "CustomDetailsViewSequencer.h"
-#include "DMPrivate.h"
 #include "DetailLayoutBuilder.h"
 #include "DMEDefs.h"
+#include "DMPrivate.h"
 #include "DMWorldSubsystem.h"
 #include "DynamicMaterialEditorModule.h"
 #include "DynamicMaterialEditorStyle.h"
@@ -33,6 +35,7 @@
 #include "Items/ICustomDetailsViewItem.h"
 #include "Menus/DMMaterialStageSourceMenus.h"
 #include "Misc/CoreDelegates.h"
+#include "Model/DynamicMaterialModel.h"
 #include "Model/DynamicMaterialModelEditorOnlyData.h"
 #include "PropertyCustomizationHelpers.h"
 #include "PropertyEditorModule.h"
@@ -247,14 +250,14 @@ SDMComponentEdit::~SDMComponentEdit()
 	}
 }
 
-void SDMComponentEdit::Construct(const FArguments& InArgs, UDMMaterialComponent* InComponent, const TWeakPtr<SDMSlot>& InSlotWidget)
+void SDMComponentEdit::Construct(const FArguments& InArgs, UDMMaterialComponent* InComponent, const TWeakPtr<SDMEditor>& InEditorWidget)
 {
 	ComponentWeak = InComponent;
-	SlotWidgetWeak = InSlotWidget;
+	EditorWidgetWeak = InEditorWidget;
 
 	KeyframeHandler = nullptr;
 
-	if (ensure(IsValid(InComponent)))
+	if (InComponent)
 	{
 		if (const UWorld* const World = InComponent->GetWorld())
 		{
@@ -263,22 +266,22 @@ void SDMComponentEdit::Construct(const FArguments& InArgs, UDMMaterialComponent*
 				KeyframeHandler = WorldSubsystem->GetKeyframeHandler();
 			}
 		}
+	}
 
-		if (UDMMaterialStage* Stage = Cast<UDMMaterialStage>(InComponent))
-		{
-			Stage->SetBeingEdited(true);
-		}
+	if (UDMMaterialStage* Stage = Cast<UDMMaterialStage>(InComponent))
+	{
+		Stage->SetBeingEdited(true);
+	}
 
-		ChildSlot
-			.HAlign(HAlign_Fill)
-			.VAlign(VAlign_Top)
+	ChildSlot
+		.HAlign(HAlign_Fill)
+		.VAlign(VAlign_Top)
+		[
+			SAssignNew(Container, SBox)
 			[
-				SAssignNew(Container, SBox)
-				[
-					CreateEditWidget()
-				]
-			];
-	}
+				CreateEditWidget()
+			]
+		];
 }
 
 TSharedRef<SWidget> SDMComponentEdit::CreateEditWidget()
@@ -398,7 +401,20 @@ TArray<FDMPropertyHandle> SDMComponentEdit::GetEditRows()
 	TArray<FDMPropertyHandle> PropertyRows;
 	TSet<UDMMaterialComponent*> ProcessedObjects;
 
-	FDynamicMaterialEditorModule::GeneratorComponentPropertyRows(SharedThis(this), ComponentWeak.Get(), PropertyRows, ProcessedObjects);
+	if (UDMMaterialComponent* Component = ComponentWeak.Get())
+	{
+		FDynamicMaterialEditorModule::GeneratorComponentPropertyRows(SharedThis(this), Component, PropertyRows, ProcessedObjects);
+	}
+	else
+	{
+		if (TSharedPtr<SDMEditor> EditorWidget = EditorWidgetWeak.Pin())
+		{
+			if (UDynamicMaterialModel* MaterialModel = EditorWidget->GetMaterialModel())
+			{
+				GenerateMaterialModelPropertyRows(EditorWidget.ToSharedRef(), MaterialModel, PropertyRows, ProcessedObjects);
+			}
+		}
+	}
 
 	return PropertyRows;
 }
@@ -415,9 +431,9 @@ TSharedRef<SWidget> SDMComponentEdit::CreateExtensionButtons(const TSharedPtr<SD
 			{
 				if (TSharedPtr<SDMComponentEdit> ComponentEditWidget = ComponentEditWidgetWeak.Pin())
 				{
-					if (TSharedPtr<SDMSlot> SlotWidget = ComponentEditWidget->GetSlotWidget())
+					if (TSharedPtr<SDMEditor> EditorWidget = ComponentEditWidget->GetEditorWidget())
 					{
-						SlotWidget->InvalidateComponentEditWidget();
+						EditorWidget->InvalidateComponentEditWidget();
 					}
 				}
 			})
@@ -707,11 +723,23 @@ TSharedRef<SWidget> SDMComponentEdit::MakeSourceTypeEditWidgetMenuContent()
 {
 	if (UDMMaterialStage* Stage = Cast<UDMMaterialStage>(ComponentWeak.Get()))
 	{
-		if (TSharedPtr<SDMSlot> SlotWidget = GetSlotWidget())
+		if (UDMMaterialLayerObject* Layer = Stage->GetLayer())
 		{
-			if (TSharedPtr<SDMStage> StageWidget = SlotWidget->FindStageWidget(Stage))
+			if (UDMMaterialSlot* Slot = Layer->GetSlot())
 			{
-				return FDMMaterialStageSourceMenus::MakeChangeSourceMenu(SlotWidget, StageWidget);
+				if (TSharedPtr<SDMEditor> EditorWidget = GetEditorWidget())
+				{
+					if (TSharedPtr<SDMSlot> SlotWidget = EditorWidget->GetActiveSlotWidget())
+					{
+						if (SlotWidget->GetSlot() == Slot)
+						{
+							if (TSharedPtr<SDMStage> StageWidget = SlotWidget->FindStageWidget(Stage))
+							{
+								return FDMMaterialStageSourceMenus::MakeChangeSourceMenu(SlotWidget, StageWidget);
+							}
+						}
+					}
+				}
 			}
 		}
 	}
@@ -736,15 +764,15 @@ void SDMComponentEdit::OnUndo()
 {
 	if (UDMMaterialStage* Stage = Cast<UDMMaterialStage>(ComponentWeak.Get()))
 	{
-		if (const UDMMaterialLayerObject* Layer = Stage->GetLayer())
+		if (UDMMaterialLayerObject* Layer = Stage->GetLayer())
 		{
 			if (Layer->GetStageType(Stage) == EDMMaterialLayerStage::Mask)
 			{
 				if (bCreatedWithLinkedUVs != Layer->IsTextureUVLinkEnabled())
 				{
-					if (TSharedPtr<SDMSlot> SlotWidget = SlotWidgetWeak.Pin())
+					if (TSharedPtr<SDMEditor> EditorWidget = GetEditorWidget())
 					{
-						SlotWidget->InvalidateComponentEditWidget();
+						EditorWidget->InvalidateComponentEditWidget();
 					}
 				}
 			}
@@ -752,6 +780,42 @@ void SDMComponentEdit::OnUndo()
 	}
 }
 
+void SDMComponentEdit::GenerateMaterialModelPropertyRows(const TSharedRef<SDMEditor> InEditorWidget, UDynamicMaterialModel* InMaterialModel,
+	TArray<FDMPropertyHandle>& InOutPropertyRows, TSet<UDMMaterialComponent*>& InOutProcessedObjects)
+{
+	if (UDynamicMaterialModelEditorOnlyData* EditorOnlyData = UDynamicMaterialModelEditorOnlyData::Get(InMaterialModel))
+	{
+		if (EditorOnlyData->GetBlendMode() != BLEND_Opaque)
+		{
+			if (UDMMaterialValueFloat1* GlobalOpacityValue = InMaterialModel->GetGlobalOpacityValue())
+			{
+				if (!InOutProcessedObjects.Contains(GlobalOpacityValue))
+				{
+					FDMPropertyHandle& GlobalOpacityHandle = InOutPropertyRows.Add_GetRef(InEditorWidget->GetPropertyHandle(&*InEditorWidget,
+						GlobalOpacityValue, UDMMaterialValue::ValueName));
+
+					GlobalOpacityHandle.NameOverride = LOCTEXT("GlobalOpacity", "Global Opacity");
+
+					InOutProcessedObjects.Add(GlobalOpacityValue);
+				}
+			}
+		}
+	}
+
+	if (UDMMaterialValueFloat2* GlobalTilingValue = InMaterialModel->GetGlobalScaleValue())
+	{
+		if (!InOutProcessedObjects.Contains(GlobalTilingValue))
+		{
+			FDMPropertyHandle& GlobalTilingHandle = InOutPropertyRows.Add_GetRef(InEditorWidget->GetPropertyHandle(&*InEditorWidget, 
+				GlobalTilingValue, UDMMaterialValue::ValueName));
+
+			GlobalTilingHandle.NameOverride = LOCTEXT("GlobalTiling", "Global Tiling");
+
+			InOutProcessedObjects.Add(GlobalTilingValue);
+		}
+	}
+}
+
 TSharedRef<SWidget> SDMComponentEdit::CreateSourceTypeEditWidget()
 {
 	TWeakPtr<SDMComponentEdit> ThisWeak = SharedThis(this);
diff --git a/Engine/Plugins/Experimental/DynamicMaterial/Source/DynamicMaterialEditor/Private/Slate/SDMEditor.cpp b/Engine/Plugins/Experimental/DynamicMaterial/Source/DynamicMaterialEditor/Private/Slate/SDMEditor.cpp
--- a/Engine/Plugins/Experimental/DynamicMaterial/Source/DynamicMaterialEditor/Private/Slate/SDMEditor.cpp
+++ b/Engine/Plugins/Experimental/DynamicMaterial/Source/DynamicMaterialEditor/Private/Slate/SDMEditor.cpp
@@ -5,13 +5,8 @@
 #include "Components/DMMaterialLayer.h"
 #include "Components/DMMaterialProperty.h"
 #include "Components/DMMaterialSlot.h"
-#include "Components/DMMaterialValue.h"
 #include "Components/MaterialStageExpressions/DMMSETextureSample.h"
-#include "Components/MaterialValues/DMMaterialValueFloat1.h"
-#include "Components/MaterialValues/DMMaterialValueFloat2.h"
 #include "Components/PrimitiveComponent.h"
-#include "CustomDetailsViewArgs.h"
-#include "CustomDetailsViewModule.h"
 #include "DetailLayoutBuilder.h"
 #include "DMBlueprintFunctionLibrary.h"
 #include "DMPrivate.h"
@@ -25,18 +20,14 @@
 #include "Framework/Commands/GenericCommands.h"
 #include "Framework/Commands/InputChord.h"
 #include "HAL/PlatformApplicationMisc.h"
-#include "ICustomDetailsView.h"
 #include "IDetailTreeNode.h"
 #include "IPropertyRowGenerator.h"
-#include "Items/ICustomDetailsViewItem.h"
-#include "MaterialDomain.h"
+#include "Materials/MaterialInterface.h"
 #include "Menus/DMToolBarMenus.h"
 #include "Misc/CoreDelegates.h"
 #include "Model/DynamicMaterialModel.h"
 #include "Model/DynamicMaterialModelEditorOnlyData.h"
-#include "Properties/Editors/SDMPropertyEditFloat2Value.h"
 #include "PropertyHandle.h"
-#include "Slate/Properties/Editors/SDMPropertyEditOpacity.h"
 #include "Slate/Properties/SDMMaterialParameters.h"
 #include "Slate/SDMComponentEdit.h"
 #include "Slate/SDMMaterialWizard.h"
@@ -46,6 +37,8 @@
 #include "Styling/StyleColors.h"
 #include "ThumbnailRendering/ThumbnailManager.h"
 #include "Widgets/Input/SButton.h"
+#include "Widgets/Input/SCheckBox.h"
+#include "Widgets/Layout/SBorder.h"
 #include "Widgets/Layout/SBox.h"
 #include "Widgets/Layout/SExpandableArea.h"
 #include "Widgets/Layout/SScrollBox.h"
@@ -175,8 +168,7 @@ void SDMEditor::Construct(const FArguments& InArgs, TWeakObjectPtr<UDynamicMater
 }
 
 void SDMEditor::PrivateRegisterAttributes(struct FSlateAttributeDescriptor::FInitializer&)
-{
-	
+{	
 }
 
 SDMEditor::~SDMEditor()
@@ -202,11 +194,27 @@ void SDMEditor::ClearEditor()
 
 void SDMEditor::SetActiveSlotIndex(int InSlotIndex)
 {
-	if (InSlotIndex != ActiveSlotIndex && SlotWidgets.IsValidIndex(InSlotIndex))
+	if (InSlotIndex == INDEX_NONE)
+	{
+		ActiveSlotIndex = INDEX_NONE;
+		RefreshSlotWidget();
+
+		EditedComponent.Reset();
+		RefreshComponentEditWidget();
+	}
+	else if (InSlotIndex != ActiveSlotIndex)
 	{
-		ActiveSlotIndex = InSlotIndex;
-		SlotWidgets[InSlotIndex]->ClearSelection();
-		RefreshSlotsList();
+		if (UDynamicMaterialModelEditorOnlyData* ModelEditorOnlyData = UDynamicMaterialModelEditorOnlyData::Get(MaterialModelWeak))
+		{
+			if (ModelEditorOnlyData->GetSlots().IsValidIndex(InSlotIndex))
+			{
+				ActiveSlotIndex = InSlotIndex;
+				RefreshSlotWidget();
+
+				EditedComponent.Reset();
+				RefreshComponentEditWidget();
+			}
+		}
 	}
 }
 
@@ -285,8 +293,11 @@ void SDMEditor::SetMaterialModel(UDynamicMaterialModel* InMaterialModel)
 {
 	MaterialModelWeak = InMaterialModel;
 
-	SlotsContainer.Reset();
-	SlotWidgets.Empty();
+	SlotPickerContainer.Reset();
+	SlotContainer.Reset();
+	ActiveSlotWidget.Reset();
+	ComponentEditContainer.Reset();
+	SplitterContainer.Reset();
 
 	Toolbar->SetMaterialModel(InMaterialModel);
 
@@ -363,6 +374,11 @@ void SDMEditor::SetMaterialActor(AActor* InActor)
 	);
 }
 
+TSharedPtr<SDMSlot> SDMEditor::GetActiveSlotWidget() const
+{
+	return ActiveSlotWidget;
+}
+
 TSharedRef<SWidget> SDMEditor::CreateActorMaterialSlotSelector(const AActor* InActor)
 {
 	TArray<TSharedPtr<FDMObjectMaterialProperty>> MaterialProperties = Toolbar->GetMaterialProperties();
@@ -432,7 +448,13 @@ TSharedRef<SWidget> SDMEditor::CreateActorMaterialSlotSelector(const AActor* InA
 		TWeakPtr<FDMObjectMaterialProperty> MaterialSlotWeak = MaterialSlot;
 
 		constexpr int32 ThumbnailSize = 48;
-		TSharedRef<FAssetThumbnail> Thumbnail = MakeShared<FAssetThumbnail>(PrimComponent->GetMaterial(MaterialSlot->Index), ThumbnailSize, ThumbnailSize, UThumbnailManager::Get().GetSharedThumbnailPool());
+
+		TSharedRef<FAssetThumbnail> Thumbnail = MakeShared<FAssetThumbnail>(
+			PrimComponent->GetMaterial(MaterialSlot->Index),
+			ThumbnailSize,
+			ThumbnailSize, 
+			GetThumbnailPool()
+		);
 
 		FAssetThumbnailConfig ThumbnailConfig;
 		ThumbnailConfig.GenericThumbnailSize = ThumbnailSize;
@@ -551,6 +573,8 @@ void SDMEditor::BindCommands()
 
 TSharedRef<SWidget> SDMEditor::CreateMainLayout()
 {
+	const float SplitterValue = UDynamicMaterialEditorSettings::Get()->SplitterLocation;
+
 	return 
 		SNew(SVerticalBox)
 		+ SVerticalBox::Slot()
@@ -565,18 +589,6 @@ TSharedRef<SWidget> SDMEditor::CreateMainLayout()
 		]
 		+ SVerticalBox::Slot()
 		.AutoHeight()
-		.HAlign(HAlign_Fill)
-		.VAlign(VAlign_Fill)
-		[
-			SAssignNew(GlobalOpacityContainer, SBox)
-			.HAlign(HAlign_Fill)
-			.VAlign(VAlign_Fill)
-			[
-				CreateGlobalDetailsView()
-			]
-		]
-		+ SVerticalBox::Slot()
-		.AutoHeight()
 		[
 			UE::DynamicMaterialEditor::bGlobalValuesEnabled && MaterialModelWeak.IsValid()
 				? CreateParametersArea()
@@ -598,118 +610,58 @@ TSharedRef<SWidget> SDMEditor::CreateMainLayout()
 		.HAlign(HAlign_Fill)
 		.VAlign(VAlign_Fill)
 		[
-			SAssignNew(SlotsContainer, SBox)
-			.HAlign(HAlign_Fill)
-			.VAlign(VAlign_Fill)
+			SAssignNew(SplitterContainer, SSplitter)
+			.Style(FAppStyle::Get(), "DetailsView.Splitter")
+			.Orientation(Orient_Vertical)
+			.ResizeMode(ESplitterResizeMode::Fill)
+			.PhysicalSplitterHandleSize(3.0f)
+			.HitDetectionSplitterHandleSize(4.0f)
+			.OnSplitterFinishedResizing(this, &SDMEditor::OnSplitterResized)
+
+			+ SSplitter::Slot()
+			.Expose(LayerViewSplitterSlot)
+			.Resizable(true)
+			.SizeRule(SSplitter::ESizeRule::FractionOfParent)
+			.MinSize(50)
+			.Value(SplitterValue)
 			[
-				CreateSlotsWidget()
+				SNew(SBorder)
+				.Padding(2.0f)
+				.BorderImage(FDynamicMaterialEditorStyle::GetBrush("LayerView.Background"))
+				.HAlign(HAlign_Fill)
+				.VAlign(VAlign_Fill)
+				[
+					SAssignNew(SlotContainer, SScrollBox)
+					+ SScrollBox::Slot()
+					.FillSize(1.f)
+					[
+						CreateSlotWidget()
+					]
+				]
 			]
-		];
-}
-
-TSharedRef<SWidget> SDMEditor::CreateGlobalDetailsView()
-{
-	UDynamicMaterialModel* MaterialModel = MaterialModelWeak.Get();
-
-	if (!MaterialModel)
-	{
-		return SNullWidget::NullWidget;
-	}
-
-	UDMMaterialValueFloat1* OpacityValue = MaterialModel->GetGlobalOpacityValue();
-	UDMMaterialValueFloat2* ScaleValue = MaterialModel->GetGlobalScaleValue();
-
-	if (!OpacityValue || !ScaleValue)
-	{
-		return SNullWidget::NullWidget;
-	}
-
-	FCustomDetailsViewArgs Args;
-	Args.KeyframeHandler = nullptr;
-	Args.bAllowGlobalExtensions = true;
-	Args.bAllowResetToDefault = true;
-	Args.bShowCategories = false;
-
-	if (const UWorld* const World = MaterialModel->GetWorld())
-	{
-		if (const UDMWorldSubsystem* const WorldSubsystem = World->GetSubsystem<UDMWorldSubsystem>())
-		{
-			Args.KeyframeHandler = WorldSubsystem->GetKeyframeHandler();
-		}
-	}
-
-	TSharedRef<ICustomDetailsView> DetailsView = ICustomDetailsViewModule::Get().CreateCustomDetailsView(Args);
-	FCustomDetailsViewItemId RootId = DetailsView->GetRootItem()->GetItemId();
-
-	static const FName UIMin = FName("UIMin");
-	static const FName UIMax = FName("UIMax");
-	static const FName ClampMin = FName("ClampMin");
-	static const FName ClampMax = FName("ClampMax");
-
-	TArray<FDMPropertyHandle> EditRows;
-
-	if (GetGlobalOpacityVisibility() == EVisibility::Visible)
-	{
-		FDMPropertyHandle& GlobalOpacityHandle = EditRows.Add_GetRef(SDMEditor::GetPropertyHandle(this, OpacityValue, UDMMaterialValue::ValueName));
-
-		GlobalOpacityHandle.PropertyHandle->SetInstanceMetaData(UIMin, FString::SanitizeFloat(OpacityValue->GetValueRange().Min));
-		GlobalOpacityHandle.PropertyHandle->SetInstanceMetaData(ClampMin, FString::SanitizeFloat(OpacityValue->GetValueRange().Min));
-		GlobalOpacityHandle.PropertyHandle->SetInstanceMetaData(UIMax, FString::SanitizeFloat(OpacityValue->GetValueRange().Max));
-		GlobalOpacityHandle.PropertyHandle->SetInstanceMetaData(ClampMax, FString::SanitizeFloat(OpacityValue->GetValueRange().Max));
-
-		GlobalOpacityHandle.NameOverride = LOCTEXT("GlobalOpacity", "Global Opacity");
-	}
-
-	FDMPropertyHandle& GlobalScaleHandle = EditRows.Add_GetRef(SDMEditor::GetPropertyHandle(this, ScaleValue, UDMMaterialValue::ValueName));
-	GlobalScaleHandle.NameOverride = LOCTEXT("GlobalScale", "Global Scale");
-
-	for (const FDMPropertyHandle& EditRow : EditRows)
-	{
-		if (!EditRow.DetailTreeNode)
-		{
-			continue;
-		}
-
-		ECustomDetailsTreeInsertPosition Position = ECustomDetailsTreeInsertPosition::Child;
-
-		if (EditRow.DetailTreeNode->CreatePropertyHandle()->HasMetaData("HighPriority"))
-		{
-			Position = ECustomDetailsTreeInsertPosition::FirstChild;
-		}
-		else if (EditRow.DetailTreeNode->CreatePropertyHandle()->HasMetaData("LowPriority"))
-		{
-			Position = ECustomDetailsTreeInsertPosition::LastChild;
-		}
-
-		TSharedRef<ICustomDetailsViewItem> Item = DetailsView->CreateDetailTreeItem(EditRow.DetailTreeNode.ToSharedRef());
-
-		if (EditRow.NameOverride.IsSet())
-		{
-			Item->SetOverrideWidget(
-				ECustomDetailsViewWidgetType::Name,
-				SNew(STextBlock)
-				.Font(IDetailLayoutBuilder::GetDetailFont())
-				.Text(EditRow.NameOverride.GetValue())
-				.ToolTipText(EditRow.NameToolTipOverride.Get(FText::GetEmpty()))
-			);
-		}
-
-		if (EditRow.DetailTreeNode->CreatePropertyHandle()->HasMetaData("NotKeyframeable"))
-		{
-			Item->SetKeyframeEnabled(false);
-		}
-
-		if (EditRow.ResetToDefaultOverride.IsSet())
-		{
-			Item->SetResetToDefaultOverride(EditRow.ResetToDefaultOverride.GetValue());
-		}
-
-		DetailsView->ExtendTree(RootId, Position, Item);
-	}
-
-	DetailsView->RebuildTree(ECustomDetailsViewBuildType::InstantBuild);
 
-	return DetailsView;
+			+ SSplitter::Slot()
+			.Expose(ExtraSpaceSplitterSlot)
+			.Resizable(true)
+			.SizeRule(SSplitter::ESizeRule::FractionOfParent)
+			.MinSize(50)
+			.Value(SplitterValue)
+			[
+				SNew(SBorder)
+				.Padding(2.0f)
+				.BorderImage(FDynamicMaterialEditorStyle::GetBrush("LayerView.Background"))
+				.HAlign(HAlign_Fill)
+				.VAlign(VAlign_Fill)
+				[
+					SAssignNew(ComponentEditContainer, SScrollBox)
+					+ SScrollBox::Slot()
+					.AutoSize()
+					[
+						CreateComponentEditWidget()
+					]
+				]
+			]
+		];
 }
 
 TSharedRef<SWidget> SDMEditor::CreateParametersArea()
@@ -748,8 +700,24 @@ TSharedRef<SWidget> SDMEditor::CreateSlotPickerWidget()
 		return SlotSelector;
 	}
 
+	SlotSelector->AddSlot()
+		[
+			SNew(SCheckBox)
+			.Style(FAppStyle::Get(), "DetailsView.SectionButton")
+			.HAlign(EHorizontalAlignment::HAlign_Center)
+			.IsChecked(this, &SDMEditor::GetSlotCheckState, EDMMaterialPropertyType::None)
+			.OnCheckStateChanged(this, &SDMEditor::OnSlotCheckStateChanged, EDMMaterialPropertyType::None)
+			.Padding(FVector2D(5.f, 3.f))
+			.ToolTipText(LOCTEXT("GlobalSettings", "Global Material Settings"))
+			.Content()
+			[
+				SNew(STextBlock)
+				.Text(LOCTEXT("Settings", "Settings"))
+			]
+		];
+
 	const FDMMaterialChannelListPreset* Preset = GetDefault<UDynamicMaterialEditorSettings>()->ChannelPresets.Find(ModelEditorOnlyData->GetChannelListPreset());
-	bool bHasBaseColorSlot = !!ModelEditorOnlyData->GetSlotForMaterialProperty(EDMMaterialPropertyType::BaseColor);
+	const bool bHasBaseColorSlot = !!ModelEditorOnlyData->GetSlotForMaterialProperty(EDMMaterialPropertyType::BaseColor);
 
 	UEnum* MaterialPropertyEnum = StaticEnum<EDMMaterialPropertyType>();
 
@@ -805,31 +773,6 @@ TSharedRef<SWidget> SDMEditor::CreateSlotPickerWidget()
 	return SlotSelector;
 }
 
-TSharedPtr<SDMSlot> SDMEditor::GetSlotWidget(UDMMaterialSlot* Slot) const
-{
-	if (ensure(IsValid(Slot)))
-	{
-		for (const TSharedRef<SDMSlot>& SlotWidget : SlotWidgets)
-		{
-			if (SlotWidget->GetSlot() == Slot)
-			{
-				return SlotWidget;
-			}
-		}
-	}
-
-	return nullptr;
-}
-
-void SDMEditor::RefreshGlobalDetailsView()
-{
-	if (GlobalOpacityContainer.IsValid())
-	{
-		GlobalOpacityContainer->SetContent(SNullWidget::NullWidget);
-		GlobalOpacityContainer->SetContent(CreateGlobalDetailsView());
-	}
-}
-
 void SDMEditor::RefreshParametersList()
 {
 	if (ParametersWidget.IsValid())
@@ -847,12 +790,16 @@ void SDMEditor::RefreshSlotPickerList()
 	}
 }
 
-void SDMEditor::RefreshSlotsList()
+void SDMEditor::RefreshSlotWidget()
 {
-	if (SlotsContainer.IsValid())
+	if (SlotContainer.IsValid())
 	{
-		SlotsContainer->SetContent(SNullWidget::NullWidget);
-		SlotsContainer->SetContent(CreateSlotsWidget());
+		SlotContainer->ClearChildren();
+		SlotContainer->AddSlot()
+			.FillSize(1.f)
+			[
+				CreateSlotWidget()
+			];
 	}
 }
 
@@ -922,24 +869,19 @@ void SDMEditor::Tick(const FGeometry& AllottedGeometry, const double InCurrentTi
 		{
 			Toolbar->SetMaterialModel(MaterialModelWeak.Get());
 		}
-		else
+		else if (ActiveSlotWidget.IsValid())
 		{
-			bool bHasInvalidSlotWidget = false;
-
-			for (const TSharedRef<SDMSlot>& Slot : SlotWidgets)
-			{
-				if (!Slot->CheckValidity())
-				{
-					bHasInvalidSlotWidget = true;
-					break;
-				}
-			}
-			if (bHasInvalidSlotWidget)
+			if (!ActiveSlotWidget->CheckValidity())
 			{
-				RefreshSlotsList();
+				RefreshSlotWidget();
 			}
 		}
 	}
+
+	if (bInvalidateComponentEditWidget)
+	{
+		RefreshComponentEditWidget();
+	}
 }
 
 void SDMEditor::OnMaterialBuilt(UDynamicMaterialModel* InMaterialModel)
@@ -953,51 +895,46 @@ void SDMEditor::OnValuesUpdated(UDynamicMaterialModel* InMaterialModel)
 
 void SDMEditor::OnSlotsUpdated(UDynamicMaterialModel* InMaterialModel)
 {
-	RefreshSlotsList();
+	RefreshSlotWidget();
 }
 
-EVisibility SDMEditor::GetGlobalOpacityVisibility() const
+bool SDMEditor::IsPropertyValidForModel(EDMMaterialPropertyType InProperty) const
 {
-	if (UDynamicMaterialModelEditorOnlyData* ModelEditorOnlyData = UDynamicMaterialModelEditorOnlyData::Get(MaterialModelWeak))
+	UDynamicMaterialModel* MaterialModel = MaterialModelWeak.Get();
+
+	if (!MaterialModel)
 	{
-		if (ModelEditorOnlyData->GetBlendMode() != BLEND_Opaque)
-		{
-			return EVisibility::Visible;
-		}
+		return false;
 	}
 
-	return EVisibility::Collapsed;
-}
+	UDynamicMaterialModelEditorOnlyData* EditorOnlyData = UDynamicMaterialModelEditorOnlyData::Get(MaterialModel);
 
-bool SDMEditor::IsPropertyValidForModel(EDMMaterialPropertyType InProperty) const
-{
-	if (UDynamicMaterialModel* MaterialModel = MaterialModelWeak.Get())
+	if (!EditorOnlyData)
 	{
-		if (UDynamicMaterialModelEditorOnlyData* EditorOnlyData = UDynamicMaterialModelEditorOnlyData::Get(MaterialModel))
+		return false;
+	}
+
+	if (UDMMaterialProperty* Property = EditorOnlyData->GetMaterialProperty(InProperty))
+	{
+		if (Property->IsValidForModel(*EditorOnlyData))
 		{
-			if (UDMMaterialProperty* Property = EditorOnlyData->GetMaterialProperty(InProperty))
-			{
-				if (Property->IsValidForModel(*EditorOnlyData))
-				{
-					return true;
-				}
-			}
+			return true;
+		}
+	}
 
-			if (InProperty == EDMMaterialPropertyType::BaseColor)
-			{
-				if (UDMMaterialProperty* Property = EditorOnlyData->GetMaterialProperty(EDMMaterialPropertyType::EmissiveColor))
-				{
-					return Property->IsValidForModel(*EditorOnlyData);
-				}
-			}
+	if (InProperty == EDMMaterialPropertyType::BaseColor)
+	{
+		if (UDMMaterialProperty* Property = EditorOnlyData->GetMaterialProperty(EDMMaterialPropertyType::EmissiveColor))
+		{
+			return Property->IsValidForModel(*EditorOnlyData);
+		}
+	}
 
-			if (InProperty == EDMMaterialPropertyType::Opacity)
-			{
-				if (UDMMaterialProperty* Property = EditorOnlyData->GetMaterialProperty(EDMMaterialPropertyType::OpacityMask))
-				{
-					return Property->IsValidForModel(*EditorOnlyData);
-				}
-			}
+	if (InProperty == EDMMaterialPropertyType::Opacity)
+	{
+		if (UDMMaterialProperty* Property = EditorOnlyData->GetMaterialProperty(EDMMaterialPropertyType::OpacityMask))
+		{
+			return Property->IsValidForModel(*EditorOnlyData);
 		}
 	}
 
@@ -1006,37 +943,35 @@ bool SDMEditor::IsPropertyValidForModel(EDMMaterialPropertyType InProperty) cons
 
 ECheckBoxState SDMEditor::GetSlotCheckState(EDMMaterialPropertyType InProperty)  const
 {
-	if (SlotWidgets.IsValidIndex(ActiveSlotIndex))
+	if (InProperty == EDMMaterialPropertyType::None)
 	{
-		if (UDMMaterialSlot* Slot = SlotWidgets[ActiveSlotIndex]->GetSlot())
-		{
-			if (UDynamicMaterialModelEditorOnlyData* EditorOnlyData = Slot->GetMaterialModelEditorOnlyData())
-			{
-				if (UDMMaterialSlot* PropertySlot = EditorOnlyData->GetSlotForMaterialProperty(InProperty))
-				{
-					return PropertySlot == Slot
-						? ECheckBoxState::Checked
-						: ECheckBoxState::Unchecked;
-				}
+		return ActiveSlotIndex == INDEX_NONE
+			? ECheckBoxState::Checked
+			: ECheckBoxState::Unchecked;
+	}
 
-				if (InProperty == EDMMaterialPropertyType::BaseColor)
-				{
-					return EditorOnlyData->GetSlotForMaterialProperty(EDMMaterialPropertyType::EmissiveColor) == Slot
-						? ECheckBoxState::Checked
-						: ECheckBoxState::Unchecked;
-				}
+	if (ActiveSlotIndex == INDEX_NONE)
+	{
+		return ECheckBoxState::Unchecked;
+	}
 
-				if (InProperty == EDMMaterialPropertyType::Opacity)
-				{
-					return EditorOnlyData->GetSlotForMaterialProperty(EDMMaterialPropertyType::OpacityMask) == Slot
-						? ECheckBoxState::Checked
-						: ECheckBoxState::Unchecked;
-				}	
-			}
-		}
+	UDynamicMaterialModelEditorOnlyData* EditorOnlyData = UDynamicMaterialModelEditorOnlyData::Get(MaterialModelWeak);
+
+	if (!EditorOnlyData)
+	{
+		return ECheckBoxState::Undetermined;
+	}
+
+	if (!EditorOnlyData->GetSlots().IsValidIndex(ActiveSlotIndex))
+	{
+		return ECheckBoxState::Unchecked;
 	}
 
-	return ECheckBoxState::Undetermined;
+	UDMMaterialSlot* PropertySlot = GetSlotForMaterialProperty(InProperty);
+
+	return EditorOnlyData->GetSlots()[ActiveSlotIndex] == PropertySlot
+		? ECheckBoxState::Checked
+		: ECheckBoxState::Unchecked;
 }
 
 FText SDMEditor::GetToolTipForProperty(EDMMaterialPropertyType InProperty) const
@@ -1051,41 +986,46 @@ FText SDMEditor::GetToolTipForProperty(EDMMaterialPropertyType InProperty) const
 
 void SDMEditor::OnSlotCheckStateChanged(ECheckBoxState InCheckState, EDMMaterialPropertyType InProperty)
 {
-	if (InCheckState == ECheckBoxState::Checked)
+	if (InCheckState != ECheckBoxState::Checked)
 	{
-		for (int32 SlotIndex = 0; SlotIndex < SlotWidgets.Num(); ++SlotIndex)
-		{
-			if (UDMMaterialSlot* Slot = SlotWidgets[SlotIndex]->GetSlot())
-			{
-				if (UDynamicMaterialModelEditorOnlyData* EditorOnlyData = Slot->GetMaterialModelEditorOnlyData())
-				{
-					if (UDMMaterialSlot* PropertySlot = EditorOnlyData->GetSlotForMaterialProperty(InProperty))
-					{
-						if (PropertySlot == Slot)
-						{
-							SetActiveSlotIndex(SlotIndex);
-						}
-					}
+		return;
+	}
 
-					if (InProperty == EDMMaterialPropertyType::BaseColor)
-					{
-						if (EditorOnlyData->GetSlotForMaterialProperty(EDMMaterialPropertyType::EmissiveColor) == Slot)
-						{
-							SetActiveSlotIndex(SlotIndex);
-						}
-					}
+	if (InProperty == EDMMaterialPropertyType::None)
+	{
+		SetActiveSlotIndex(INDEX_NONE);
+		return;
+	}
 
-					if (InProperty == EDMMaterialPropertyType::Opacity)
-					{
-						if (EditorOnlyData->GetSlotForMaterialProperty(EDMMaterialPropertyType::OpacityMask) == Slot)
-						{
-							SetActiveSlotIndex(SlotIndex);
-						}
-					}
-				}
-			}
+	UDynamicMaterialModelEditorOnlyData* ModelEditorOnlyData = UDynamicMaterialModelEditorOnlyData::Get(MaterialModelWeak);
+
+	if (!ModelEditorOnlyData)
+	{
+		SetActiveSlotIndex(INDEX_NONE);
+		return;
+	}
+
+	UDMMaterialSlot* PropertySlot = GetSlotForMaterialProperty(InProperty);
+
+	const TArray<UDMMaterialSlot*>& Slots = ModelEditorOnlyData->GetSlots();
+
+	if (Slots.IsEmpty())
+	{
+		SetActiveSlotIndex(INDEX_NONE);
+		return;
+	}
+
+	for (int32 SlotIndex = 0; SlotIndex < Slots.Num(); ++SlotIndex)
+	{
+		if (Slots[SlotIndex] == PropertySlot)
+		{
+			SetActiveSlotIndex(SlotIndex);
+			return;
 		}
 	}
+
+	SetActiveSlotIndex(INDEX_NONE);
+	return;
 }
 
 FReply SDMEditor::OnCreateMaterialButtonClicked(TWeakPtr<FDMObjectMaterialProperty> InMaterialProperty)
@@ -1126,35 +1066,34 @@ void SDMEditor::OnSettingsChanged(const FPropertyChangedEvent& InPropertyChanged
 
 bool SDMEditor::CanAddNewLayer() const
 {
-	return SlotWidgets.IsValidIndex(ActiveSlotIndex)
-		&& MaterialModelWeak.IsValid();
+	return ActiveSlotWidget.IsValid() && !!UDynamicMaterialModelEditorOnlyData::Get(MaterialModelWeak);
 }
 
 void SDMEditor::AddNewLayer()
 {
-	if (SlotWidgets.IsValidIndex(ActiveSlotIndex))
+	if (ActiveSlotWidget.IsValid())
 	{
-		SlotWidgets[ActiveSlotIndex]->AddNewLayer_Expression(
+		ActiveSlotWidget->AddNewLayer_Expression(
 			TSubclassOf<UDMMaterialStageExpression>(UDMMaterialStageExpressionTextureSample::StaticClass()),
 			EDMMaterialLayerStage::All
 		);
 
-		SlotWidgets[ActiveSlotIndex]->InvalidateMainWidget();
+		ActiveSlotWidget->InvalidateMainWidget();
 	}
 }
 
 bool SDMEditor::CanInsertNewLayer() const
 {
-	return SlotWidgets.IsValidIndex(ActiveSlotIndex)
-		&& SlotWidgets[ActiveSlotIndex]->GetSelectedLayerIndices().Num() == 1
+	return ActiveSlotWidget.IsValid()
+		&& ActiveSlotWidget->GetSelectedLayerIndices().Num() == 1
 		&& MaterialModelWeak.IsValid();
 }
 
 void SDMEditor::InsertNewLayer()
 {
-	if (SlotWidgets.IsValidIndex(ActiveSlotIndex))
+	if (ActiveSlotWidget.IsValid())
 	{
-		if (UDMMaterialLayerObject* SelectedLayer = SlotWidgets[ActiveSlotIndex]->GetSelectedLayer())
+		if (UDMMaterialLayerObject* SelectedLayer = ActiveSlotWidget->GetSelectedLayer())
 		{
 			if (UDMMaterialSlot* Slot = SelectedLayer->GetSlot())
 			{
@@ -1162,14 +1101,14 @@ void SDMEditor::InsertNewLayer()
 				FDMScopedUITransaction Transaction(LOCTEXT("InsertNewLayer", "Material Designer Insert Layer"));
 				Slot->Modify();
 
-				SlotWidgets[ActiveSlotIndex]->AddNewLayer_Expression(
+				ActiveSlotWidget->AddNewLayer_Expression(
 					TSubclassOf<UDMMaterialStageExpression>(UDMMaterialStageExpressionTextureSample::StaticClass()),
 					EDMMaterialLayerStage::All
 				);
 
 				Slot->MoveLayerAfter(Slot->GetLayers().Last(), SelectedLayer);
 
-				SlotWidgets[ActiveSlotIndex]->InvalidateMainWidget();
+				ActiveSlotWidget->InvalidateMainWidget();
 			}
 		}
 	}
@@ -1177,17 +1116,16 @@ void SDMEditor::InsertNewLayer()
 
 bool SDMEditor::CanCopySelectedLayer() const
 {
-	return SlotWidgets.IsValidIndex(ActiveSlotIndex)
-		&& SlotWidgets[ActiveSlotIndex]->GetSelectedLayerIndices().Num() == 1
+	return ActiveSlotWidget.IsValid()
+		&& ActiveSlotWidget->GetSelectedLayerIndices().Num() == 1
 		&& MaterialModelWeak.IsValid();
 }
 
 void SDMEditor::CopySelectedLayer()
 {
-	if (SlotWidgets.IsValidIndex(ActiveSlotIndex)
-		&& SlotWidgets[ActiveSlotIndex]->GetSelectedLayerIndices().Num() == 1)
+	if (ActiveSlotWidget.IsValid() && ActiveSlotWidget->GetSelectedLayerIndices().Num() == 1)
 	{
-		if (const UDMMaterialLayerObject* Layer = SlotWidgets[ActiveSlotIndex]->GetSelectedLayer())
+		if (const UDMMaterialLayerObject* Layer = ActiveSlotWidget->GetSelectedLayer())
 		{
 			FPlatformApplicationMisc::ClipboardCopy(*Layer->SerializeToString());
 		}
@@ -1233,7 +1171,7 @@ void SDMEditor::PasteLayer()
 				Slot->Modify();
 				Slot->PasteLayer(PastedLayer);
 
-				SlotWidgets[ActiveSlotIndex]->InvalidateMainWidget();
+				ActiveSlotWidget->InvalidateMainWidget();
 			}
 		}
 	}
@@ -1261,10 +1199,9 @@ void SDMEditor::DuplicateSelectedLayer()
 
 bool SDMEditor::CanDeleteSelectedLayer() const
 {
-	if (SlotWidgets.IsValidIndex(ActiveSlotIndex)
-		&& SlotWidgets[ActiveSlotIndex]->GetSelectedLayerIndices().IsEmpty() == false)
+	if (ActiveSlotWidget.IsValid() && ActiveSlotWidget->GetSelectedLayerIndices().IsEmpty() == false)
 	{
-		return SlotWidgets[ActiveSlotIndex]->GetLayerRowsButtonsCanRemove();
+		return ActiveSlotWidget->GetLayerRowsButtonsCanRemove();
 	}
 
 	return false;
@@ -1272,10 +1209,9 @@ bool SDMEditor::CanDeleteSelectedLayer() const
 
 void SDMEditor::DeleteSelectedLayer()
 {
-	if (SlotWidgets.IsValidIndex(ActiveSlotIndex)
-		&& SlotWidgets[ActiveSlotIndex]->GetSelectedLayerIndices().IsEmpty() == false)
+	if (ActiveSlotWidget.IsValid() && ActiveSlotWidget->GetSelectedLayerIndices().IsEmpty() == false)
 	{
-		SlotWidgets[ActiveSlotIndex]->OnLayerRowButtonsRemoveClicked();
+		ActiveSlotWidget->OnLayerRowButtonsRemoveClicked();
 	}
 }
 
@@ -1317,9 +1253,9 @@ void SDMEditor::PostRedo(bool bSuccess)
 	OnUndo();
 }
 
-TSharedRef<SWidget> SDMEditor::CreateSlotsWidget()
+TSharedRef<SWidget> SDMEditor::CreateSlotWidget()
 {
-	SlotWidgets.Empty();
+	ActiveSlotWidget.Reset();
 
 	auto CreateEmptySlotsContent = []() -> TSharedRef<SBox>
 	{
@@ -1330,10 +1266,15 @@ TSharedRef<SWidget> SDMEditor::CreateSlotsWidget()
 			.Padding(10.f, 5.f, 10.f, 5.f)
 			[
 				SNew(STextBlock)
-				.Text(LOCTEXT("SlotsContent", "Slots Content"))
+				.Text(LOCTEXT("SlotsContent", "Select a Material Property to view the Layer Stack."))
 				.Font(IDetailLayoutBuilder::GetDetailFont())
 			];
 	};
+
+	if (ActiveSlotIndex == INDEX_NONE)
+	{
+		return CreateEmptySlotsContent();
+	}
 	
 	UDynamicMaterialModelEditorOnlyData* ModelEditorOnlyData = UDynamicMaterialModelEditorOnlyData::Get(MaterialModelWeak);
 
@@ -1344,53 +1285,144 @@ TSharedRef<SWidget> SDMEditor::CreateSlotsWidget()
 
 	const TArray<UDMMaterialSlot*>& Slots = ModelEditorOnlyData->GetSlots();
 
-	if (Slots.IsEmpty())
+	if (Slots.IsEmpty() || !Slots.IsValidIndex(ActiveSlotIndex))
 	{
 		return CreateEmptySlotsContent();
 	}
 
-	SlotWidgets.Reserve(Slots.Num());
+	ActiveSlotWidget =
+		SNew(SDMSlot, SharedThis(this), Slots[ActiveSlotIndex])
+		.SlotPreviewSize_Lambda([]() { return UDynamicMaterialEditorSettings::Get()->SlotPreviewSize; })
+		.LayerPreviewSize_Lambda([]() { return UDynamicMaterialEditorSettings::Get()->LayerPreviewSize; });
+
+	return ActiveSlotWidget.ToSharedRef();
+}
+
+void SDMEditor::OnUndo()
+{
+	UDynamicMaterialModel* MaterialModel = MaterialModelWeak.Get();
+
+	if (!IsValid(MaterialModel))
+	{
+		return;
+	}
 
-	for (int32 SlotIdx = 0; SlotIdx < Slots.Num(); ++SlotIdx)
+	if (UDynamicMaterialModelEditorOnlyData* ModelEditorOnlyData = UDynamicMaterialModelEditorOnlyData::Get(MaterialModelWeak))
 	{
-		UDMMaterialSlot* Slot = Slots[SlotIdx];
+		if (ModelEditorOnlyData->GetBlendMode() == BLEND_Opaque)
+		{
+			SetActiveSlotIndex(0);
+		}
+	}
+}
 
-		TSharedRef<SDMSlot> SlotWidget =
-			SNew(SDMSlot, SharedThis(this), Slot)
-			.SlotPreviewSize_Lambda([]() { return UDynamicMaterialEditorSettings::Get()->SlotPreviewSize; })
-			.LayerPreviewSize_Lambda([]() { return UDynamicMaterialEditorSettings::Get()->LayerPreviewSize; });
+UDMMaterialSlot* SDMEditor::GetSlotForMaterialProperty(EDMMaterialPropertyType InProperty) const
+{
+	UDynamicMaterialModelEditorOnlyData* ModelEditorOnlyData = UDynamicMaterialModelEditorOnlyData::Get(MaterialModelWeak);
 
-		SlotWidgets.Add(SlotWidget);
+	if (!ModelEditorOnlyData)
+	{
+		return nullptr;
 	}
 
-	if (!SlotWidgets.IsValidIndex(ActiveSlotIndex))
+	UDMMaterialSlot* PropertySlot = ModelEditorOnlyData->GetSlotForMaterialProperty(InProperty);
+
+	if (!PropertySlot)
 	{
-		if (SlotWidgets.IsEmpty())
+		switch (InProperty)
 		{
-			return SNullWidget::NullWidget;
+			case EDMMaterialPropertyType::BaseColor:
+				PropertySlot = ModelEditorOnlyData->GetSlotForMaterialProperty(EDMMaterialPropertyType::EmissiveColor);
+				break;
+
+			case EDMMaterialPropertyType::EmissiveColor:
+				PropertySlot = ModelEditorOnlyData->GetSlotForMaterialProperty(EDMMaterialPropertyType::BaseColor);
+				break;
+
+			case EDMMaterialPropertyType::Opacity:
+				PropertySlot = ModelEditorOnlyData->GetSlotForMaterialProperty(EDMMaterialPropertyType::OpacityMask);
+				break;
+
+			case EDMMaterialPropertyType::OpacityMask:
+				PropertySlot = ModelEditorOnlyData->GetSlotForMaterialProperty(EDMMaterialPropertyType::Opacity);
+				break;
 		}
+	}
+
+	return PropertySlot;
+}
+
+UDMMaterialComponent* SDMEditor::GetEditedComponent() const
+{
+	return EditedComponent.Get();
+}
 
-		ActiveSlotIndex = 0;
+void SDMEditor::SetEditedComponent(UDMMaterialComponent* InComponent)
+{
+	if (EditedComponent.IsValid())
+	{
+		EditedComponent->GetOnUpdate().RemoveAll(this);
 	}
 
-	return SlotWidgets[ActiveSlotIndex];
+	EditedComponent = InComponent;
+
+	if (EditedComponent.IsValid())
+	{
+		EditedComponent->GetOnUpdate().AddSP(this, &SDMEditor::OnComponentUpdated);
+	}
+
+	InvalidateComponentEditWidget();
 }
 
-void SDMEditor::OnUndo()
+TSharedRef<SWidget> SDMEditor::CreateComponentEditWidget()
 {
-	UDynamicMaterialModel* MaterialModel = MaterialModelWeak.Get();
+	return
+		SNew(SBorder)
+		.HAlign(HAlign_Fill)
+		.VAlign(VAlign_Fill)
+		.Padding(3.0f)
+		.BorderImage(FDynamicMaterialEditorStyle::GetBrush("LayerView.Details.Background"))
+		[
+			SNew(SDMComponentEdit, EditedComponent.Get(), SharedThis(this))
+		];
+}
 
-	if (!IsValid(MaterialModel))
+void SDMEditor::InvalidateComponentEditWidget()
+{
+	bInvalidateComponentEditWidget = true;
+}
+
+void SDMEditor::RefreshComponentEditWidget()
+{
+	bInvalidateComponentEditWidget = false;
+
+	if (ComponentEditContainer.IsValid())
 	{
-		return;
+		ComponentEditContainer->ClearChildren();
+		ComponentEditContainer->AddSlot()
+			.AutoSize()
+			[
+				CreateComponentEditWidget()
+			];
 	}
+}
 
-	if (UDynamicMaterialModelEditorOnlyData* ModelEditorOnlyData = UDynamicMaterialModelEditorOnlyData::Get(MaterialModelWeak))
+void SDMEditor::OnSplitterResized() const
+{
+	UDynamicMaterialEditorSettings* Settings = UDynamicMaterialEditorSettings::Get();
+
+	if (LayerViewSplitterSlot)
 	{
-		if (ModelEditorOnlyData->GetBlendMode() == BLEND_Opaque)
-		{
-			SetActiveSlotIndex(0);
-		}
+		Settings->SplitterLocation = LayerViewSplitterSlot->GetSizeValue();
+		Settings->SaveConfig();
+	}
+}
+
+void SDMEditor::OnComponentUpdated(UDMMaterialComponent* InComponent, EDMUpdateType InUpdateType)
+{
+	if (InUpdateType == EDMUpdateType::Structure)
+	{
+		InvalidateComponentEditWidget();
 	}
 }
 
diff --git a/Engine/Plugins/Experimental/DynamicMaterial/Source/DynamicMaterialEditor/Private/Slate/SDMEditor.h b/Engine/Plugins/Experimental/DynamicMaterial/Source/DynamicMaterialEditor/Private/Slate/SDMEditor.h
--- a/Engine/Plugins/Experimental/DynamicMaterial/Source/DynamicMaterialEditor/Private/Slate/SDMEditor.h
+++ b/Engine/Plugins/Experimental/DynamicMaterial/Source/DynamicMaterialEditor/Private/Slate/SDMEditor.h
@@ -2,12 +2,14 @@
 
 #pragma once
 
+#include "Widgets/SCompoundWidget.h"
 #include "Components/DMMaterialStage.h"
 #include "DMEDefs.h"
 #include "DMObjectMaterialProperty.h"
 #include "EditorUndoClient.h"
-#include "Widgets/SCompoundWidget.h"
+#include "Widgets/Layout/SSplitter.h"
 
+class FAssetThumbnailPool;
 class FUICommandList;
 class IDetailTreeNode;
 class IPropertyHandle;
@@ -58,13 +60,18 @@ public:
 
 	void SetMaterialActor(AActor* InActor);
 
-	const TArray<TSharedRef<SDMSlot>>& GetSlotWidgets() const { return SlotWidgets; }
-	TSharedPtr<SDMSlot> GetSlotWidget(UDMMaterialSlot* Slot) const;
+	TSharedPtr<SDMSlot> GetActiveSlotWidget() const;
 
-	void RefreshGlobalDetailsView();
 	void RefreshParametersList();
 	void RefreshSlotPickerList();
-	void RefreshSlotsList();
+	void RefreshSlotWidget();
+	void RefreshComponentEditWidget();
+
+	void InvalidateComponentEditWidget();
+
+	UDMMaterialComponent* GetEditedComponent() const;
+
+	void SetEditedComponent(UDMMaterialComponent* InComponent);
 
 	//~ Begin SWidget
 	virtual bool SupportsKeyboardFocus() const override { return true; }
@@ -115,30 +122,35 @@ protected:
 
 	static FDMPropertyHandle CreatePropertyHandle(const void* InOwningWidget, UObject* InObject, const FName& InPropertyName);
 
-	TSharedPtr<FUICommandList> CommandList;
-
 	TSharedPtr<SBox> Container;
 	TSharedPtr<SDMToolBar> Toolbar;
-	TSharedPtr<SBox> GlobalOpacityContainer;
 	TSharedPtr<SDMMaterialParameters> ParametersWidget;
-	TSharedPtr<SBox> SlotsContainer;
 	TSharedPtr<SBox> SlotPickerContainer;
-	TSharedPtr<SWidgetSwitcher> SlotSwitcher;
-	int32 ActiveSlotIndex;
+	TSharedPtr<SSplitter> SplitterContainer;
+	TSharedPtr<SScrollBox> SlotContainer;
+	TSharedPtr<SScrollBox> ComponentEditContainer;
+	SSplitter::FSlot* LayerViewSplitterSlot = nullptr;
+	SSplitter::FSlot* ExtraSpaceSplitterSlot = nullptr;
+	TSharedPtr<SDMSlot> ActiveSlotWidget;
 
+	TSharedPtr<FUICommandList> CommandList;
+	int32 ActiveSlotIndex;
 	TWeakObjectPtr<UDynamicMaterialModel> MaterialModelWeak;
 	FDMObjectMaterialProperty ObjectProperty;
-
-	TArray<TSharedRef<SDMSlot>> SlotWidgets;
+	TWeakObjectPtr<UDMMaterialComponent> EditedComponent;
+	bool bInvalidateComponentEditWidget;
 
 	void BindCommands();
 
 	TSharedRef<SWidget> CreateMainLayout();
-	TSharedRef<SWidget> CreateGlobalDetailsView();
 	TSharedRef<SWidget> CreateParametersArea();
-
 	TSharedRef<SWidget> CreateSlotPickerWidget();
-	TSharedRef<SWidget> CreateSlotsWidget();
+	TSharedRef<SWidget> CreateSlotWidget();
+	TSharedRef<SWidget> CreateComponentEditWidget();
+
+	void OnSplitterResized() const;
+
+	void OnComponentUpdated(UDMMaterialComponent* InComponent, EDMUpdateType InUpdateType);
 
 	TSharedRef<SWidget> CreateActorMaterialSlotSelector(const AActor* InActor);
 
@@ -146,8 +158,6 @@ protected:
 	void OnValuesUpdated(UDynamicMaterialModel* InMaterialModel);
 	void OnSlotsUpdated(UDynamicMaterialModel* InMaterialModel);
 
-	EVisibility GetGlobalOpacityVisibility() const;
-
 	bool IsPropertyValidForModel(EDMMaterialPropertyType InProperty) const;
 	ECheckBoxState GetSlotCheckState(EDMMaterialPropertyType InProperty) const;
 	FText GetToolTipForProperty(EDMMaterialPropertyType InProperty) const;
@@ -161,4 +171,6 @@ protected:
 	void OnSettingsChanged(const FPropertyChangedEvent& InPropertyChangedEvent);
 
 	void OnUndo();
+
+	UDMMaterialSlot* GetSlotForMaterialProperty(EDMMaterialPropertyType InProperty) const;
 };
diff --git a/Engine/Plugins/Experimental/DynamicMaterial/Source/DynamicMaterialEditor/Private/Slate/SDMSlot.cpp b/Engine/Plugins/Experimental/DynamicMaterial/Source/DynamicMaterialEditor/Private/Slate/SDMSlot.cpp
--- a/Engine/Plugins/Experimental/DynamicMaterial/Source/DynamicMaterialEditor/Private/Slate/SDMSlot.cpp
+++ b/Engine/Plugins/Experimental/DynamicMaterial/Source/DynamicMaterialEditor/Private/Slate/SDMSlot.cpp
@@ -9,7 +9,6 @@
 #include "Components/DMMaterialStageFunction.h"
 #include "Components/DMMaterialStageGradient.h"
 #include "Components/DMMaterialStageThroughputLayerBlend.h"
-#include "Components/DMMaterialSubStage.h"
 #include "Components/DMMaterialValue.h"
 #include "Components/DMRenderTargetRenderer.h"
 #include "Components/MaterialStageBlends/DMMSBNormal.h"
@@ -22,7 +21,6 @@
 #include "Components/MaterialStageInputs/DMMSITextureUV.h"
 #include "Components/MaterialStageInputs/DMMSIValue.h"
 #include "Components/MaterialValues/DMMaterialValueFloat1.h"
-#include "Components/MaterialValues/DMMaterialValueRenderTarget.h"
 #include "DetailLayoutBuilder.h"
 #include "DMBlueprintFunctionLibrary.h"
 #include "DMPrivate.h"
@@ -110,23 +108,10 @@ void SDMSlot::Construct(const FArguments& InArgs, const TSharedRef<SDMEditor>& I
 		bInvalidateMainWidget = false;
 		bInvalidateHeaderWidget = false;
 		bInvalidateSettingsWidget = false;
-		bInvalidateComponentEditWidget = false;
 		OnEndFrameDelegateHandle = FCoreDelegates::OnEndFrame.AddSP(this, &SDMSlot::HandleEndFrameRefresh);
 	}
 }
 
-void SDMSlot::OnLayerSelected(TSharedPtr<FDMMaterialLayerReference> InLayerItem, const int32 InLayerIndex)
-{
-	InvalidateSlotSettingsRowWidget();
-	InvalidateComponentEditWidget();
-}
-
-void SDMSlot::OnLayerStageSelected(const bool bInSelected, const TSharedRef<SDMStage>& InStageWidget)
-{
-	InvalidateSlotSettingsRowWidget();
-	SetEditedComponent(InStageWidget->GetStage());
-}
-
 void SDMSlot::HandleEndFrameRefresh()
 {
 	if (bInvalidateMainWidget)
@@ -144,11 +129,6 @@ void SDMSlot::HandleEndFrameRefresh()
 		{
 			RefreshSlotSettingsRowWidget();
 		}
-
-		if (bInvalidateComponentEditWidget)
-		{
-			RefreshComponentEditWidget();
-		}
 	}
 }
 
@@ -264,28 +244,6 @@ void SDMSlot::SetSelectedLayer(UDMMaterialLayerObject* InLayer) const
 	LayerView->SelectLayerItem(InLayer, /* Mask */ false, /* Selected */ true, ESelectInfo::Direct);
 }
 
-UDMMaterialComponent* SDMSlot::GetEditedComponent() const
-{
-	return EditedComponent.Get();
-}
-
-void SDMSlot::SetEditedComponent(UDMMaterialComponent* InComponent)
-{
-	if (EditedComponent.IsValid())
-	{
-		EditedComponent->GetOnUpdate().RemoveAll(this);
-	}
-
-	EditedComponent = InComponent;
-
-	if (EditedComponent.IsValid())
-	{
-		EditedComponent->GetOnUpdate().AddSP(this, &SDMSlot::OnComponentUpdated);
-	}
-
-	InvalidateComponentEditWidget();
-}
-
 int32 SDMSlot::GetSourceBlendTypeSwitcherIndex() const
 {
 	if (LayerView->GetLayerItems().IsValidIndex(LayerView->GetSelectedLayerIndex()))
@@ -391,7 +349,6 @@ void SDMSlot::RefreshMainWidget()
 	bInvalidateMainWidget = false;
 	bInvalidateHeaderWidget = false;
 	bInvalidateSettingsWidget = false;
-	bInvalidateComponentEditWidget = false;
 
 	const float SplitterLocation = FMath::Clamp(UDynamicMaterialEditorSettings::Get()->SplitterLocation, 0.3f, 0.7f);
 	const float LayerViewSplitterLocation = SplitterLocation;
@@ -403,20 +360,11 @@ void SDMSlot::RefreshMainWidget()
 		.PreviewSize(LayerPreviewSize);
 
 	ChildSlot.AttachWidget(SNullWidget::NullWidget);
-	SplitterContainer.Reset();
-
-	SplitterContainer =
-		SNew(SSplitter)
-		.Style(FAppStyle::Get(), "DetailsView.Splitter")
-		.Orientation(Orient_Vertical)
-		.ResizeMode(ESplitterResizeMode::Fill)
-		.PhysicalSplitterHandleSize(3.0f)
-		.HitDetectionSplitterHandleSize(4.0f)
-		.OnSplitterFinishedResizing(this, &SDMSlot::OnSplitterResized)
-
-		+ SSplitter::Slot()
-		.Resizable(false)
-		.SizeRule(SSplitter::ESizeRule::SizeToContent)
+
+	TSharedRef<SVerticalBox> VerticalBox = SNew(SVerticalBox)
+
+		+ SVerticalBox::Slot()
+		.AutoHeight()
 		[
 			SAssignNew(SlotSettingsRowContainer, SBox)
 			.Padding(3.0f, 4.0f, 3.0f, 4.0f)
@@ -425,44 +373,24 @@ void SDMSlot::RefreshMainWidget()
 			]
 		]
 
-		+ SSplitter::Slot()
-			.Expose(LayerViewSplitterSlot)
-			.Resizable(true)
-			.SizeRule(SSplitter::ESizeRule::FractionOfParent)
-			.MinSize(50)
-			.Value(LayerViewSplitterLocation)
-			[
-				SNew(SBorder)
-				.Padding(2.0f)
-				.BorderImage(FDynamicMaterialEditorStyle::GetBrush("LayerView.Background"))
-				.HAlign(HAlign_Fill)
-				.VAlign(VAlign_Fill)
-				[
-					LayerView.ToSharedRef()
-				]
-			]
-
-		+ SSplitter::Slot()
-			.Resizable(false)
-			.SizeRule(SSplitter::ESizeRule::SizeToContent)
+		+ SVerticalBox::Slot()
+		.FillHeight(1.f)
+		[
+			SNew(SBorder)
+			.Padding(2.0f)
+			.BorderImage(FDynamicMaterialEditorStyle::GetBrush("LayerView.Background"))
+			.HAlign(HAlign_Fill)
+			.VAlign(VAlign_Fill)
 			[
-				CreateLayerButtonsRowWidget()
+				LayerView.ToSharedRef()
 			]
+		]
 
-		+ SSplitter::Slot()
-			.Expose(ExtraSpaceSplitterSlot)
-			.Resizable(true)
-			.SizeRule(SSplitter::ESizeRule::FractionOfParent)
-			.MinSize(50)
-			.Value(ExtraSpaceSplitterLocation)
-			[
-				SAssignNew(ComponentEditContainer, SScrollBox)
-				+ SScrollBox::Slot()
-				.AutoSize()
-				[
-					CreateComponentEditWidget()
-				]
-			];
+		+ SVerticalBox::Slot()
+		.AutoHeight()
+		[
+			CreateLayerButtonsRowWidget()
+		];
 
 	TSharedRef<SVerticalBox> Child = SNew(SVerticalBox);
 
@@ -487,7 +415,7 @@ void SDMSlot::RefreshMainWidget()
 		.VAlign(VAlign_Fill)
 		.Padding(0.0f, 2.0f, 0.0f, 0.0f)
 		[
-			SplitterContainer.ToSharedRef()
+			VerticalBox
 		];
 
 	ChildSlot
@@ -711,24 +639,6 @@ TSharedRef<SWidget> SDMSlot::CreateLayerButtonsRowWidget()
 		];
 }
 
-TSharedRef<SWidget> SDMSlot::CreateComponentEditWidget()
-{
-	if (UDMMaterialComponent* Component = EditedComponent.Get())
-	{
-		return
-			SNew(SBorder)
-			.HAlign(HAlign_Fill)
-			.VAlign(VAlign_Fill)
-			.Padding(3.0f)
-			.BorderImage(FDynamicMaterialEditorStyle::GetBrush("LayerView.Details.Background"))
-			[
-				SNew(SDMComponentEdit, Component, SharedThis(this))
-			];
-	}
-
-	return SNullWidget::NullWidget;
-}
-
 TSharedRef<SWidget> SDMSlot::CreateSlotSettingsRow()
 {
 	TSharedPtr<SDMPropertyEdit> OpacityPropertyEditWidget = nullptr;
@@ -892,7 +802,10 @@ void SDMSlot::InvalidateSlotSettingsRowWidget()
 
 void SDMSlot::InvalidateComponentEditWidget()
 {
-	bInvalidateComponentEditWidget = true;
+	if (TSharedPtr<SDMEditor> EditorWidget = GetEditorWidget())
+	{
+		EditorWidget->InvalidateComponentEditWidget();
+	}
 }
 
 void SDMSlot::AddNewLayer_NewLocalValue(EDMValueType InValueType)
@@ -924,32 +837,6 @@ void SDMSlot::RefreshSlotSettingsRowWidget()
 	}
 }
 
-void SDMSlot::RefreshComponentEditWidget()
-{
-	bInvalidateComponentEditWidget = false;
-
-	if (ComponentEditContainer.IsValid())
-	{
-		ComponentEditContainer->ClearChildren();
-		ComponentEditContainer->AddSlot()
-			.AutoSize()
-			[
-				CreateComponentEditWidget()
-			];
-	}
-}
-
-void SDMSlot::OnSplitterResized() const
-{
-	UDynamicMaterialEditorSettings* Settings = UDynamicMaterialEditorSettings::Get();
-
-	if (LayerViewSplitterSlot)
-	{
-		Settings->SplitterLocation = LayerViewSplitterSlot->GetSizeValue();
-		Settings->SaveConfig();
-	}
-}
-
 void SDMSlot::StartStageOpacityTransaction()
 {
 	StartTransaction(LOCTEXT("StageOpacityTransDescription", "Material Designer Value Scrubbing (Stage Opacity)"));
@@ -1065,14 +952,6 @@ void SDMSlot::OnSlotPropertiesUpdated(UDMMaterialSlot* InSlot)
 	InvalidateHeaderPropertyListWidget();
 }
 
-void SDMSlot::OnComponentUpdated(UDMMaterialComponent* InComponent, EDMUpdateType InUpdateType)
-{
-	if (InUpdateType == EDMUpdateType::Structure)
-	{
-		InvalidateComponentEditWidget();
-	}
-}
-
 TSharedPtr<SDMStage> SDMSlot::FindStageWidget(UDMMaterialStage* const InStage) const
 {
 	if (!LayerView.IsValid())
@@ -1626,4 +1505,20 @@ void SDMSlot::AddNewLayer_Renderer(TSubclassOf<UDMRenderTargetRenderer> InRender
 	UDMBlueprintFunctionLibrary::SetStageInputToRenderer(NewBase, InRendererClass, UDMMaterialStageBlend::InputB);
 }
 
+void SDMSlot::OnLayerSelected(TSharedPtr<FDMMaterialLayerReference> InLayerItem, const int32 InLayerIndex)
+{
+	InvalidateSlotSettingsRowWidget();
+	InvalidateComponentEditWidget();
+}
+
+void SDMSlot::OnLayerStageSelected(const bool bInSelected, const TSharedRef<SDMStage>& InStageWidget)
+{
+	InvalidateSlotSettingsRowWidget();
+
+	if (TSharedPtr<SDMEditor> EditorWidget = EditorWidgetWeak.Pin())
+	{
+		EditorWidget->SetEditedComponent(InStageWidget->GetStage());
+	}
+}
+
 #undef LOCTEXT_NAMESPACE
diff --git a/Engine/Plugins/Experimental/DynamicMaterial/Source/DynamicMaterialEditor/Private/Slate/SDMSlot.h b/Engine/Plugins/Experimental/DynamicMaterial/Source/DynamicMaterialEditor/Private/Slate/SDMSlot.h
--- a/Engine/Plugins/Experimental/DynamicMaterial/Source/DynamicMaterialEditor/Private/Slate/SDMSlot.h
+++ b/Engine/Plugins/Experimental/DynamicMaterial/Source/DynamicMaterialEditor/Private/Slate/SDMSlot.h
@@ -87,10 +87,6 @@ public:
 
 	void SetSelectedLayer(UDMMaterialLayerObject* InLayer) const;
 
-	UDMMaterialComponent* GetEditedComponent() const;
-
-	void SetEditedComponent(UDMMaterialComponent* InComponent);
-
 	UDMMaterialLayerObject* AddNewLayer(UDMMaterialStage* InNewBaseStage, UDMMaterialStage* InNewMaskStage = nullptr);
 
 protected:
@@ -102,11 +98,6 @@ protected:
 	TWeakObjectPtr<UDMMaterialSlot> SlotWeak;
 	TSharedPtr<SBox> HeaderPropertyListWidget;
 	TSharedPtr<SBox> SlotSettingsRowContainer;
-	TSharedPtr<SScrollBox> ComponentEditContainer;
-
-	TSharedPtr<SSplitter> SplitterContainer;
-	SSplitter::FSlot* LayerViewSplitterSlot = nullptr;
-	SSplitter::FSlot* ExtraSpaceSplitterSlot = nullptr;
 
 	TSharedPtr<SDMSlotLayerView> LayerView;
 
@@ -119,24 +110,17 @@ protected:
 	bool bInvalidateMainWidget;
 	bool bInvalidateHeaderWidget;
 	bool bInvalidateSettingsWidget;
-	bool bInvalidateComponentEditWidget;
 	FDelegateHandle OnEndFrameDelegateHandle;
 
-	TWeakObjectPtr<UDMMaterialComponent> EditedComponent;
-
 	TSharedRef<SWidget> CreateHeaderPropertyListWidget();
 	TSharedRef<SWidget> CreateLayerButtonsRowWidget();
-	TSharedRef<SWidget> CreateComponentEditWidget();
 
 	TSharedRef<SWidget> CreateSlotSettingsRow();
 
-	void OnSplitterResized() const;
-
 	void AddPropertyToSlot(EDMMaterialPropertyType Property);
 
 	void OnSlotLayersUpdated(UDMMaterialSlot* InSlot);
 	void OnSlotPropertiesUpdated(UDMMaterialSlot* InSlot);
-	void OnComponentUpdated(UDMMaterialComponent* InComponent, EDMUpdateType InUpdateType);
 
 	void StartStageOpacityTransaction();
 	void EndStageOpacityTransaction(const float InValue) { EndTransaction(); }
@@ -153,7 +137,6 @@ protected:
 	void RefreshMainWidget();
 	void RefreshHeaderPropertyListWidget();
 	void RefreshSlotSettingsRowWidget();
-	void RefreshComponentEditWidget();
 
 	void HandleEndFrameRefresh();
 
diff --git a/Engine/Plugins/Experimental/DynamicMaterial/Source/DynamicMaterialEditor/Public/Slate/SDMComponentEdit.h b/Engine/Plugins/Experimental/DynamicMaterial/Source/DynamicMaterialEditor/Public/Slate/SDMComponentEdit.h
--- a/Engine/Plugins/Experimental/DynamicMaterial/Source/DynamicMaterialEditor/Public/Slate/SDMComponentEdit.h
+++ b/Engine/Plugins/Experimental/DynamicMaterial/Source/DynamicMaterialEditor/Public/Slate/SDMComponentEdit.h
@@ -15,7 +15,8 @@ class IDetailTreeNode;
 class IPropertyHandle;
 class IPropertyRowGenerator;
 class SBox;
-class SDMSlot;
+class SDMEditor;
+class SDMEditor;
 class SDMStage;
 class SWidget;
 class UDMMaterialComponent;
@@ -30,6 +31,7 @@ class UDMMaterialStageSource;
 class UDMMaterialStageThroughput;
 class UDMMaterialValueFloat1;
 class UDMTextureUV;
+class UDynamicMaterialModelEditorOnlyData;
 class UMaterial;
 enum class ECheckBoxState : uint8;
 struct FDMPropertyHandle;
@@ -53,10 +55,10 @@ public:
 	SDMComponentEdit() = default;
 	virtual ~SDMComponentEdit() override;
 
-	void Construct(const FArguments& InArgs, UDMMaterialComponent* InComponent, const TWeakPtr<SDMSlot>& InSlotWidget);
+	void Construct(const FArguments& InArgs, UDMMaterialComponent* InComponent, const TWeakPtr<SDMEditor>& InEditorWidget);
 
 	FORCEINLINE UDMMaterialComponent* GetComponent() const { return ComponentWeak.Get(); }
-	FORCEINLINE TSharedPtr<SDMSlot> GetSlotWidget() const { return SlotWidgetWeak.Pin(); }
+	FORCEINLINE TSharedPtr<SDMEditor> GetEditorWidget() const { return EditorWidgetWeak.Pin(); }
 
 	TSharedPtr<SWidget> CreateSinglePropertyEditWidget(UDMMaterialComponent* InComponent, const FName& InPropertyName);
 
@@ -66,8 +68,11 @@ public:
 	//~ End FUndoClient
 
 protected:
+	static void GenerateMaterialModelPropertyRows(const TSharedRef<SDMEditor> InEditorWidget, UDynamicMaterialModel* InMaterialModel, 
+		TArray<FDMPropertyHandle>& InOutPropertyRows, TSet<UDMMaterialComponent*>& InOutProcessedObjects);
+
 	TWeakObjectPtr<UDMMaterialComponent> ComponentWeak;
-	TWeakPtr<SDMSlot> SlotWidgetWeak;
+	TWeakPtr<SDMEditor> EditorWidgetWeak;
 
 	TSharedPtr<IDetailKeyframeHandler> KeyframeHandler;
 
