Material Designer
- Global Scale renamed/refactored to Global Tiling
- Parameter moved from TextureUV function to generaterd material and added as an input to the TextureUV functions.

diff --git a/Engine/Plugins/Experimental/DynamicMaterial/Source/DynamicMaterial/Private/Model/DynamicMaterialModel.cpp b/Engine/Plugins/Experimental/DynamicMaterial/Source/DynamicMaterial/Private/Model/DynamicMaterialModel.cpp
--- a/Engine/Plugins/Experimental/DynamicMaterial/Source/DynamicMaterial/Private/Model/DynamicMaterialModel.cpp
+++ b/Engine/Plugins/Experimental/DynamicMaterial/Source/DynamicMaterial/Private/Model/DynamicMaterialModel.cpp
@@ -17,12 +17,12 @@
 #include "Materials/Material.h"
 #endif
 
-const FString UDynamicMaterialModel::ValuesPathToken          = FString(TEXT("Values"));
-const FString UDynamicMaterialModel::ParametersPathToken      = FString(TEXT("Parameters"));
-const FName UDynamicMaterialModel::GlobalOpacityPropertyName  = FName("GlobalOpacityValue");
-const FName UDynamicMaterialModel::GlobalOpacityParameterName = FName("GlobalOpacity");
-const FName UDynamicMaterialModel::GlobalScalePropertyName    = FName("GlobalScaleValue");
-const FName UDynamicMaterialModel::GlobalScaleParameterName   = FName("GlobalScale");
+const FString UDynamicMaterialModel::ValuesPathToken           = FString(TEXT("Values"));
+const FString UDynamicMaterialModel::ParametersPathToken       = FString(TEXT("Parameters"));
+const FName UDynamicMaterialModel::GlobalOpacityPropertyName   = FName("GlobalOpacityValue");
+const FName UDynamicMaterialModel::GlobalOpacityParameterName  = FName("GlobalOpacity");
+const FName UDynamicMaterialModel::GlobalTilingPropertyName    = FName("GlobalTilingValue");
+const FName UDynamicMaterialModel::GlobalTilingParameterName   = FName("GlobalTiling");
 
 UDynamicMaterialModel::UDynamicMaterialModel()
 {
@@ -47,18 +47,18 @@ UDynamicMaterialModel::UDynamicMaterialModel()
 
 	ParameterMap.Add(GlobalOpacityParameterName, GlobalOpacityParameter);
 
-	GlobalScaleValue = CreateDefaultSubobject<UDMMaterialValueFloat2>(GlobalScaleParameterName);
+	GlobalTilingValue = CreateDefaultSubobject<UDMMaterialValueFloat2>(GlobalTilingParameterName);
 
 #if WITH_EDITOR
-	GlobalScaleValue->SetDefaultValue(FVector2D(1.0, 1.0));
-	GlobalScaleValue->ApplyDefaultValue();
+	GlobalTilingValue->SetDefaultValue(FVector2D(1.0, 1.0));
+	GlobalTilingValue->ApplyDefaultValue();
 #endif
 
-	GlobalScaleParameter = CreateDefaultSubobject<UDMMaterialParameter>("GlobalScaleParameter");
-	GlobalScaleParameter->ParameterName = GlobalScaleParameterName;
-	GlobalScaleValue->Parameter = GlobalScaleParameter;
+	GlobalTilingParameter = CreateDefaultSubobject<UDMMaterialParameter>("GlobalTilingParameter");
+	GlobalTilingParameter->ParameterName = GlobalTilingParameterName;
+	GlobalTilingValue->Parameter = GlobalTilingParameter;
 
-	ParameterMap.Add(GlobalScaleParameterName, GlobalScaleParameter);
+	ParameterMap.Add(GlobalTilingParameterName, GlobalTilingParameter);
 }
 
 void UDynamicMaterialModel::SetDynamicMaterialInstance(UDynamicMaterialInstance* InDynamicMaterialInstance)
@@ -543,34 +543,34 @@ void UDynamicMaterialModel::FixGlobalVars()
 		GlobalOpacityValue->Parameter = GlobalOpacityParameter;
 	}
 
-	if (const TWeakObjectPtr<UDMMaterialParameter>* ParameterPtr = ParameterMap.Find(GlobalScaleParameterName))
+	if (const TWeakObjectPtr<UDMMaterialParameter>* ParameterPtr = ParameterMap.Find(GlobalTilingParameterName))
 	{
 		UDMMaterialParameter* Parameter = ParameterPtr->Get();
 
 		if (!Parameter || Parameter->HasAnyFlags(RF_ArchetypeObject)
-			|| (GlobalScaleParameter && Parameter != GlobalScaleParameter))
+			|| (GlobalTilingParameter && Parameter != GlobalTilingParameter))
 		{
-			if (GlobalScaleParameter)
+			if (GlobalTilingParameter)
 			{
-				ParameterMap[GlobalScaleParameterName] = GlobalScaleParameter;
+				ParameterMap[GlobalTilingParameterName] = GlobalTilingParameter;
 			}
 			else
 			{
-				ParameterMap.Remove(GlobalScaleParameterName);
+				ParameterMap.Remove(GlobalTilingParameterName);
 			}
 		}
 	}
 
-	if (GlobalScaleValue && GlobalScaleParameter != GlobalScaleValue->Parameter)
+	if (GlobalTilingValue && GlobalTilingParameter != GlobalTilingValue->Parameter)
 	{
 		if (GUndo)
 		{
-			GlobalScaleParameter->Modify();
-			GlobalScaleValue->Modify();
+			GlobalTilingParameter->Modify();
+			GlobalTilingValue->Modify();
 		}
 
-		GlobalScaleParameter->ParameterName = GlobalScaleParameterName;
-		GlobalScaleValue->Parameter = GlobalScaleParameter;
+		GlobalTilingParameter->ParameterName = GlobalTilingParameterName;
+		GlobalTilingValue->Parameter = GlobalTilingParameter;
 	}
 }
 
diff --git a/Engine/Plugins/Experimental/DynamicMaterial/Source/DynamicMaterial/Public/Model/DynamicMaterialModel.h b/Engine/Plugins/Experimental/DynamicMaterial/Source/DynamicMaterial/Public/Model/DynamicMaterialModel.h
--- a/Engine/Plugins/Experimental/DynamicMaterial/Source/DynamicMaterial/Public/Model/DynamicMaterialModel.h
+++ b/Engine/Plugins/Experimental/DynamicMaterial/Source/DynamicMaterial/Public/Model/DynamicMaterialModel.h
@@ -45,8 +45,8 @@ public:
 	static const FString ParametersPathToken;
 	static const FName GlobalOpacityPropertyName;
 	static const FName GlobalOpacityParameterName;
-	static const FName GlobalScalePropertyName;
-	static const FName GlobalScaleParameterName;
+	static const FName GlobalTilingPropertyName;
+	static const FName GlobalTilingParameterName;
 
 	UDynamicMaterialModel();
 
@@ -62,7 +62,7 @@ public:
 	UDMMaterialValueFloat1* GetGlobalOpacityValue() const { return GlobalOpacityValue; }
 
 	UFUNCTION(BlueprintPure, Category = "Material Designer")
-	UDMMaterialValueFloat2* GetGlobalScaleValue() const { return GlobalScaleValue; }
+	UDMMaterialValueFloat2* GetGlobalTilingValue() const { return GlobalTilingValue; }
 
 	UFUNCTION(BlueprintCallable, Category = "Material Designer")
 	UDMMaterialComponent* GetComponentByPath(const FString& InPath) const;
@@ -161,10 +161,10 @@ protected:
 	TObjectPtr<UDMMaterialParameter> GlobalOpacityParameter;
 
 	UPROPERTY(VisibleInstanceOnly, BlueprintReadOnly, Category = "Material Designer")
-	TObjectPtr<UDMMaterialValueFloat2> GlobalScaleValue;
+	TObjectPtr<UDMMaterialValueFloat2> GlobalTilingValue;
 
 	UPROPERTY(VisibleInstanceOnly, BlueprintReadOnly, Category = "Material Designer")
-	TObjectPtr<UDMMaterialParameter> GlobalScaleParameter;
+	TObjectPtr<UDMMaterialParameter> GlobalTilingParameter;
 
 	UPROPERTY(VisibleInstanceOnly, DuplicateTransient, TextExportTransient, Category = "Material Designer")
 	TMap<FName, TWeakObjectPtr<UDMMaterialParameter>> ParameterMap;
diff --git a/Engine/Plugins/Experimental/DynamicMaterial/Source/DynamicMaterialEditor/Private/Components/MaterialStageInputs/DMMSITextureUV.cpp b/Engine/Plugins/Experimental/DynamicMaterial/Source/DynamicMaterialEditor/Private/Components/MaterialStageInputs/DMMSITextureUV.cpp
--- a/Engine/Plugins/Experimental/DynamicMaterial/Source/DynamicMaterialEditor/Private/Components/MaterialStageInputs/DMMSITextureUV.cpp
+++ b/Engine/Plugins/Experimental/DynamicMaterial/Source/DynamicMaterialEditor/Private/Components/MaterialStageInputs/DMMSITextureUV.cpp
@@ -448,6 +448,12 @@ TArray<UMaterialExpression*> UDMMaterialStageInputTextureUV::CreateTextureUVExpr
 	CreateParameterExpression(ParamID::Scale,    ParamID::ScaleX,   TEXT("ScaleX"),   InTextureUV->GetScale().X);
 	CreateParameterExpression(ParamID::Scale,    ParamID::ScaleY,   TEXT("ScaleY"),   InTextureUV->GetScale().Y);
 
+	if (UMaterialExpression* GlobalTiling = InBuildState->GetGlobalTilingExpression())
+	{
+		check(NameToInputIndex.Contains(UDynamicMaterialModel::GlobalTilingParameterName));
+		GlobalTiling->ConnectExpression(TextureUVFunc->GetInput(NameToInputIndex[UDynamicMaterialModel::GlobalTilingParameterName]), 0);
+	}
+
 	// Output
 	Nodes.Add(TextureUVFunc);
 
diff --git a/Engine/Plugins/Experimental/DynamicMaterial/Source/DynamicMaterialEditor/Private/Menus/DMMaterialSlotLayerMenus.cpp b/Engine/Plugins/Experimental/DynamicMaterial/Source/DynamicMaterialEditor/Private/Menus/DMMaterialSlotLayerMenus.cpp
--- a/Engine/Plugins/Experimental/DynamicMaterial/Source/DynamicMaterialEditor/Private/Menus/DMMaterialSlotLayerMenus.cpp
+++ b/Engine/Plugins/Experimental/DynamicMaterial/Source/DynamicMaterialEditor/Private/Menus/DMMaterialSlotLayerMenus.cpp
@@ -110,6 +110,13 @@ namespace UE::DynamicMaterialEditor::Private
 						return;
 					}
 
+					TSharedPtr<SDMEditor> EditorWidget = MenuContext->GetEditorWidget();
+
+					if (!EditorWidget.IsValid() || !EditorWidget->GetActiveSlotWidget().IsValid())
+					{
+						return;
+					}
+
 					const TArray<UDMMaterialValue*>& Values = MaterialModel->GetValues();
 
 					for (UDMMaterialValue* Value : Values)
@@ -119,7 +126,7 @@ namespace UE::DynamicMaterialEditor::Private
 								Value->GetDescription(),
 								LOCTEXT("AddValueStageSpecificTooltip", "Add a Material Stage based on this Material Value."),
 								FSlateIcon(),
-								FUIAction(FExecuteAction::CreateSP(MenuContext->GetSlotWidget().Pin().Get(), &SDMSlot::AddNewLayer_GlobalValue, Value)))
+								FUIAction(FExecuteAction::CreateSP(EditorWidget->GetActiveSlotWidget().Get(), &SDMSlot::AddNewLayer_GlobalValue, Value)))
 						);
 					}
 				})
@@ -138,6 +145,13 @@ namespace UE::DynamicMaterialEditor::Private
 						return;
 					}
 
+					TSharedPtr<SDMEditor> EditorWidget = MenuContext->GetEditorWidget();
+
+					if (!EditorWidget.IsValid() || !EditorWidget->GetActiveSlotWidget().IsValid())
+					{
+						return;
+					}
+
 					for (EDMValueType ValueType : UDMValueDefinitionLibrary::GetValueTypes())
 					{
 						FText Name = UDMValueDefinitionLibrary::GetValueDefinition(ValueType).GetDisplayName();
@@ -148,7 +162,7 @@ namespace UE::DynamicMaterialEditor::Private
 								Name,
 								FormattedTooltip,
 								FSlateIcon(),
-								FUIAction(FExecuteAction::CreateSP(MenuContext->GetSlotWidget().Pin().Get(), &SDMSlot::AddNewLayer_NewGlobalValue, ValueType)))
+								FUIAction(FExecuteAction::CreateSP(EditorWidget->GetActiveSlotWidget().Get(), &SDMSlot::AddNewLayer_NewGlobalValue, ValueType)))
 						);
 					}
 				})
@@ -180,9 +194,9 @@ namespace UE::DynamicMaterialEditor::Private
 			return;
 		}
 
-		const TSharedPtr<SDMSlot> SlotWidget = MenuContext->GetSlotWidget().Pin();
+		TSharedPtr<SDMEditor> EditorWidget = MenuContext->GetEditorWidget();
 
-		if (!SlotWidget.IsValid())
+		if (!EditorWidget.IsValid() || !EditorWidget->GetActiveSlotWidget().IsValid())
 		{
 			return;
 		}
@@ -210,7 +224,7 @@ namespace UE::DynamicMaterialEditor::Private
 			if (ensure(MaterialProperty))
 			{
 				UE::DynamicMaterialEditor::Private::AddSlotMenuEntry(
-					SlotWidget.Get(),
+					EditorWidget->GetActiveSlotWidget().Get(),
 					InMenu,
 					MaterialProperty->GetDescription(),
 					InSlot,
@@ -234,14 +248,14 @@ namespace UE::DynamicMaterialEditor::Private
 			return;
 		}
 
-		const TSharedPtr<SDMSlot> SlotWidget = MenuContext->GetSlotWidget().Pin();
+		TSharedPtr<SDMEditor> EditorWidget = MenuContext->GetEditorWidget();
 
-		if (!SlotWidget.IsValid())
+		if (!EditorWidget.IsValid() || !EditorWidget->GetActiveSlotWidget().IsValid())
 		{
 			return;
 		}
 
-		const UDMMaterialSlot* const Slot = SlotWidget->GetSlot();
+		const UDMMaterialSlot* const Slot = EditorWidget->GetActiveSlotWidget()->GetSlot();
 
 		if (!Slot)
 		{
@@ -285,7 +299,7 @@ namespace UE::DynamicMaterialEditor::Private
 				if (ensure(MaterialProperty))
 				{
 					UE::DynamicMaterialEditor::Private::AddSlotMenuEntry(
-						SlotWidget.Get(),
+						EditorWidget->GetActiveSlotWidget().Get(),
 						InMenu,
 						FText::Format(SlotNameFormatTemplate, SlotIter->GetDescription(), MaterialProperty->GetDescription()),
 						SlotIter,
@@ -319,9 +333,9 @@ namespace UE::DynamicMaterialEditor::Private
 			return;
 		}
 
-		const TSharedPtr<SDMSlot> SlotWidget = MenuContext->GetSlotWidget().Pin();
+		TSharedPtr<SDMEditor> EditorWidget = MenuContext->GetEditorWidget();
 
-		if (!SlotWidget.IsValid())
+		if (!EditorWidget.IsValid() || !EditorWidget->GetActiveSlotWidget().IsValid())
 		{
 			return;
 		}
@@ -345,7 +359,7 @@ namespace UE::DynamicMaterialEditor::Private
 					LOCTEXT("ChangeGradientSourceTooltip", "Change the source of this stage to a Material Gradient."),
 					FSlateIcon(),
 					FUIAction(FExecuteAction::CreateSP(
-						SlotWidget.Get(),
+						EditorWidget->GetActiveSlotWidget().Get(),
 						&SDMSlot::AddNewLayer_Gradient,
 						TSubclassOf<UDMMaterialStageGradient>(Gradient.Get()))
 					)
@@ -359,7 +373,7 @@ UToolMenu* FDMMaterialSlotLayerMenus::GenerateSlotLayerMenu(const TSharedPtr<SDM
 {
 	using namespace UE::DynamicMaterialEditor::Private;
 
-	UToolMenu* NewToolMenu = UDMMenuContext::GenerateContextMenuLayer(SlotLayerMenuName, InSlotWidget, InLayerObject);
+	UToolMenu* NewToolMenu = UDMMenuContext::GenerateContextMenuLayer(SlotLayerMenuName, InSlotWidget->GetEditorWidget(), InLayerObject);
 
 	if (!NewToolMenu)
 	{
@@ -396,14 +410,14 @@ void FDMMaterialSlotLayerMenus::AddAddLayerSection(UToolMenu* InMenu)
 		return;
 	}
 
-	const TSharedPtr<SDMSlot> SlotWidget = MenuContext->GetSlotWidget().Pin();
+	TSharedPtr<SDMEditor> EditorWidget = MenuContext->GetEditorWidget();
 
-	if (!SlotWidget.IsValid())
+	if (!EditorWidget.IsValid() || !EditorWidget->GetActiveSlotWidget().IsValid())
 	{
 		return;
 	}
 
-	const UDMMaterialSlot* const Slot = SlotWidget->GetSlot();
+	const UDMMaterialSlot* const Slot = EditorWidget->GetActiveSlotWidget()->GetSlot();
 
 	if (!Slot)
 	{
@@ -462,7 +476,7 @@ void FDMMaterialSlotLayerMenus::AddAddLayerSection(UToolMenu* InMenu)
 		LOCTEXT("AddTextureSampleTooltip", "Add a Material Stage based on a Texture."),
 		FSlateIcon(),
 		FUIAction(FExecuteAction::CreateSP(
-			SlotWidget.Get(),
+			EditorWidget->GetActiveSlotWidget().Get(),
 			&SDMSlot::AddNewLayer_Expression,
 			TSubclassOf<UDMMaterialStageExpression>(UDMMaterialStageExpressionTextureSample::StaticClass()),
 			EDMMaterialLayerStage::All
@@ -475,7 +489,7 @@ void FDMMaterialSlotLayerMenus::AddAddLayerSection(UToolMenu* InMenu)
 		LOCTEXT("AddTextureSampleBaseOnlyTooltip", "Add a Material Stage based on a Texture with the Alpha disabled."),
 		FSlateIcon(),
 		FUIAction(FExecuteAction::CreateSP(
-			SlotWidget.Get(),
+			EditorWidget->GetActiveSlotWidget().Get(),
 			&SDMSlot::AddNewLayer_Expression,
 			TSubclassOf<UDMMaterialStageExpression>(UDMMaterialStageExpressionTextureSample::StaticClass()),
 			EDMMaterialLayerStage::Base
@@ -490,7 +504,7 @@ void FDMMaterialSlotLayerMenus::AddAddLayerSection(UToolMenu* InMenu)
 			LOCTEXT("AddAlphaOnlyTooltip", "Add an Alpha-Only Material Layer.\n\nThe base layer will be disabled by default. It can still be re-enabled later."),
 			FSlateIcon(),
 			FUIAction(FExecuteAction::CreateSP(
-				SlotWidget.Get(),
+				EditorWidget->GetActiveSlotWidget().Get(),
 				&SDMSlot::AddNewLayer_Expression,
 				TSubclassOf<UDMMaterialStageExpression>(UDMMaterialStageExpressionTextureSample::StaticClass()),
 				EDMMaterialLayerStage::Mask
@@ -503,7 +517,7 @@ void FDMMaterialSlotLayerMenus::AddAddLayerSection(UToolMenu* InMenu)
 		LOCTEXT("AddColor", "Solid Color"),
 		LOCTEXT("AddColorTooltip", "Add a new Material Layer with a solid RGB color."),
 		FSlateIcon(),
-		FUIAction(FExecuteAction::CreateSP(SlotWidget.Get(), &SDMSlot::AddNewLayer_NewLocalValue, EDMValueType::VT_Float3_RGB))
+		FUIAction(FExecuteAction::CreateSP(EditorWidget->GetActiveSlotWidget().Get(), &SDMSlot::AddNewLayer_NewLocalValue, EDMValueType::VT_Float3_RGB))
 	);
 
 	NewSection.AddMenuEntry(
@@ -511,7 +525,7 @@ void FDMMaterialSlotLayerMenus::AddAddLayerSection(UToolMenu* InMenu)
 		LOCTEXT("AddColorAtlas", "Color Atlas"),
 		LOCTEXT("AddColorAtlasTooltip", "Add a new Material Layer with a Color Atlas."),
 		FSlateIcon(),
-		FUIAction(FExecuteAction::CreateSP(SlotWidget.Get(), &SDMSlot::AddNewLayer_NewLocalValue, EDMValueType::VT_ColorAtlas))
+		FUIAction(FExecuteAction::CreateSP(EditorWidget->GetActiveSlotWidget().Get(), &SDMSlot::AddNewLayer_NewLocalValue, EDMValueType::VT_ColorAtlas))
 	);
 
 	NewSection.AddMenuEntry(
@@ -520,7 +534,7 @@ void FDMMaterialSlotLayerMenus::AddAddLayerSection(UToolMenu* InMenu)
 		LOCTEXT("AddEdgeColorTooltip", "Add a new Material Layer with a solid color based on the edge color on a texture."),
 		FSlateIcon(),
 		FUIAction(FExecuteAction::CreateSP(
-			SlotWidget.Get(),
+			EditorWidget->GetActiveSlotWidget().Get(),
 			&SDMSlot::AddNewLayer_Expression,
 			TSubclassOf<UDMMaterialStageExpression>(UDMMaterialStageExpressionTextureSampleEdgeColor::StaticClass()),
 			EDMMaterialLayerStage::All
@@ -533,7 +547,7 @@ void FDMMaterialSlotLayerMenus::AddAddLayerSection(UToolMenu* InMenu)
 		LOCTEXT("AddSceneTextureTooltip", "Add a new Material Layer that represents the Scene Texture for a post process material."),
 		FSlateIcon(),
 		FUIAction(FExecuteAction::CreateSP(
-			SlotWidget.Get(),
+			EditorWidget->GetActiveSlotWidget().Get(),
 			&SDMSlot::AddNewLayer_SceneTexture
 		))
 	);
@@ -544,7 +558,7 @@ void FDMMaterialSlotLayerMenus::AddAddLayerSection(UToolMenu* InMenu)
 		LOCTEXT("AddTextTooltip", "Add a Material Stage based on a Text Renderer."),
 		FSlateIcon(),
 		FUIAction(FExecuteAction::CreateSP(
-			SlotWidget.Get(),
+			EditorWidget->GetActiveSlotWidget().Get(),
 			&SDMSlot::AddNewLayer_Renderer,
 			TSubclassOf<UDMRenderTargetRenderer>(UDMRenderTargetTextRenderer::StaticClass())
 		))
@@ -580,7 +594,7 @@ void FDMMaterialSlotLayerMenus::AddAddLayerSection(UToolMenu* InMenu)
 		LOCTEXT("AddMaterialFunction", "Material Function"),
 		LOCTEXT("AddMaterialFunctionTooltip", "Add a new Material Layer based on a Material Function."),
 		FSlateIcon(),
-		FUIAction(FExecuteAction::CreateSP(SlotWidget.Get(), &SDMSlot::AddNewLayer_MaterialFunction))
+		FUIAction(FExecuteAction::CreateSP(EditorWidget->GetActiveSlotWidget().Get(), &SDMSlot::AddNewLayer_MaterialFunction))
 	);
 
 	if constexpr (UE::DynamicMaterialEditor::bGlobalValuesEnabled)
diff --git a/Engine/Plugins/Experimental/DynamicMaterial/Source/DynamicMaterialEditor/Private/Menus/DMMaterialSlotMenus.cpp b/Engine/Plugins/Experimental/DynamicMaterial/Source/DynamicMaterialEditor/Private/Menus/DMMaterialSlotMenus.cpp
--- a/Engine/Plugins/Experimental/DynamicMaterial/Source/DynamicMaterialEditor/Private/Menus/DMMaterialSlotMenus.cpp
+++ b/Engine/Plugins/Experimental/DynamicMaterial/Source/DynamicMaterialEditor/Private/Menus/DMMaterialSlotMenus.cpp
@@ -28,6 +28,13 @@ namespace UE::DynamicMaterialEditor::Private
 			return;
 		}
 
+		TSharedPtr<SDMEditor> EditorWidget = MenuContext->GetEditorWidget();
+
+		if (!EditorWidget.IsValid() || !EditorWidget->GetActiveSlotWidget().IsValid())
+		{
+			return;
+		}
+
 		FToolMenuSection& NewSection = InMenu->AddSection(SlotSectionName, LOCTEXT("MaterialSlotMenu", "Material Slot"));
 		{
 			NewSection.AddMenuEntry(
@@ -37,7 +44,7 @@ namespace UE::DynamicMaterialEditor::Private
 					"Warning: Removing a stage off may result in inputs being reset where incompatibilities are found."),
 				FSlateIcon(),
 				FUIAction(FExecuteAction::CreateSP(
-					MenuContext->GetSlotWidget().Pin().Get(),
+					EditorWidget->GetActiveSlotWidget().Get(),
 					&SDMSlot::RemoveSlot
 				))
 			);
@@ -61,7 +68,7 @@ TSharedRef<SWidget> FDMMaterialSlotMenus::MakeAddLayerButtonMenu(const TSharedPt
 		NewToolMenu->AddDynamicSection(NAME_None, FNewToolMenuDelegate::CreateStatic(&FDMMaterialSlotLayerMenus::AddAddLayerSection));
 	}
 
-	FToolMenuContext MenuContext(UDMMenuContext::CreateSlot(InSlotWidget));
+	FToolMenuContext MenuContext(UDMMenuContext::CreateEditor(InSlotWidget->GetEditorWidget()));
 
 	return UToolMenus::Get()->GenerateWidget(SlotAddLayerMenuName, MenuContext);
 }
diff --git a/Engine/Plugins/Experimental/DynamicMaterial/Source/DynamicMaterialEditor/Private/Menus/DMMaterialStageMenus.cpp b/Engine/Plugins/Experimental/DynamicMaterial/Source/DynamicMaterialEditor/Private/Menus/DMMaterialStageMenus.cpp
--- a/Engine/Plugins/Experimental/DynamicMaterial/Source/DynamicMaterialEditor/Private/Menus/DMMaterialStageMenus.cpp
+++ b/Engine/Plugins/Experimental/DynamicMaterial/Source/DynamicMaterialEditor/Private/Menus/DMMaterialStageMenus.cpp
@@ -24,7 +24,7 @@ FName FDMMaterialStageMenus::GetStageToggleSectionName()
 
 UToolMenu* FDMMaterialStageMenus::GenerateStageMenu(const TSharedPtr<SDMSlot>& InSlotWidget, const TSharedPtr<SDMStage>& InStageWidget)
 {
-	UToolMenu* NewToolMenu = UDMMenuContext::GenerateContextMenuStage(GetStageMenuName(), InSlotWidget, InStageWidget);
+	UToolMenu* NewToolMenu = UDMMenuContext::GenerateContextMenuStage(GetStageMenuName(), InSlotWidget->GetEditorWidget(), InStageWidget);
 	if (!NewToolMenu)
 	{
 		return nullptr;
@@ -139,21 +139,24 @@ void FDMMaterialStageMenus::AddStageSection(UToolMenu* InMenu)
 
 		if (bAllowRemoveLayer)
 		{
-			if (TSharedPtr<SDMSlot> SlotWidget = MenuContext->GetSlotWidget().Pin())
+			if (TSharedPtr<SDMEditor> EditorWidget = MenuContext->GetEditorWidget())
 			{
-				//FDMSlotLayerMenus::AddLayerSection
-				NewSection.AddMenuEntry(NAME_None,
-					LOCTEXT("RemoveLayer", "Remove"),
-					LOCTEXT("RemoveLayerTooltip", "Remove this layer from its Material Slot.\n\n"
-						"Warning: Removing a stage off may result in inputs being reset where incompatibilities are found."),
-					FSlateIcon(),
-					FUIAction(FExecuteAction::CreateSP(
-						SlotWidget.Get(),
-						&SDMSlot::RemoveLayerByStage,
-						Stage, 
-						false
-					))
-				);
+				if (EditorWidget->GetActiveSlotWidget())
+				{
+					//FDMSlotLayerMenus::AddLayerSection
+					NewSection.AddMenuEntry(NAME_None,
+						LOCTEXT("RemoveLayer", "Remove"),
+						LOCTEXT("RemoveLayerTooltip", "Remove this layer from its Material Slot.\n\n"
+							"Warning: Removing a stage off may result in inputs being reset where incompatibilities are found."),
+						FSlateIcon(),
+						FUIAction(FExecuteAction::CreateSP(
+							EditorWidget->GetActiveSlotWidget().Get(),
+							&SDMSlot::RemoveLayerByStage,
+							Stage,
+							false
+						))
+					);
+				}
 			}
 		}
 	}
diff --git a/Engine/Plugins/Experimental/DynamicMaterial/Source/DynamicMaterialEditor/Private/Menus/DMMaterialStageSourceMenus.cpp b/Engine/Plugins/Experimental/DynamicMaterial/Source/DynamicMaterialEditor/Private/Menus/DMMaterialStageSourceMenus.cpp
--- a/Engine/Plugins/Experimental/DynamicMaterial/Source/DynamicMaterialEditor/Private/Menus/DMMaterialStageSourceMenus.cpp
+++ b/Engine/Plugins/Experimental/DynamicMaterial/Source/DynamicMaterialEditor/Private/Menus/DMMaterialStageSourceMenus.cpp
@@ -132,9 +132,9 @@ namespace UE::DynamicMaterialEditor::Private
 								ensure(false);
 							}
 
-							if (TSharedPtr<SDMSlot> SlotWidget = MenuContext->GetSlotWidget().Pin())
+							if (TSharedPtr<SDMEditor> EditorWidget = MenuContext->GetEditorWidget())
 							{
-								SlotWidget->InvalidateComponentEditWidget();
+								EditorWidget->InvalidateComponentEditWidget();
 							}
 						})
 				)
@@ -234,9 +234,9 @@ namespace UE::DynamicMaterialEditor::Private
 							ensure(false);
 						}
 
-						if (TSharedPtr<SDMSlot> SlotWidget = MenuContext->GetSlotWidget().Pin())
+						if (TSharedPtr<SDMEditor> EditorWidget = MenuContext->GetEditorWidget())
 						{
-							SlotWidget->InvalidateComponentEditWidget();
+							EditorWidget->InvalidateComponentEditWidget();
 						}
 					})
 				)
@@ -330,9 +330,9 @@ namespace UE::DynamicMaterialEditor::Private
 								ensure(false);
 							}
 
-							if (TSharedPtr<SDMSlot> SlotWidget = MenuContext->GetSlotWidget().Pin())
+							if (TSharedPtr<SDMEditor> EditorWidget = MenuContext->GetEditorWidget())
 							{
-								SlotWidget->InvalidateComponentEditWidget();
+								EditorWidget->InvalidateComponentEditWidget();
 							}
 						})
 				)
@@ -455,9 +455,9 @@ namespace UE::DynamicMaterialEditor::Private
 								ensure(false);
 							}
 
-							if (TSharedPtr<SDMSlot> SlotWidget = MenuContext->GetSlotWidget().Pin())
+							if (TSharedPtr<SDMEditor> EditorWidget = MenuContext->GetEditorWidget())
 							{
-								SlotWidget->InvalidateComponentEditWidget();
+								EditorWidget->InvalidateComponentEditWidget();
 							}
 						})
 					)
@@ -601,9 +601,9 @@ namespace UE::DynamicMaterialEditor::Private
 									ensure(false);
 								}
 
-								if (TSharedPtr<SDMSlot> SlotWidget = MenuContext->GetSlotWidget().Pin())
+								if (TSharedPtr<SDMEditor> EditorWidget = MenuContext->GetEditorWidget())
 								{
-									SlotWidget->InvalidateComponentEditWidget();
+									EditorWidget->InvalidateComponentEditWidget();
 								}
 							})
 						)
@@ -704,9 +704,9 @@ namespace UE::DynamicMaterialEditor::Private
 								ensure(false);
 							}
 
-							if (TSharedPtr<SDMSlot> SlotWidget = MenuContext->GetSlotWidget().Pin())
+							if (TSharedPtr<SDMEditor> EditorWidget = MenuContext->GetEditorWidget())
 							{
-								SlotWidget->InvalidateComponentEditWidget();
+								EditorWidget->InvalidateComponentEditWidget();
 							}
 						})
 					)
@@ -769,9 +769,9 @@ namespace UE::DynamicMaterialEditor::Private
 			ensure(false);
 		}
 
-		if (TSharedPtr<SDMSlot> SlotWidget = InMenuContext->GetSlotWidget().Pin())
+		if (TSharedPtr<SDMEditor> EditorWidget = InMenuContext->GetEditorWidget())
 		{
-			SlotWidget->InvalidateComponentEditWidget();
+			EditorWidget->InvalidateComponentEditWidget();
 		}
 	}
 
@@ -827,9 +827,9 @@ namespace UE::DynamicMaterialEditor::Private
 			ensure(false);
 		}
 
-		if (TSharedPtr<SDMSlot> SlotWidget = InMenuContext->GetSlotWidget().Pin())
+		if (TSharedPtr<SDMEditor> EditorWidget = InMenuContext->GetEditorWidget())
 		{
-			SlotWidget->InvalidateComponentEditWidget();
+			EditorWidget->InvalidateComponentEditWidget();
 		}
 	}
 
@@ -885,9 +885,9 @@ namespace UE::DynamicMaterialEditor::Private
 			ensure(false);
 		}
 
-		if (TSharedPtr<SDMSlot> SlotWidget = InMenuContext->GetSlotWidget().Pin())
+		if (TSharedPtr<SDMEditor> EditorWidget = InMenuContext->GetEditorWidget())
 		{
-			SlotWidget->InvalidateComponentEditWidget();
+			EditorWidget->InvalidateComponentEditWidget();
 		}
 	}
 
@@ -957,9 +957,9 @@ namespace UE::DynamicMaterialEditor::Private
 			MaskStage->SetEnabled(false);
 		}
 
-		if (TSharedPtr<SDMSlot> SlotWidget = InMenuContext->GetSlotWidget().Pin())
+		if (TSharedPtr<SDMEditor> EditorWidget = InMenuContext->GetEditorWidget())
 		{
-			SlotWidget->InvalidateComponentEditWidget();
+			EditorWidget->InvalidateComponentEditWidget();
 		}
 	}
 
@@ -1041,9 +1041,9 @@ namespace UE::DynamicMaterialEditor::Private
 			ensure(false);
 		}
 
-		if (TSharedPtr<SDMSlot> SlotWidget = InMenuContext->GetSlotWidget().Pin())
+		if (TSharedPtr<SDMEditor> EditorWidget = InMenuContext->GetEditorWidget())
 		{
-			SlotWidget->InvalidateComponentEditWidget();
+			EditorWidget->InvalidateComponentEditWidget();
 		}
 	}
 
@@ -1126,9 +1126,9 @@ namespace UE::DynamicMaterialEditor::Private
 			ensure(false);
 		}
 
-		if (TSharedPtr<SDMSlot> SlotWidget = InMenuContext->GetSlotWidget().Pin())
+		if (TSharedPtr<SDMEditor> EditorWidget = InMenuContext->GetEditorWidget())
 		{
-			SlotWidget->InvalidateComponentEditWidget();
+			EditorWidget->InvalidateComponentEditWidget();
 		}
 	}
 
@@ -1179,9 +1179,9 @@ namespace UE::DynamicMaterialEditor::Private
 			ensure(false);
 		}
 
-		if (TSharedPtr<SDMSlot> SlotWidget = InMenuContext->GetSlotWidget().Pin())
+		if (TSharedPtr<SDMEditor> EditorWidget = InMenuContext->GetEditorWidget())
 		{
-			SlotWidget->InvalidateComponentEditWidget();
+			EditorWidget->InvalidateComponentEditWidget();
 		}
 	}
 
@@ -1394,7 +1394,7 @@ TSharedRef<SWidget> FDMMaterialStageSourceMenus::MakeChangeSourceMenu(const TSha
 		NewSection.AddDynamicEntry(NAME_None, FNewToolMenuSectionDelegate::CreateStatic(&CreateChangeMaterialStageSource));
 	}
 
-	FToolMenuContext MenuContext(UDMMenuContext::CreateStage(InSlotWidget, InStageWidget));
+	FToolMenuContext MenuContext(UDMMenuContext::CreateStage(InSlotWidget->GetEditorWidget(), InStageWidget));
 
 	return UToolMenus::Get()->GenerateWidget(ChangeStageSourceMenuName, MenuContext);
 }
diff --git a/Engine/Plugins/Experimental/DynamicMaterial/Source/DynamicMaterialEditor/Private/Menus/DMMenuContext.cpp b/Engine/Plugins/Experimental/DynamicMaterial/Source/DynamicMaterialEditor/Private/Menus/DMMenuContext.cpp
--- a/Engine/Plugins/Experimental/DynamicMaterial/Source/DynamicMaterialEditor/Private/Menus/DMMenuContext.cpp
+++ b/Engine/Plugins/Experimental/DynamicMaterial/Source/DynamicMaterialEditor/Private/Menus/DMMenuContext.cpp
@@ -7,22 +7,22 @@
 #include "Components/DMMaterialStageBlend.h"
 #include "Components/DMMaterialStageSource.h"
 #include "Model/DynamicMaterialModel.h"
-#include "Model/DynamicMaterialModelEditorOnlyData.h"
-#include "Slate/SDMSlot.h"
+#include "Slate/SDMEditor.h"
 #include "Slate/SDMStage.h"
 #include "ToolMenu.h"
 #include "ToolMenus.h"
+#include "Slate/SDMSlot.h"
 
-UDMMenuContext* UDMMenuContext::Create(const TWeakPtr<SDMSlot>& InSlotWidget, const TWeakPtr<SDMStage>& InStageWidget, UDMMaterialLayerObject* InLayerObject)
+UDMMenuContext* UDMMenuContext::Create(const TWeakPtr<SDMEditor>& InEditorWidget, const TWeakPtr<SDMStage>& InStageWidget, UDMMaterialLayerObject* InLayerObject)
 {
 	UDMMenuContext* Context = NewObject<UDMMenuContext>();
-	Context->SlotWidgetWeak = InSlotWidget;
+	Context->EditorWidgetWeak = InEditorWidget;
 	Context->StageWidgetWeak = InStageWidget;
 	Context->LayerObjectWeak = InLayerObject;
 	return Context;
 }
 
-UToolMenu* UDMMenuContext::GenerateContextMenu(const FName InMenuName, const TWeakPtr<SDMSlot>& InSlotWidget, const TWeakPtr<SDMStage>& InStageWidget, UDMMaterialLayerObject* InLayerObject)
+UToolMenu* UDMMenuContext::GenerateContextMenu(const FName InMenuName, const TWeakPtr<SDMEditor>& InEditorWidget, const TWeakPtr<SDMStage>& InStageWidget, UDMMaterialLayerObject* InLayerObject)
 {
 	UToolMenu* NewMenu = UToolMenus::Get()->RegisterMenu(InMenuName, NAME_None, EMultiBoxType::Menu, false);
 	if (!NewMenu)
@@ -47,15 +47,12 @@ UToolMenu* UDMMenuContext::GenerateContextMenu(const FName InMenuName, const TWe
 
 	TSharedPtr<FUICommandList> CommandList = nullptr;
 
-	if (TSharedPtr<SDMSlot> Slot = InSlotWidget.Pin())
+	if (TSharedPtr<SDMEditor> Editor = InEditorWidget.Pin())
 	{
-		if (TSharedPtr<SDMEditor> Editor = Slot->GetEditorWidget())
-		{
-			CommandList = Editor->GetCommandList();
-		}
+		CommandList = Editor->GetCommandList();
 	}
 
-	NewMenu->Context = FToolMenuContext(CommandList, TSharedPtr<FExtender>(), Create(InSlotWidget, InStageWidget, InLayerObject));
+	NewMenu->Context = FToolMenuContext(CommandList, TSharedPtr<FExtender>(), Create(InEditorWidget, InStageWidget, InLayerObject));
 
 	return NewMenu;
 }
@@ -65,19 +62,19 @@ UDMMenuContext* UDMMenuContext::CreateEmpty()
 	return Create(nullptr, nullptr, nullptr);
 }
 
-UDMMenuContext* UDMMenuContext::CreateSlot(const TWeakPtr<SDMSlot>& InSlotWidget)
+UDMMenuContext* UDMMenuContext::CreateEditor(const TWeakPtr<SDMEditor>& InEditorWidget)
 {
-	return Create(InSlotWidget, nullptr, nullptr);
+	return Create(InEditorWidget, nullptr, nullptr);
 }
 
-UDMMenuContext* UDMMenuContext::CreateLayer(const TWeakPtr<SDMSlot>& InSlotWidget, UDMMaterialLayerObject* InLayerObject)
+UDMMenuContext* UDMMenuContext::CreateLayer(const TWeakPtr<SDMEditor>& InEditorWidget, UDMMaterialLayerObject* InLayerObject)
 {
-	return Create(InSlotWidget, nullptr, InLayerObject);
+	return Create(InEditorWidget, nullptr, InLayerObject);
 }
 
-UDMMenuContext* UDMMenuContext::CreateStage(const TWeakPtr<SDMSlot>& InSlotWidget, const TWeakPtr<SDMStage>& InStageWidget)
+UDMMenuContext* UDMMenuContext::CreateStage(const TWeakPtr<SDMEditor>& InEditorWidget, const TWeakPtr<SDMStage>& InStageWidget)
 {
-	return Create(InSlotWidget, InStageWidget, nullptr);
+	return Create(InEditorWidget, InStageWidget, nullptr);
 }
 
 UToolMenu* UDMMenuContext::GenerateContextMenuDefault(const FName InMenuName)
@@ -85,39 +82,39 @@ UToolMenu* UDMMenuContext::GenerateContextMenuDefault(const FName InMenuName)
 	return GenerateContextMenu(InMenuName, nullptr, nullptr, nullptr);
 }
 
-UToolMenu* UDMMenuContext::GenerateContextMenuSlot(const FName InMenuName, const TWeakPtr<SDMSlot>& InSlotWidget)
+UToolMenu* UDMMenuContext::GenerateContextMenuEditor(const FName InMenuName, const TWeakPtr<SDMEditor>& InEditorWidget)
 {
-	return GenerateContextMenu(InMenuName, InSlotWidget, nullptr, nullptr);
+	return GenerateContextMenu(InMenuName, InEditorWidget, nullptr, nullptr);
 }
 
-UToolMenu* UDMMenuContext::GenerateContextMenuLayer(const FName InMenuName, const TWeakPtr<SDMSlot>& InSlotWidget, UDMMaterialLayerObject* InLayerObject)
+UToolMenu* UDMMenuContext::GenerateContextMenuLayer(const FName InMenuName, const TWeakPtr<SDMEditor>& InEditorWidget, UDMMaterialLayerObject* InLayerObject)
 {
-	return GenerateContextMenu(InMenuName, InSlotWidget, nullptr, InLayerObject);
+	return GenerateContextMenu(InMenuName, InEditorWidget, nullptr, InLayerObject);
 }
 
-UToolMenu* UDMMenuContext::GenerateContextMenuStage(const FName InMenuName, const TWeakPtr<SDMSlot>& InSlotWidget, const TWeakPtr<SDMStage>& InStageWidget)
+UToolMenu* UDMMenuContext::GenerateContextMenuStage(const FName InMenuName, const TWeakPtr<SDMEditor>& InEditorWidget, const TWeakPtr<SDMStage>& InStageWidget)
 {
-	return GenerateContextMenu(InMenuName, InSlotWidget, InStageWidget, nullptr);
+	return GenerateContextMenu(InMenuName, InEditorWidget, InStageWidget, nullptr);
 }
 
 UDMMaterialSlot* UDMMenuContext::GetSlot() const
 {
-	if (TSharedPtr<SDMSlot> SlotWidget = SlotWidgetWeak.Pin())
+	if (TSharedPtr<SDMEditor> EditorWidget = EditorWidgetWeak.Pin())
 	{
-		return SlotWidget->GetSlot();
+		if (TSharedPtr<SDMSlot> SlotWidget = EditorWidget->GetActiveSlotWidget())
+		{
+			return SlotWidget->GetSlot();
+		}
 	}
 
-	return nullptr;
+	return nullptr;;
 }
 
 UDynamicMaterialModel* UDMMenuContext::GetModel() const
 {
-	if (const UDMMaterialSlot* const Slot = GetSlot())
+	if (TSharedPtr<SDMEditor> EditorWidget = EditorWidgetWeak.Pin())
 	{
-		if (UDynamicMaterialModelEditorOnlyData* ModelEditorOnlyData = Slot->GetMaterialModelEditorOnlyData())
-		{
-			return ModelEditorOnlyData->GetMaterialModel();
-		}
+		return EditorWidget->GetMaterialModel();
 	}
 
 	return nullptr;
diff --git a/Engine/Plugins/Experimental/DynamicMaterial/Source/DynamicMaterialEditor/Private/Menus/DMMenuContext.h b/Engine/Plugins/Experimental/DynamicMaterial/Source/DynamicMaterialEditor/Private/Menus/DMMenuContext.h
--- a/Engine/Plugins/Experimental/DynamicMaterial/Source/DynamicMaterialEditor/Private/Menus/DMMenuContext.h
+++ b/Engine/Plugins/Experimental/DynamicMaterial/Source/DynamicMaterialEditor/Private/Menus/DMMenuContext.h
@@ -5,7 +5,7 @@
 #include "Templates/SharedPointer.h"
 #include "DMMenuContext.generated.h"
 
-class SDMSlot;
+class SDMEditor;
 class SDMStage;
 class UDMMaterialLayerObject;
 class UDMMaterialSlot;
@@ -22,17 +22,17 @@ class UDMMenuContext : public UObject
 
 public:
 	static UDMMenuContext* CreateEmpty();
-	static UDMMenuContext* CreateSlot(const TWeakPtr<SDMSlot>& InSlotWidget);
-	static UDMMenuContext* CreateLayer(const TWeakPtr<SDMSlot>& InSlotWidget, UDMMaterialLayerObject* InLayerObject);
-	static UDMMenuContext* CreateStage(const TWeakPtr<SDMSlot>& InSlotWidget, const TWeakPtr<SDMStage>& InStageWidget);
+	static UDMMenuContext* CreateEditor(const TWeakPtr<SDMEditor>& InEditorWidget);
+	static UDMMenuContext* CreateLayer(const TWeakPtr<SDMEditor>& InEditorWidget, UDMMaterialLayerObject* InLayerObject);
+	static UDMMenuContext* CreateStage(const TWeakPtr<SDMEditor>& InEditorWidget, const TWeakPtr<SDMStage>& InStageWidget);
 
 	static UToolMenu* GenerateContextMenuDefault(const FName InMenuName);
-	static UToolMenu* GenerateContextMenuSlot(const FName InMenuName, const TWeakPtr<SDMSlot>& InSlotWidget);
-	static UToolMenu* GenerateContextMenuLayer(const FName InMenuName, const TWeakPtr<SDMSlot>& InSlotWidget, UDMMaterialLayerObject* InLayerObject);
-	static UToolMenu* GenerateContextMenuStage(const FName InMenuName, const TWeakPtr<SDMSlot>& InSlotWidget, const TWeakPtr<SDMStage>& InStageWidget);
+	static UToolMenu* GenerateContextMenuEditor(const FName InMenuName, const TWeakPtr<SDMEditor>& InEditorWidget);
+	static UToolMenu* GenerateContextMenuLayer(const FName InMenuName, const TWeakPtr<SDMEditor>& InEditorWidget, UDMMaterialLayerObject* InLayerObject);
+	static UToolMenu* GenerateContextMenuStage(const FName InMenuName, const TWeakPtr<SDMEditor>& InEditorWidget, const TWeakPtr<SDMStage>& InStageWidget);
 
-	const TWeakPtr<SDMSlot>& GetSlotWidget() const { return SlotWidgetWeak; }
-	const TWeakPtr<SDMStage>& GetStageWidget() const { return StageWidgetWeak; }
+	const TSharedPtr<SDMEditor> GetEditorWidget() const { return EditorWidgetWeak.Pin(); }
+	const TSharedPtr<SDMStage> GetStageWidget() const { return StageWidgetWeak.Pin(); }
 
 	UDMMaterialSlot* GetSlot() const;
 
@@ -47,10 +47,10 @@ public:
 	const UDMMaterialLayerObject* GetLayer() const;
 
 protected:
-	TWeakPtr<SDMSlot> SlotWidgetWeak;
+	TWeakPtr<SDMEditor> EditorWidgetWeak;
 	TWeakPtr<SDMStage> StageWidgetWeak;
 	TWeakObjectPtr<UDMMaterialLayerObject> LayerObjectWeak;
 
-	static UDMMenuContext* Create(const TWeakPtr<SDMSlot>& InSlotWidget, const TWeakPtr<SDMStage>& InStageWidget, UDMMaterialLayerObject* InLayerObject);
-	static UToolMenu* GenerateContextMenu(const FName InMenuName, const TWeakPtr<SDMSlot>& InSlotWidget, const TWeakPtr<SDMStage>& InStageWidget, UDMMaterialLayerObject* InLayerObject);
+	static UDMMenuContext* Create(const TWeakPtr<SDMEditor>& InEditorWidget, const TWeakPtr<SDMStage>& InStageWidget, UDMMaterialLayerObject* InLayerObject);
+	static UToolMenu* GenerateContextMenu(const FName InMenuName, const TWeakPtr<SDMEditor>& InEditorWidget, const TWeakPtr<SDMStage>& InStageWidget, UDMMaterialLayerObject* InLayerObject);
 };
diff --git a/Engine/Plugins/Experimental/DynamicMaterial/Source/DynamicMaterialEditor/Private/Menus/DMToolBarMenus.cpp b/Engine/Plugins/Experimental/DynamicMaterial/Source/DynamicMaterialEditor/Private/Menus/DMToolBarMenus.cpp
--- a/Engine/Plugins/Experimental/DynamicMaterial/Source/DynamicMaterialEditor/Private/Menus/DMToolBarMenus.cpp
+++ b/Engine/Plugins/Experimental/DynamicMaterial/Source/DynamicMaterialEditor/Private/Menus/DMToolBarMenus.cpp
@@ -44,17 +44,9 @@ namespace UE::DynamicMaterialEditor::Private
 			{
 				if (UDynamicMaterialModelEditorOnlyData* EditorOnlyData = UDynamicMaterialModelEditorOnlyData::Get(MaterialModel))
 				{
-					TSharedPtr<SDMEditor> EditorWidget;
-
-					if (TSharedPtr<SDMSlot> SlotWidget = Context->GetSlotWidget().Pin())
-					{
-						EditorWidget = SlotWidget->GetEditorWidget();
-
-					}
-
 					EditorOnlyData->SetChannelListPreset(InPresetName);
 
-					if (EditorWidget.IsValid())
+					if (TSharedPtr<SDMEditor> EditorWidget = Context->GetEditorWidget())
 					{
 						EditorWidget->RefreshSlotPickerList();
 						EditorWidget->RefreshSlotWidget();
@@ -111,17 +103,9 @@ namespace UE::DynamicMaterialEditor::Private
 			{
 				if (UDynamicMaterialModelEditorOnlyData* EditorOnlyData = UDynamicMaterialModelEditorOnlyData::Get(MaterialModel))
 				{
-					TSharedPtr<SDMEditor> EditorWidget;
-
-					if (TSharedPtr<SDMSlot> SlotWidget = Context->GetSlotWidget().Pin())
-					{
-						EditorWidget = SlotWidget->GetEditorWidget();
-
-					}
-
 					EditorOnlyData->SetDomain(InDomain);
 
-					if (EditorWidget.IsValid())
+					if (TSharedPtr<SDMEditor> EditorWidget = Context->GetEditorWidget())
 					{
 						EditorWidget->RefreshSlotPickerList();
 						EditorWidget->RefreshSlotWidget();
@@ -194,17 +178,9 @@ namespace UE::DynamicMaterialEditor::Private
 			{
 				if (UDynamicMaterialModelEditorOnlyData* EditorOnlyData = UDynamicMaterialModelEditorOnlyData::Get(MaterialModel))
 				{
-					TSharedPtr<SDMEditor> EditorWidget;
-
-					if (TSharedPtr<SDMSlot> SlotWidget = Context->GetSlotWidget().Pin())
-					{
-						EditorWidget = SlotWidget->GetEditorWidget();
-
-					}
-
 					EditorOnlyData->SetBlendMode(InBlendMode);
 
-					if (EditorWidget.IsValid())
+					if (TSharedPtr<SDMEditor> EditorWidget = Context->GetEditorWidget())
 					{
 						EditorWidget->RefreshSlotPickerList();
 						EditorWidget->RefreshSlotWidget();
@@ -277,14 +253,6 @@ namespace UE::DynamicMaterialEditor::Private
 			{
 				if (UDynamicMaterialModelEditorOnlyData* EditorOnlyData = UDynamicMaterialModelEditorOnlyData::Get(MaterialModel))
 				{
-					TSharedPtr<SDMEditor> EditorWidget;
-
-					if (TSharedPtr<SDMSlot> SlotWidget = Context->GetSlotWidget().Pin())
-					{
-						EditorWidget = SlotWidget->GetEditorWidget();
-
-					}
-
 					switch (EditorOnlyData->GetShadingModel())
 					{
 						case EDMMaterialShadingModel::DefaultLit:
@@ -296,7 +264,7 @@ namespace UE::DynamicMaterialEditor::Private
 							break;
 					}
 
-					if (EditorWidget.IsValid())
+					if (TSharedPtr<SDMEditor> EditorWidget = Context->GetEditorWidget())
 					{
 						EditorWidget->RefreshSlotPickerList();
 					}
@@ -349,17 +317,7 @@ namespace UE::DynamicMaterialEditor::Private
 			{
 				if (UDynamicMaterialModelEditorOnlyData* EditorOnlyData = UDynamicMaterialModelEditorOnlyData::Get(MaterialModel))
 				{
-					TSharedPtr<SDMEditor> EditorWidget;
-
-					if (TSharedPtr<SDMSlot> SlotWidget = Context->GetSlotWidget().Pin())
-					{
-						EditorWidget = SlotWidget->GetEditorWidget();
-
-					}
-
-					EditorOnlyData->SetPixelAnimationFlag(!EditorOnlyData->IsPixelAnimationFlagSet());
-
-					if (EditorWidget.IsValid())
+					if (TSharedPtr<SDMEditor> EditorWidget = Context->GetEditorWidget())
 					{
 						EditorWidget->RefreshSlotPickerList();
 					}					
@@ -407,17 +365,7 @@ namespace UE::DynamicMaterialEditor::Private
 			{
 				if (UDynamicMaterialModelEditorOnlyData* EditorOnlyData = UDynamicMaterialModelEditorOnlyData::Get(MaterialModel))
 				{
-					TSharedPtr<SDMEditor> EditorWidget;
-
-					if (TSharedPtr<SDMSlot> SlotWidget = Context->GetSlotWidget().Pin())
-					{
-						EditorWidget = SlotWidget->GetEditorWidget();
-
-					}
-
-					EditorOnlyData->SetTwoSidedFlag(!EditorOnlyData->IsTwoSidedFlagSet());
-
-					if (EditorWidget.IsValid())
+					if (TSharedPtr<SDMEditor> EditorWidget = Context->GetEditorWidget())
 					{
 						EditorWidget->RefreshSlotPickerList();
 					}					
@@ -1007,7 +955,7 @@ TSharedRef<SWidget> FDMToolBarMenus::MakeEditorLayoutMenu(const TSharedPtr<SDMEd
 	FToolMenuContext MenuContext(
 		FDynamicMaterialEditorModule::Get().GetCommandList(),
 		TSharedPtr<FExtender>(),
-		UDMMenuContext::CreateSlot(InEditorWidget->GetActiveSlotWidget())
+		UDMMenuContext::CreateEditor(InEditorWidget)
 	);
 
 	return UToolMenus::Get()->GenerateWidget(ToolBarEditorLayoutMenuName, MenuContext);
diff --git a/Engine/Plugins/Experimental/DynamicMaterial/Source/DynamicMaterialEditor/Private/Model/DMMaterialBuildState.cpp b/Engine/Plugins/Experimental/DynamicMaterial/Source/DynamicMaterialEditor/Private/Model/DMMaterialBuildState.cpp
--- a/Engine/Plugins/Experimental/DynamicMaterial/Source/DynamicMaterialEditor/Private/Model/DMMaterialBuildState.cpp
+++ b/Engine/Plugins/Experimental/DynamicMaterial/Source/DynamicMaterialEditor/Private/Model/DMMaterialBuildState.cpp
@@ -28,6 +28,7 @@ FDMMaterialBuildState::FDMMaterialBuildState(UMaterial* InDynamicMaterial, UDyna
 	, bIgnoreUVs(false)
 	, bIsPreviewMaterial(false)
 	, Utils(MakeShared<FDMMaterialBuildUtils>(*this))
+	, GlobalTilingExpression(nullptr)
 {
 	check(InDynamicMaterial);
 	check(InMaterialModel);
diff --git a/Engine/Plugins/Experimental/DynamicMaterial/Source/DynamicMaterialEditor/Private/Model/DynamicMaterialModelEditorOnlyData.cpp b/Engine/Plugins/Experimental/DynamicMaterial/Source/DynamicMaterialEditor/Private/Model/DynamicMaterialModelEditorOnlyData.cpp
--- a/Engine/Plugins/Experimental/DynamicMaterial/Source/DynamicMaterialEditor/Private/Model/DynamicMaterialModelEditorOnlyData.cpp
+++ b/Engine/Plugins/Experimental/DynamicMaterial/Source/DynamicMaterialEditor/Private/Model/DynamicMaterialModelEditorOnlyData.cpp
@@ -26,6 +26,7 @@
 #include "Components/MaterialProperties/DMMPWorldPositionOffset.h"
 #include "Components/MaterialStageBlends/DMMSBNormal.h"
 #include "Components/MaterialValues/DMMaterialValueFloat1.h"
+#include "Components/MaterialValues/DMMaterialValueFloat2.h"
 #include "CoreGlobals.h"
 #include "DMComponentPath.h"
 #include "DMDefs.h"
@@ -296,6 +297,15 @@ void UDynamicMaterialModelEditorOnlyData::BuildMaterial(bool bInDirtyAssets)
 
 	const bool bIsPostProcessMaterial = Domain == EMaterialDomain::MD_PostProcess;
 
+	/**
+	 * Add global tiling parameter
+	 */
+	if (MaterialModel->GlobalTilingValue)
+	{
+		MaterialModel->GlobalTilingValue->GenerateExpression(BuildState);
+		BuildState->SetGlobalTilingExpression(BuildState->GetLastValueExpression(MaterialModel->GlobalTilingValue));
+	}
+
 	/**
 	 * Process slots to build base material inputs.
 	 */
diff --git a/Engine/Plugins/Experimental/DynamicMaterial/Source/DynamicMaterialEditor/Private/Slate/SDMComponentEdit.cpp b/Engine/Plugins/Experimental/DynamicMaterial/Source/DynamicMaterialEditor/Private/Slate/SDMComponentEdit.cpp
--- a/Engine/Plugins/Experimental/DynamicMaterial/Source/DynamicMaterialEditor/Private/Slate/SDMComponentEdit.cpp
+++ b/Engine/Plugins/Experimental/DynamicMaterial/Source/DynamicMaterialEditor/Private/Slate/SDMComponentEdit.cpp
@@ -802,7 +802,7 @@ void SDMComponentEdit::GenerateMaterialModelPropertyRows(const TSharedRef<SDMEdi
 		}
 	}
 
-	if (UDMMaterialValueFloat2* GlobalTilingValue = InMaterialModel->GetGlobalScaleValue())
+	if (UDMMaterialValueFloat2* GlobalTilingValue = InMaterialModel->GetGlobalTilingValue())
 	{
 		if (!InOutProcessedObjects.Contains(GlobalTilingValue))
 		{
diff --git a/Engine/Plugins/Experimental/DynamicMaterial/Source/DynamicMaterialEditor/Public/Model/DMMaterialBuildState.h b/Engine/Plugins/Experimental/DynamicMaterial/Source/DynamicMaterialEditor/Public/Model/DMMaterialBuildState.h
--- a/Engine/Plugins/Experimental/DynamicMaterial/Source/DynamicMaterialEditor/Public/Model/DMMaterialBuildState.h
+++ b/Engine/Plugins/Experimental/DynamicMaterial/Source/DynamicMaterialEditor/Public/Model/DMMaterialBuildState.h
@@ -54,6 +54,10 @@ struct DYNAMICMATERIALEDITOR_API FDMMaterialBuildState : public IDMMaterialBuild
 
 	bool IsPreviewMaterial() const { return bIsPreviewMaterial; }
 
+	UMaterialExpression* GetGlobalTilingExpression() const { return GlobalTilingExpression; }
+
+	void SetGlobalTilingExpression(UMaterialExpression* InExpression) { GlobalTilingExpression = InExpression; }
+
 	virtual IDMMaterialBuildUtilsInterface& GetBuildUtils() const override;
 
 	FExpressionInput* GetMaterialProperty(EDMMaterialPropertyType InProperty) const;
@@ -166,4 +170,5 @@ private:
 	TMap<const UDMMaterialStageSource*, TArray<UMaterialExpression*>> StageSources;
 	TMap<const UDMMaterialStageSource*, TArray<UDMMaterialStageSource*>> Callbacks;
 	TSet<UMaterialExpression*> OtherExpressions;
+	UMaterialExpression* GlobalTilingExpression;
 };
