Material Designer
- Swapped position of toggle/popout/mode toggles buttons
- Made mode toggle buttons checkboxes instead of a single toggle button.

diff --git a/Engine/Plugins/Experimental/DynamicMaterial/Source/DynamicMaterialEditor/Private/Slate/Properties/SDMTextureUVVisualizerPopout.cpp b/Engine/Plugins/Experimental/DynamicMaterial/Source/DynamicMaterialEditor/Private/Slate/Properties/SDMTextureUVVisualizerPopout.cpp
--- a/Engine/Plugins/Experimental/DynamicMaterial/Source/DynamicMaterialEditor/Private/Slate/Properties/SDMTextureUVVisualizerPopout.cpp
+++ b/Engine/Plugins/Experimental/DynamicMaterial/Source/DynamicMaterialEditor/Private/Slate/Properties/SDMTextureUVVisualizerPopout.cpp
@@ -21,6 +21,7 @@
 #include "Widgets/Colors/SColorBlock.h"
 #include "Widgets/Docking/SDockTab.h"
 #include "Widgets/Input/SButton.h"
+#include "Widgets/Input/SCheckBox.h"
 #include "Widgets/Layout/SBox.h"
 #include "Widgets/Text/STextBlock.h"
 
@@ -150,16 +151,6 @@ void SDMTextureUVVisualizerPopout::Construct(const FArguments& InArgs, UDMMateri
 	];
 }
 
-FReply SDMTextureUVVisualizerPopout::OnToggleModeClicked()
-{
-	if (Visualizer.IsValid())
-	{
-		Visualizer->TogglePivotEditMode();
-	}
-
-	return FReply::Handled();
-}
-
 FText SDMTextureUVVisualizerPopout::GetModeButtonText() const
 {
 	if (Visualizer.IsValid() && Visualizer->IsInPivotEditMode())
@@ -237,13 +228,40 @@ TSharedRef<SWidget> SDMTextureUVVisualizerPopout::CreatePropertyWidget(UDMTextur
 	FDMPropertyHandle EditModeButtonRow;
 	EditModeButtonRow.ValueName = TEXT("EditMode");
 	EditModeButtonRow.NameOverride = LOCTEXT("EditMode", "Edit Mode");
-	EditModeButtonRow.ValueWidget = SNew(SButton)
-		.OnClicked(this, &SDMTextureUVVisualizerPopout::OnToggleModeClicked)
-		.Content()
+	EditModeButtonRow.ValueWidget = SNew(SHorizontalBox)
+		+ SHorizontalBox::Slot()
+		.FillContentWidth(1.f)
 		[
-			SNew(STextBlock)
-			.Text(this, &SDMTextureUVVisualizerPopout::GetModeButtonText)
-			.Font(IDetailLayoutBuilder::GetDetailFontBold())
+			SNew(SCheckBox)
+			.Style(FAppStyle::Get(), "DetailsView.SectionButton")
+			.HAlign(EHorizontalAlignment::HAlign_Center)
+			.Padding(FVector2D(10.f, 3.f))
+			.IsChecked(this, &SDMTextureUVVisualizerPopout::GetModeCheckBoxState, /* Is Pivot */ false)
+			.OnCheckStateChanged(this, &SDMTextureUVVisualizerPopout::OnModeCheckBoxStateChanged, /* Is Pivot */ false)
+			.ToolTipText(LOCTEXT("VisualizerOffsetToolTip", "Allows changing of the UV offset."))
+			.Content()
+			[
+				SNew(STextBlock)
+				.Text(LOCTEXT("VisualizerOffset", "Offset"))
+				.Font(IDetailLayoutBuilder::GetDetailFontBold())
+			]
+		]
+		+ SHorizontalBox::Slot()
+		.FillContentWidth(1.f)
+		[
+			SNew(SCheckBox)
+			.Style(FAppStyle::Get(), "DetailsView.SectionButton")
+			.HAlign(EHorizontalAlignment::HAlign_Center)
+			.Padding(FVector2D(10.f, 3.f))
+			.IsChecked(this, &SDMTextureUVVisualizerPopout::GetModeCheckBoxState, /* Is Pivot */ true)
+			.OnCheckStateChanged(this, &SDMTextureUVVisualizerPopout::OnModeCheckBoxStateChanged, /* Is Pivot */ true)
+			.ToolTipText(LOCTEXT("VisualizerPivotToolTip", "Allows changing of the UV pivot, rotation and scale."))
+			.Content()
+			[
+				SNew(STextBlock)
+				.Text(LOCTEXT("VisualizerPivot", "Pivot"))
+				.Font(IDetailLayoutBuilder::GetDetailFontBold())
+			]
 		];
 
 	TextureUVPropertyRows.Add(EditModeButtonRow);
@@ -325,4 +343,26 @@ TSharedRef<SWidget> SDMTextureUVVisualizerPopout::CreatePropertyWidget(UDMTextur
 	return DetailsView;
 }
 
+ECheckBoxState SDMTextureUVVisualizerPopout::GetModeCheckBoxState(bool bInIsPivot) const
+{
+	if (!Visualizer.IsValid())
+	{
+		return ECheckBoxState::Undetermined;
+	}
+
+	return Visualizer->IsInPivotEditMode() == bInIsPivot
+		? ECheckBoxState::Checked
+		: ECheckBoxState::Unchecked;
+}
+
+void SDMTextureUVVisualizerPopout::OnModeCheckBoxStateChanged(ECheckBoxState InState, bool bInIsPivot)
+{
+	if (InState != ECheckBoxState::Checked || !Visualizer.IsValid())
+	{
+		return;
+	}
+
+	Visualizer->SetInPivotEditMode(bInIsPivot);
+}
+
 #undef LOCTEXT_NAMESPACE
diff --git a/Engine/Plugins/Experimental/DynamicMaterial/Source/DynamicMaterialEditor/Private/Slate/Properties/SDMTextureUVVisualizerPopout.h b/Engine/Plugins/Experimental/DynamicMaterial/Source/DynamicMaterialEditor/Private/Slate/Properties/SDMTextureUVVisualizerPopout.h
--- a/Engine/Plugins/Experimental/DynamicMaterial/Source/DynamicMaterialEditor/Private/Slate/Properties/SDMTextureUVVisualizerPopout.h
+++ b/Engine/Plugins/Experimental/DynamicMaterial/Source/DynamicMaterialEditor/Private/Slate/Properties/SDMTextureUVVisualizerPopout.h
@@ -10,6 +10,7 @@ class SDMTextureUVVisualizer;
 class SDockTab;
 class UDMMaterialStage;
 class UDMTextureUV;
+enum class ECheckBoxState : uint8;
 
 /**
  * Material Designer Texture UV Visualizer Popout
@@ -35,8 +36,6 @@ public:
 protected:
 	TSharedPtr<SDMTextureUVVisualizer> Visualizer;
 
-	FReply OnToggleModeClicked();
-
 	FText GetModeButtonText() const;
 
 	FVector2D GetHorizontalBarSize() const;
@@ -44,4 +43,8 @@ protected:
 	FVector2D GetSideBlockSize() const;
 
 	TSharedRef<SWidget> CreatePropertyWidget(UDMTextureUV* InTextureUV);
+
+	ECheckBoxState GetModeCheckBoxState(bool bInIsPivot) const;
+
+	void OnModeCheckBoxStateChanged(ECheckBoxState InState, bool bInIsPivot);
 };
diff --git a/Engine/Plugins/Experimental/DynamicMaterial/Source/DynamicMaterialEditor/Private/Slate/Properties/SDMTextureUVVisualizerProperty.cpp b/Engine/Plugins/Experimental/DynamicMaterial/Source/DynamicMaterialEditor/Private/Slate/Properties/SDMTextureUVVisualizerProperty.cpp
--- a/Engine/Plugins/Experimental/DynamicMaterial/Source/DynamicMaterialEditor/Private/Slate/Properties/SDMTextureUVVisualizerProperty.cpp
+++ b/Engine/Plugins/Experimental/DynamicMaterial/Source/DynamicMaterialEditor/Private/Slate/Properties/SDMTextureUVVisualizerProperty.cpp
@@ -9,6 +9,7 @@
 #include "Slate/Properties/SDMTextureUVVisualizerPopout.h"
 #include "Widgets/Docking/SDockTab.h"
 #include "Widgets/Input/SButton.h"
+#include "Widgets/Input/SCheckBox.h"
 #include "Widgets/Layout/SBox.h"
 #include "Widgets/Text/STextBlock.h"
 
@@ -24,19 +25,18 @@ void SDMTextureUVVisualizerProperty::Construct(const FArguments& InArgs, UDMMate
 		SNew(SVerticalBox)
 		+ SVerticalBox::Slot()
 		.AutoHeight()
-		.HAlign(EHorizontalAlignment::HAlign_Fill)
+		.HAlign(EHorizontalAlignment::HAlign_Left)
+		.Padding(0.f, 3.f, 0.f, 0.f)
 		[
 			SNew(SHorizontalBox)
-			.Visibility(this, &SDMTextureUVVisualizerProperty::GetVisualizerVisibility)
 			+ SHorizontalBox::Slot()
-			.FillContentWidth(1.f)
 			[
 				SNew(SButton)
-				.OnClicked(this, &SDMTextureUVVisualizerProperty::OnOpenPopoutClicked)
+				.OnClicked(this, &SDMTextureUVVisualizerProperty::OnToggleVisualizerClicked)
 				.Content()
 				[
 					SNew(STextBlock)
-					.Text(LOCTEXT("PopoutVisualizer", "Popout"))
+					.Text(LOCTEXT("ToggleVisualizer", "Toggle"))
 					.Font(IDetailLayoutBuilder::GetDetailFontBold())
 				]
 			]
@@ -44,11 +44,11 @@ void SDMTextureUVVisualizerProperty::Construct(const FArguments& InArgs, UDMMate
 			.FillContentWidth(1.f)
 			[
 				SNew(SButton)
-				.OnClicked(this, &SDMTextureUVVisualizerProperty::OnToggleModeClicked)
+				.OnClicked(this, &SDMTextureUVVisualizerProperty::OnOpenPopoutClicked)
 				.Content()
 				[
 					SNew(STextBlock)
-					.Text(this, &SDMTextureUVVisualizerProperty::GetModeButtonText)
+					.Text(LOCTEXT("PopoutVisualizer", "Popout"))
 					.Font(IDetailLayoutBuilder::GetDetailFontBold())
 				]
 			]
@@ -70,18 +70,41 @@ void SDMTextureUVVisualizerProperty::Construct(const FArguments& InArgs, UDMMate
 		]
 		+ SVerticalBox::Slot()
 		.AutoHeight()
-		.HAlign(EHorizontalAlignment::HAlign_Left)
-		.Padding(0.f, 3.f, 0.f, 0.f)
+		.HAlign(EHorizontalAlignment::HAlign_Fill)
 		[
 			SNew(SHorizontalBox)
+			.Visibility(this, &SDMTextureUVVisualizerProperty::GetVisualizerVisibility)
 			+ SHorizontalBox::Slot()
+			.FillContentWidth(1.f)
 			[
-				SNew(SButton)
-				.OnClicked(this, &SDMTextureUVVisualizerProperty::OnToggleVisualizerClicked)
+				SNew(SCheckBox)
+				.Style(FAppStyle::Get(), "DetailsView.SectionButton")
+				.HAlign(EHorizontalAlignment::HAlign_Center)
+				.Padding(FVector2D(5.f, 3.f))
+				.IsChecked(this, &SDMTextureUVVisualizerProperty::GetModeCheckBoxState, /* Is Pivot */ false)
+				.OnCheckStateChanged(this, &SDMTextureUVVisualizerProperty::OnModeCheckBoxStateChanged, /* Is Pivot */ false)
+				.ToolTipText(LOCTEXT("VisualizerOffsetToolTip", "Allows changing of the UV offset."))
 				.Content()
 				[
 					SNew(STextBlock)
-					.Text(LOCTEXT("ToggleVisualizer", "Toggle"))
+					.Text(LOCTEXT("VisualizerOffset", "Offset"))
+					.Font(IDetailLayoutBuilder::GetDetailFontBold())
+				]
+			]
+			+ SHorizontalBox::Slot()
+			.FillContentWidth(1.f)
+			[
+				SNew(SCheckBox)
+				.Style(FAppStyle::Get(), "DetailsView.SectionButton")
+				.HAlign(EHorizontalAlignment::HAlign_Center)
+				.Padding(FVector2D(5.f, 3.f))
+				.IsChecked(this, &SDMTextureUVVisualizerProperty::GetModeCheckBoxState, /* Is Pivot */ true)
+				.OnCheckStateChanged(this, &SDMTextureUVVisualizerProperty::OnModeCheckBoxStateChanged, /* Is Pivot */ true)
+				.ToolTipText(LOCTEXT("VisualizerPivotToolTip", "Allows changing of the UV pivot, rotation and scale."))
+				.Content()
+				[
+					SNew(STextBlock)
+					.Text(LOCTEXT("VisualizerPivot", "Pivot"))
 					.Font(IDetailLayoutBuilder::GetDetailFontBold())
 				]
 			]
@@ -111,16 +134,6 @@ FReply SDMTextureUVVisualizerProperty::OnToggleModeClicked()
 	return FReply::Handled();
 }
 
-FText SDMTextureUVVisualizerProperty::GetModeButtonText() const
-{
-	if (Visualizer.IsValid() && Visualizer->IsInPivotEditMode())
-	{
-		return LOCTEXT("VisualizerPivot", "Pivot");
-	}
-
-	return LOCTEXT("VisualizerOffset", "Offset");
-}
-
 FReply SDMTextureUVVisualizerProperty::OnOpenPopoutClicked()
 {
 	if (!Visualizer.IsValid())
@@ -153,4 +166,26 @@ EVisibility SDMTextureUVVisualizerProperty::GetVisualizerVisibility() const
 	return EVisibility::Collapsed;
 }
 
+ECheckBoxState SDMTextureUVVisualizerProperty::GetModeCheckBoxState(bool bInIsPivot) const
+{
+	if (!Visualizer.IsValid())
+	{
+		return ECheckBoxState::Undetermined;
+	}
+
+	return Visualizer->IsInPivotEditMode() == bInIsPivot
+		? ECheckBoxState::Checked
+		: ECheckBoxState::Unchecked;
+}
+
+void SDMTextureUVVisualizerProperty::OnModeCheckBoxStateChanged(ECheckBoxState InState, bool bInIsPivot)
+{
+	if (InState != ECheckBoxState::Checked || !Visualizer.IsValid())
+	{
+		return;
+	}
+
+	Visualizer->SetInPivotEditMode(bInIsPivot);
+}
+
 #undef LOCTEXT_NAMESPACE
diff --git a/Engine/Plugins/Experimental/DynamicMaterial/Source/DynamicMaterialEditor/Private/Slate/Properties/SDMTextureUVVisualizerProperty.h b/Engine/Plugins/Experimental/DynamicMaterial/Source/DynamicMaterialEditor/Private/Slate/Properties/SDMTextureUVVisualizerProperty.h
--- a/Engine/Plugins/Experimental/DynamicMaterial/Source/DynamicMaterialEditor/Private/Slate/Properties/SDMTextureUVVisualizerProperty.h
+++ b/Engine/Plugins/Experimental/DynamicMaterial/Source/DynamicMaterialEditor/Private/Slate/Properties/SDMTextureUVVisualizerProperty.h
@@ -8,6 +8,7 @@ class SDMTextureUVVisualizer;
 class SDMTextureUVVisualizerPopout;
 class UDMMaterialStage;
 class UDMTextureUV;
+enum class ECheckBoxState : uint8;
 
 /**
  * Material Designer Texture UV Visualizer Property
@@ -31,9 +32,11 @@ protected:
 
 	FReply OnToggleModeClicked();
 
-	FText GetModeButtonText() const;
-
 	FReply OnOpenPopoutClicked();
 
 	EVisibility GetVisualizerVisibility() const;
+
+	ECheckBoxState GetModeCheckBoxState(bool bInIsPivot) const;
+
+	void OnModeCheckBoxStateChanged(ECheckBoxState InState, bool bInIsPivot);
 };
