Material Designer
- Update to texture uv / row generator to add reset to default callbacks

diff --git a/Engine/Plugins/Experimental/DynamicMaterial/Source/DynamicMaterial/Private/Components/DMTextureUV.cpp b/Engine/Plugins/Experimental/DynamicMaterial/Source/DynamicMaterial/Private/Components/DMTextureUV.cpp
--- a/Engine/Plugins/Experimental/DynamicMaterial/Source/DynamicMaterial/Private/Components/DMTextureUV.cpp
+++ b/Engine/Plugins/Experimental/DynamicMaterial/Source/DynamicMaterial/Private/Components/DMTextureUV.cpp
@@ -378,6 +378,72 @@ void UDMTextureUV::SetMIDParameters(UMaterialInstanceDynamic* InMID)
 }
 
 #if WITH_EDITOR
+bool UDMTextureUV::CanResetToDefault(TSharedPtr<IPropertyHandle> InPropertyHandle) const
+{
+	FProperty* Property = InPropertyHandle->GetProperty();
+
+	if (!Property)
+	{
+		return false;
+	}
+
+	FName PropertyName = Property->GetFName();
+
+	if (PropertyName == NAME_None)
+	{
+		return false;
+	}
+
+	UDMTextureUV* DefaultObject = Cast<UDMTextureUV>(GetClass()->GetDefaultObject());
+
+	if (!DefaultObject)
+	{
+		return false;
+	}
+
+	if (PropertyName == NAME_UVSource)
+	{
+		return DefaultObject->GetUVSource() != GetUVSource();
+	}
+
+	if (PropertyName == NAME_bMirrorOnX)
+	{
+		return DefaultObject->GetMirrorOnX() != GetMirrorOnX();
+	}
+
+	if (PropertyName == NAME_bMirrorOnY)
+	{
+		return DefaultObject->GetMirrorOnY() != GetMirrorOnY();
+	}
+
+	if (PropertyName == NAME_Offset)
+	{
+		return !DefaultObject->GetOffset().Equals(GetOffset());
+	}
+
+	if (PropertyName == NAME_Pivot)
+	{
+		return !DefaultObject->GetPivot().Equals(GetPivot());
+	}
+
+	if (PropertyName == NAME_Rotation)
+	{
+		return DefaultObject->GetRotation() != GetRotation();
+	}
+
+	if (PropertyName == NAME_Scale)
+	{
+		return !DefaultObject->GetScale().Equals(GetScale());
+	}
+
+	return false;
+}
+
+void UDMTextureUV::ResetToDefault(TSharedPtr<IPropertyHandle> InPropertyHandle)
+{
+	InPropertyHandle->ResetToDefault();
+}
+
 bool UDMTextureUV::Modify(bool bInAlwaysMarkDirty)
 {
 	const bool bSaved = Super::Modify(bInAlwaysMarkDirty);
diff --git a/Engine/Plugins/Experimental/DynamicMaterial/Source/DynamicMaterial/Public/Components/DMTextureUV.h b/Engine/Plugins/Experimental/DynamicMaterial/Source/DynamicMaterial/Public/Components/DMTextureUV.h
--- a/Engine/Plugins/Experimental/DynamicMaterial/Source/DynamicMaterial/Public/Components/DMTextureUV.h
+++ b/Engine/Plugins/Experimental/DynamicMaterial/Source/DynamicMaterial/Public/Components/DMTextureUV.h
@@ -155,6 +155,11 @@ public:
 
 	void SetMIDParameters(UMaterialInstanceDynamic* InMID);
 
+#if WITH_EDITOR
+	bool CanResetToDefault(TSharedPtr<IPropertyHandle> InPropertyHandle) const;
+	void ResetToDefault(TSharedPtr<IPropertyHandle> InPropertyHandle);
+#endif
+
 	//~ Begin UObject
 #if WITH_EDITOR
 	virtual bool Modify(bool bInAlwaysMarkDirty = true) override;
diff --git a/Engine/Plugins/Experimental/DynamicMaterial/Source/DynamicMaterialEditor/Private/Slate/Properties/Generators/DMTextureUVPropertyRowGenerator.cpp b/Engine/Plugins/Experimental/DynamicMaterial/Source/DynamicMaterialEditor/Private/Slate/Properties/Generators/DMTextureUVPropertyRowGenerator.cpp
--- a/Engine/Plugins/Experimental/DynamicMaterial/Source/DynamicMaterialEditor/Private/Slate/Properties/Generators/DMTextureUVPropertyRowGenerator.cpp
+++ b/Engine/Plugins/Experimental/DynamicMaterial/Source/DynamicMaterialEditor/Private/Slate/Properties/Generators/DMTextureUVPropertyRowGenerator.cpp
@@ -103,7 +103,12 @@ bool FDMTextureUVPropertyRowGenerator::AllowKeyframeButton(UDMMaterialComponent*
 
 void UE::DynamicMaterialEditor::Private::AddPropertyRow(const TSharedRef<SWidget>& InComponentEditWidget, UDMTextureUV* InTextureUV, FName InProperty, TArray<FDMPropertyHandle>& InOutPropertyRows)
 {
-	InOutPropertyRows.Add(SDMEditor::GetPropertyHandle(&*InComponentEditWidget, InTextureUV, InProperty));
+	FDMPropertyHandle& NewHandle = InOutPropertyRows.Add_GetRef(SDMEditor::GetPropertyHandle(&*InComponentEditWidget, InTextureUV, InProperty));
+
+	NewHandle.ResetToDefaultOverride = FResetToDefaultOverride::Create(
+		FIsResetToDefaultVisible::CreateUObject(InTextureUV, &UDMTextureUV::CanResetToDefault),
+		FResetToDefaultHandler::CreateUObject(InTextureUV, &UDMTextureUV::ResetToDefault)
+	);
 }
 
 void UE::DynamicMaterialEditor::Private::AddVisualizerRow(const TSharedRef<SWidget>& InComponentEditWidget, UDMTextureUV* InTextureUV, TArray<FDMPropertyHandle>& InOutPropertyRows)
