Material Designer
- Renamed channel preset variable to bludgeon our way through updating defaults.
- Added the 'Emissive' default
- Fixed the slot picker so that it displays the emissive bubble when it is actually emissive, instead of showing base color for both.

diff --git a/Engine/Plugins/Experimental/DynamicMaterial/Source/DynamicMaterialEditor/Private/DynamicMaterialEditorSettings.cpp b/Engine/Plugins/Experimental/DynamicMaterial/Source/DynamicMaterialEditor/Private/DynamicMaterialEditorSettings.cpp
--- a/Engine/Plugins/Experimental/DynamicMaterial/Source/DynamicMaterialEditor/Private/DynamicMaterialEditorSettings.cpp
+++ b/Engine/Plugins/Experimental/DynamicMaterial/Source/DynamicMaterialEditor/Private/DynamicMaterialEditorSettings.cpp
@@ -1,8 +1,10 @@
 // Copyright Epic Games, Inc. All Rights Reserved.
 
 #include "DynamicMaterialEditorSettings.h"
+
 #include "AssetRegistry/AssetData.h"
 #include "Components/DMMaterialComponent.h"
+#include "DynamicMaterialEditorModule.h"
 #include "Engine/AssetManager.h"
 #include "GenericPlatform/GenericApplication.h"
 #include "Interfaces/IPluginManager.h"
@@ -89,6 +91,7 @@ UDynamicMaterialEditorSettings::UDynamicMaterialEditorSettings()
 	OverriddenDefaultSlotTextures = {};
 
 	FDMMaterialChannelListPreset Opaque;
+	Opaque.Name = TEXT("Opaque");
 	Opaque.bRGB = true;
 	Opaque.bOpacity = false;
 	Opaque.DefaultBlendMode = BLEND_Opaque;
@@ -96,7 +99,17 @@ UDynamicMaterialEditorSettings::UDynamicMaterialEditorSettings()
 	Opaque.bDefaultAnimated = false;
 	Opaque.bDefaultTwoSided = true;
 
+	FDMMaterialChannelListPreset Emissive;
+	Emissive.Name = TEXT("Emissive");
+	Emissive.bRGB = true;
+	Emissive.bOpacity = false;
+	Emissive.DefaultBlendMode = BLEND_Opaque;
+	Emissive.DefaultShadingModel = EDMMaterialShadingModel::Unlit;
+	Emissive.bDefaultAnimated = false;
+	Emissive.bDefaultTwoSided = true;
+
 	FDMMaterialChannelListPreset Translucent;
+	Translucent.Name = TEXT("Translucent");
 	Translucent.bRGB = true;
 	Translucent.bOpacity = true;
 	Translucent.DefaultBlendMode = BLEND_Translucent;
@@ -105,6 +118,7 @@ UDynamicMaterialEditorSettings::UDynamicMaterialEditorSettings()
 	Translucent.bDefaultTwoSided = true;
 
 	FDMMaterialChannelListPreset PBR;
+	PBR.Name = TEXT("PBR");
 	PBR.bRGB = true;
 	PBR.bOpacity = true;
 	PBR.bMetallic = true;
@@ -118,6 +132,7 @@ UDynamicMaterialEditorSettings::UDynamicMaterialEditorSettings()
 	PBR.bDefaultTwoSided = true;
 
 	FDMMaterialChannelListPreset All;
+	All.Name = TEXT("All");
 	All.bRGB = true;
 	All.bOpacity = true;
 	All.bMetallic = true;
@@ -135,10 +150,11 @@ UDynamicMaterialEditorSettings::UDynamicMaterialEditorSettings()
 	All.bDefaultAnimated = false;
 	All.bDefaultTwoSided = true;
 
-	ChannelPresets.Add(TEXT("PBR"), PBR);
-	ChannelPresets.Add(TEXT("Opaque"), Opaque);
-	ChannelPresets.Add(TEXT("Translucent"), Translucent);
-	ChannelPresets.Add(TEXT("All"), All);		
+	MaterialChannelPresets.Add(PBR);
+	MaterialChannelPresets.Add(Opaque);
+	MaterialChannelPresets.Add(Emissive);
+	MaterialChannelPresets.Add(Translucent);
+	MaterialChannelPresets.Add(All);		
 }
 
 UDynamicMaterialEditorSettings* UDynamicMaterialEditorSettings::Get()
@@ -153,6 +169,19 @@ UDynamicMaterialEditorSettings* UDynamicMaterialEditorSettings::Get()
 	return DefaultSettings;
 }
 
+void UDynamicMaterialEditorSettings::PreEditChange(FEditPropertyChain& InPropertyAboutToChange)
+{
+	Super::PreEditChange(InPropertyAboutToChange);
+
+	PreEditPresetNames.Empty();
+	PreEditPresetNames.Reserve(MaterialChannelPresets.Num());
+
+	for (const FDMMaterialChannelListPreset& Preset : MaterialChannelPresets)
+	{
+		PreEditPresetNames.Add(Preset.Name);
+	}
+}
+
 void UDynamicMaterialEditorSettings::PostEditChangeProperty(FPropertyChangedEvent& InPropertyChangedEvent)
 {
 	Super::PostEditChangeProperty(InPropertyChangedEvent);
@@ -182,6 +211,11 @@ void UDynamicMaterialEditorSettings::PostEditChangeProperty(FPropertyChangedEven
 			Component->MarkComponentDirty();
 		}
 	}
+
+	if (PropertyName == GET_MEMBER_NAME_CHECKED(UDynamicMaterialEditorSettings, MaterialChannelPresets))
+	{
+		EnsureUniqueChannelPresetNames();
+	}
 }
 
 void UDynamicMaterialEditorSettings::OpenEditorSettingsWindow() const
@@ -358,4 +392,122 @@ UTexture* UDynamicMaterialEditorSettings::GetDefaultTextureForSlot(EDMMaterialPr
 	return nullptr;
 }
 
+const FDMMaterialChannelListPreset* UDynamicMaterialEditorSettings::GetPresetByName(FName InName) const
+{
+	for (const FDMMaterialChannelListPreset& Preset : MaterialChannelPresets)
+	{
+		if (InName == Preset.Name)
+		{
+			return &Preset;
+		}
+	}
+
+	return nullptr;
+}
+
+void UDynamicMaterialEditorSettings::EnsureUniqueChannelPresetNames()
+{
+	// This is splitting the *text* portion of the FName, not the number part.
+	// So an FName might me Foo5_1 and this would return "Foo" and 5.
+	auto SplitName = [](FString InName, FString& OutName, int32& OutNumber)
+		{
+			const int32 NameLen = InName.Len();
+
+			for (int32 CharIndex = NameLen - 1; CharIndex >= 0; --CharIndex)
+			{
+				if (InName[CharIndex] >= '0' && InName[CharIndex] <= '9')
+				{
+					continue;
+				}
+
+				OutName = InName.Left(CharIndex + 1);
+
+				if (CharIndex < (NameLen - 1))
+				{
+					OutNumber = FCString::Atoi(*InName.Mid(CharIndex + 1));
+				}
+				else
+				{
+					OutNumber = 0;
+				}
+
+				return;
+			}
+
+			OutName = "";
+
+			// Every char was a digit.
+			if (NameLen > 0)
+			{
+				OutNumber = FCString::Atoi(*InName);
+			}
+			else
+			{
+				OutNumber = 0;
+			}
+		};
+
+	const int32 Count = MaterialChannelPresets.Num();
+
+	for (int32 IndexBase = 0; IndexBase < Count; ++IndexBase)
+	{
+		// Name hasn't changed, don't try to fix it
+		if (PreEditPresetNames.IsValidIndex(IndexBase) && PreEditPresetNames[IndexBase] == MaterialChannelPresets[IndexBase].Name)
+		{
+			continue;
+		}
+
+		for (int32 IndexCheck = 0; IndexCheck < Count; ++IndexCheck)
+		{
+			if (IndexCheck == IndexBase)
+			{
+				continue;
+			}
+
+			const bool bEqual = MaterialChannelPresets[IndexBase].Name.IsEqual(
+				MaterialChannelPresets[IndexCheck].Name,
+				ENameCase::IgnoreCase,
+				/* Check number */ false
+			);
+
+			if (!bEqual)
+			{
+				continue;
+			}
+
+			UE_LOG(LogDynamicMaterialEditor, Warning, TEXT("Duplicate channel list preset name detected."));
+
+			FString BaseName;
+			int32 NumberSuffix;
+			SplitName(MaterialChannelPresets[IndexBase].Name.GetPlainNameString(), BaseName, NumberSuffix);
+
+			if (NumberSuffix < 2)
+			{
+				NumberSuffix = 2;
+			}
+
+			for (int32 IndexSameNameCheck = 0; IndexSameNameCheck < Count; ++IndexSameNameCheck)
+			{
+				if (IndexSameNameCheck == IndexBase)
+				{
+					continue;
+				}
+
+				FString BaseNameCheck;
+				int32 NumberSuffixCheck;
+				SplitName(MaterialChannelPresets[IndexSameNameCheck].Name.GetPlainNameString(), BaseNameCheck, NumberSuffixCheck);
+
+				if (BaseNameCheck.Equals(BaseName, ESearchCase::IgnoreCase))
+				{
+					NumberSuffix = FMath::Max(NumberSuffix, NumberSuffixCheck + 1);
+				}
+			}
+
+			MaterialChannelPresets[IndexBase].Name = *(BaseName + FString::FromInt(NumberSuffix));
+		}
+	}
+
+	PreEditPresetNames.Empty();
+}
+
 #undef LOCTEXT_NAMESPACE
diff --git a/Engine/Plugins/Experimental/DynamicMaterial/Source/DynamicMaterialEditor/Private/DynamicMaterialEditorSettings.h b/Engine/Plugins/Experimental/DynamicMaterial/Source/DynamicMaterialEditor/Private/DynamicMaterialEditorSettings.h
--- a/Engine/Plugins/Experimental/DynamicMaterial/Source/DynamicMaterialEditor/Private/DynamicMaterialEditorSettings.h
+++ b/Engine/Plugins/Experimental/DynamicMaterial/Source/DynamicMaterialEditor/Private/DynamicMaterialEditorSettings.h
@@ -34,6 +34,9 @@ struct FDMMaterialChannelListPreset
 {
 	GENERATED_BODY()
 
+	UPROPERTY(Config, EditAnywhere, BlueprintReadOnly, Category = "Channel")
+	FName Name;
+
 	/** This cannot be disabled. It is here for illustrative purposes only. */
 	UPROPERTY(Config, VisibleAnywhere, BlueprintReadOnly, Category = "Channel")
 	bool bRGB = true;
@@ -192,12 +195,13 @@ public:
 	UPROPERTY(Config, EditAnywhere, BlueprintReadWrite, Category = "Effects")
 	TArray<FName> CustomEffectsFolders;
 
-	UPROPERTY(Config, EditAnywhere, BlueprintReadWrite, Category = "Channels")
-	TMap<FName, FDMMaterialChannelListPreset> ChannelPresets;
+	UPROPERTY(Config, EditAnywhere, BlueprintReadWrite, Category = "Channels", meta = (TitleProperty = Name))
+	TArray<FDMMaterialChannelListPreset> MaterialChannelPresets;
 
 	FOnFinishedChangingProperties OnSettingsChanged;
 
 	//~ Begin UObject
+	virtual void PreEditChange(FEditPropertyChain& InPropertyAboutToChange) override;
 	virtual void PostEditChangeProperty(FPropertyChangedEvent& InPropertyChangedEvent) override;
 	//~ End UObject
 
@@ -210,4 +214,11 @@ public:
 	TArray<FDMMaterialEffectList> GetEffectList() const;
 
 	UTexture* GetDefaultTextureForSlot(EDMMaterialPropertyType InProperty) const;
+
+	const FDMMaterialChannelListPreset* GetPresetByName(FName InName) const;
+
+private:
+	TArray<FName> PreEditPresetNames;
+
+	void EnsureUniqueChannelPresetNames();
 };
diff --git a/Engine/Plugins/Experimental/DynamicMaterial/Source/DynamicMaterialEditor/Private/Menus/DMToolBarMenus.cpp b/Engine/Plugins/Experimental/DynamicMaterial/Source/DynamicMaterialEditor/Private/Menus/DMToolBarMenus.cpp
--- a/Engine/Plugins/Experimental/DynamicMaterial/Source/DynamicMaterialEditor/Private/Menus/DMToolBarMenus.cpp
+++ b/Engine/Plugins/Experimental/DynamicMaterial/Source/DynamicMaterialEditor/Private/Menus/DMToolBarMenus.cpp
@@ -78,15 +78,15 @@ namespace UE::DynamicMaterialEditor::Private
 	{
 		FToolMenuSection& NewSection = InMenu->AddSection("MaterialType", LOCTEXT("MaterialType", "Material Type"));
 
-		for (const TPair<FName, FDMMaterialChannelListPreset>& Preset : GetDefault<UDynamicMaterialEditorSettings>()->ChannelPresets)
+		for (const FDMMaterialChannelListPreset& Preset : GetDefault<UDynamicMaterialEditorSettings>()->MaterialChannelPresets)
 		{
 			FToolUIAction MaterialChannelListType;
-			MaterialChannelListType.ExecuteAction = FToolMenuExecuteAction::CreateStatic(&AddToolBarMaterialInstanceChannelListMenu_Execute, Preset.Key);
-			MaterialChannelListType.GetActionCheckState = FToolMenuGetActionCheckState::CreateStatic(&AddToolBarMaterialInstanceChannelListMenu_GetCheckState, Preset.Key);
+			MaterialChannelListType.ExecuteAction = FToolMenuExecuteAction::CreateStatic(&AddToolBarMaterialInstanceChannelListMenu_Execute, Preset.Name);
+			MaterialChannelListType.GetActionCheckState = FToolMenuGetActionCheckState::CreateStatic(&AddToolBarMaterialInstanceChannelListMenu_GetCheckState, Preset.Name);
 
 			NewSection.AddMenuEntry(
-				Preset.Key,
-				FText::FromName(Preset.Key),
+				Preset.Name,
+				FText::FromName(Preset.Name),
 				FText::GetEmpty(),
 				TAttribute<FSlateIcon>(),
 				FToolUIActionChoice(MaterialChannelListType),
diff --git a/Engine/Plugins/Experimental/DynamicMaterial/Source/DynamicMaterialEditor/Private/Model/DynamicMaterialModelEditorOnlyData.cpp b/Engine/Plugins/Experimental/DynamicMaterial/Source/DynamicMaterialEditor/Private/Model/DynamicMaterialModelEditorOnlyData.cpp
--- a/Engine/Plugins/Experimental/DynamicMaterial/Source/DynamicMaterialEditor/Private/Model/DynamicMaterialModelEditorOnlyData.cpp
+++ b/Engine/Plugins/Experimental/DynamicMaterial/Source/DynamicMaterialEditor/Private/Model/DynamicMaterialModelEditorOnlyData.cpp
@@ -695,7 +695,7 @@ void UDynamicMaterialModelEditorOnlyData::SetChannelListPreset(FName InPresetNam
 {
 	ChannelListPreset = InPresetName;
 
-	const FDMMaterialChannelListPreset* Preset = GetDefault<UDynamicMaterialEditorSettings>()->ChannelPresets.Find(InPresetName);
+	const FDMMaterialChannelListPreset* Preset = GetDefault<UDynamicMaterialEditorSettings>()->GetPresetByName(InPresetName);
 
 	if (!Preset)
 	{
diff --git a/Engine/Plugins/Experimental/DynamicMaterial/Source/DynamicMaterialEditor/Private/Slate/SDMEditor.cpp b/Engine/Plugins/Experimental/DynamicMaterial/Source/DynamicMaterialEditor/Private/Slate/SDMEditor.cpp
--- a/Engine/Plugins/Experimental/DynamicMaterial/Source/DynamicMaterialEditor/Private/Slate/SDMEditor.cpp
+++ b/Engine/Plugins/Experimental/DynamicMaterial/Source/DynamicMaterialEditor/Private/Slate/SDMEditor.cpp
@@ -716,19 +716,23 @@ TSharedRef<SWidget> SDMEditor::CreateSlotPickerWidget()
 			]
 		];
 
-	const FDMMaterialChannelListPreset* Preset = GetDefault<UDynamicMaterialEditorSettings>()->ChannelPresets.Find(ModelEditorOnlyData->GetChannelListPreset());
+	const FDMMaterialChannelListPreset* Preset = GetDefault<UDynamicMaterialEditorSettings>()->GetPresetByName(ModelEditorOnlyData->GetChannelListPreset());
 	const bool bHasBaseColorSlot = !!ModelEditorOnlyData->GetSlotForMaterialProperty(EDMMaterialPropertyType::BaseColor);
 
 	UEnum* MaterialPropertyEnum = StaticEnum<EDMMaterialPropertyType>();
 
 	for (const TPair<EDMMaterialPropertyType, UDMMaterialProperty*>& Property : ModelEditorOnlyData->GetMaterialProperties())
 	{
-		// Always create base colour, not emissive - Will be sorted out by the button itself.
-		if (Property.Key == EDMMaterialPropertyType::EmissiveColor)
+		if (Property.Key == EDMMaterialPropertyType::BaseColor && ModelEditorOnlyData->GetShadingModel() == EDMMaterialShadingModel::Unlit)
 		{
 			continue;
 		}
 
+		if (Property.Key == EDMMaterialPropertyType::EmissiveColor && ModelEditorOnlyData->GetShadingModel() == EDMMaterialShadingModel::DefaultLit)
+		{
+			continue;
+		}		
+
 		// Always create opacity, not opacity mask - Will be sorted out by the button itself.
 		if (Property.Key == EDMMaterialPropertyType::OpacityMask)
 		{
diff --git a/Engine/Plugins/Experimental/DynamicMaterial/Source/DynamicMaterialEditor/Private/Slate/SDMMaterialWizard.cpp b/Engine/Plugins/Experimental/DynamicMaterial/Source/DynamicMaterialEditor/Private/Slate/SDMMaterialWizard.cpp
--- a/Engine/Plugins/Experimental/DynamicMaterial/Source/DynamicMaterialEditor/Private/Slate/SDMMaterialWizard.cpp
+++ b/Engine/Plugins/Experimental/DynamicMaterial/Source/DynamicMaterialEditor/Private/Slate/SDMMaterialWizard.cpp
@@ -33,11 +33,13 @@ void SDMMaterialWizard::Construct(const FArguments& InArgs, const TSharedRef<SDM
 {
 	EditorWeak = InEditor;
 
-	for (const TPair<FName, FDMMaterialChannelListPreset>& Preset : GetDefault<UDynamicMaterialEditorSettings>()->ChannelPresets)
+	if (const UDynamicMaterialEditorSettings* Settings = GetDefault<UDynamicMaterialEditorSettings>())
 	{
-		CurrentPreset = Preset.Key;
-		break;
-	}
+		if (!Settings->MaterialChannelPresets.IsEmpty())
+		{
+			CurrentPreset = Settings->MaterialChannelPresets[0].Name;
+		}
+	}	
 
 	ChildSlot
 		[
@@ -134,7 +136,7 @@ TSharedRef<SWidget> SDMMaterialWizard::CreateChannelPresets()
 		.InnerSlotPadding(WrapBoxSlotPadding)
 		.Orientation(EOrientation::Orient_Horizontal);
 
-	for (const TPair<FName, FDMMaterialChannelListPreset>& Preset : GetDefault<UDynamicMaterialEditorSettings>()->ChannelPresets)
+	for (const FDMMaterialChannelListPreset& Preset : GetDefault<UDynamicMaterialEditorSettings>()->MaterialChannelPresets)
 	{
 		ChannelPresets->AddSlot()
 			[
@@ -142,11 +144,11 @@ TSharedRef<SWidget> SDMMaterialWizard::CreateChannelPresets()
 					.Style(FAppStyle::Get(), "DetailsView.SectionButton")
 					.HAlign(EHorizontalAlignment::HAlign_Center)
 					.Padding(ButtonPadding)
-					.IsChecked(this, &SDMMaterialWizard::Preset_GetState, Preset.Key)
-					.OnCheckStateChanged(this, &SDMMaterialWizard::Preset_OnChange, Preset.Key)
+					.IsChecked(this, &SDMMaterialWizard::Preset_GetState, Preset.Name)
+					.OnCheckStateChanged(this, &SDMMaterialWizard::Preset_OnChange, Preset.Name)
 					[
 						SNew(STextBlock)
-						.Text(FText::FromName(Preset.Key))
+						.Text(FText::FromName(Preset.Name))
 					]
 			];
 	}
@@ -172,8 +174,7 @@ TSharedRef<SWidget> SDMMaterialWizard::CreateChannelList()
 		.InnerSlotPadding(WrapBoxSlotPadding)
 		.Orientation(EOrientation::Orient_Horizontal);
 
-
-	if (const FDMMaterialChannelListPreset* Preset = GetDefault<UDynamicMaterialEditorSettings>()->ChannelPresets.Find(CurrentPreset))
+	if (const FDMMaterialChannelListPreset* Preset = GetDefault<UDynamicMaterialEditorSettings>()->GetPresetByName(CurrentPreset))
 	{
 		for (const TPair<EDMMaterialPropertyType, UDMMaterialProperty*>& Property : ModelEditorOnlyData->GetMaterialProperties())
 		{
