Material Designer
- Layer effects will now be applied to stage preview materials.

diff --git a/Engine/Plugins/Experimental/DynamicMaterial/Source/DynamicMaterialEditor/Private/Components/DMMaterialStage.cpp b/Engine/Plugins/Experimental/DynamicMaterial/Source/DynamicMaterialEditor/Private/Components/DMMaterialStage.cpp
--- a/Engine/Plugins/Experimental/DynamicMaterial/Source/DynamicMaterialEditor/Private/Components/DMMaterialStage.cpp
+++ b/Engine/Plugins/Experimental/DynamicMaterial/Source/DynamicMaterialEditor/Private/Components/DMMaterialStage.cpp
@@ -3,25 +3,13 @@
 #include "Components/DMMaterialStage.h"
 #include "AssetRegistry/AssetRegistryModule.h"
 #include "Components/DMMaterialLayer.h"
-#include "Components/DMMaterialProperty.h"
 #include "Components/DMMaterialSlot.h"
-#include "Components/DMMaterialStageBlend.h"
-#include "Components/DMMaterialStageExpression.h"
-#include "Components/DMMaterialStageFunction.h"
-#include "Components/DMMaterialStageGradient.h"
 #include "Components/DMMaterialStageInput.h"
 #include "Components/DMMaterialStageSource.h"
 #include "Components/DMMaterialStageThroughputLayerBlend.h"
+#include "Components/DMMaterialSubStage.h"
 #include "Components/DMMaterialValue.h"
 #include "Components/DMTextureUV.h"
-#include "Components/MaterialStageExpressions/DMMSEMathBase.h"
-#include "Components/MaterialStageInputs/DMMSIExpression.h"
-#include "Components/MaterialStageInputs/DMMSIFunction.h"
-#include "Components/MaterialStageInputs/DMMSIGradient.h"
-#include "Components/MaterialStageInputs/DMMSISlot.h"
-#include "Components/MaterialStageInputs/DMMSITextureUV.h"
-#include "Components/MaterialStageInputs/DMMSIValue.h"
-#include "Containers/TransArray.h"
 #include "DMComponentPath.h"
 #include "DynamicMaterialEditorModule.h"
 #include "DynamicMaterialEditorSettings.h"
@@ -29,12 +17,9 @@
 #include "Factories/MaterialFactoryNew.h"
 #include "Materials/Material.h"
 #include "Materials/MaterialExpression.h"
-#include "Materials/MaterialExpressionAppendVector.h"
-#include "Materials/MaterialExpressionComponentMask.h"
 #include "Materials/MaterialInstanceDynamic.h"
 #include "MaterialValueType.h"
 #include "Model/DMMaterialBuildState.h"
-#include "Model/DMMaterialBuildUtils.h"
 #include "Model/DynamicMaterialModel.h"
 #include "Model/DynamicMaterialModelEditorOnlyData.h"
 #include "UObject/Package.h"
@@ -557,7 +542,7 @@ void UDMMaterialStage::GenerateExpressions(const TSharedRef<FDMMaterialBuildStat
 
 	if (!IsValid(Source))
 	{
-		UE_LOG(LogDynamicMaterialEditor, Warning, TEXT("Stage with no expressions attempted to generate material expressions."));
+		UE_LOG(LogDynamicMaterialEditor, Warning, TEXT("Stage with no source attempted to generate material expressions."));
 		return;
 	}
 
@@ -621,6 +606,27 @@ void UDMMaterialStage::GenerateExpressions(const TSharedRef<FDMMaterialBuildStat
 		StageExpressions.Append(InBuildState->GetStageSourceExpressions(Source));		
 	}
 
+	if (InBuildState->IsPreviewMaterial())
+	{
+		int32 OutputChannel = FDMMaterialStageConnectorChannel::WHOLE_CHANNEL;
+		int32 OutputIndex = 0;
+
+		const UDMMaterialStage* Stage = this;
+
+		if (const UDMMaterialSubStage* SubStage = Cast<const UDMMaterialSubStage>(Stage))
+		{
+			Stage = SubStage->GetParentMostStage();
+		}
+
+		Layer->ApplyEffects(
+			InBuildState,
+			Stage,
+			StageExpressions,
+			OutputChannel,
+			OutputIndex
+		);
+	}
+
 	InBuildState->AddStageExpressions(this, StageExpressions);
 }
 
