Material Designer
- Removed the "Material" from "Material Slot {x}"

diff --git a/Engine/Plugins/Experimental/DynamicMaterial/Source/DynamicMaterialEditor/Private/DMObjectMaterialProperty.cpp b/Engine/Plugins/Experimental/DynamicMaterial/Source/DynamicMaterialEditor/Private/DMObjectMaterialProperty.cpp
--- a/Engine/Plugins/Experimental/DynamicMaterial/Source/DynamicMaterialEditor/Private/DMObjectMaterialProperty.cpp
+++ b/Engine/Plugins/Experimental/DynamicMaterial/Source/DynamicMaterialEditor/Private/DMObjectMaterialProperty.cpp
@@ -252,14 +252,14 @@ FText FDMObjectMaterialProperty::GetPropertyName(bool bInIgnoreNewStatus) const
 				if (MaterialModel || bInIgnoreNewStatus)
 				{
 					return FText::Format(
-						LOCTEXT("MaterialListNameFormat", "Material Slot {0}"),
+						LOCTEXT("MaterialListNameFormat", "Slot {0}"),
 						FText::AsNumber(Index)
 					);
 				}
 				else
 				{
 					return FText::Format(
-						LOCTEXT("MaterialListNameFormatNew", "Material Slot {0} (Create New)"),
+						LOCTEXT("MaterialListNameFormatNew", "Slot {0} (Create New)"),
 						FText::AsNumber(Index)
 					);
 				}
