Material Designer
- Widget Render Target Renderer

diff --git a/Engine/Plugins/Experimental/DynamicMaterial/Source/DynamicMaterial/DynamicMaterial.Build.cs b/Engine/Plugins/Experimental/DynamicMaterial/Source/DynamicMaterial/DynamicMaterial.Build.cs
--- a/Engine/Plugins/Experimental/DynamicMaterial/Source/DynamicMaterial/DynamicMaterial.Build.cs
+++ b/Engine/Plugins/Experimental/DynamicMaterial/Source/DynamicMaterial/DynamicMaterial.Build.cs
@@ -13,7 +13,9 @@ public class DynamicMaterial : ModuleRules
 			{
 				"Core",
 				"CoreUObject",
-				"Slate"
+				"Slate",
+				"SlateCore",
+				"UMG"
 			}
 		);
 
diff --git a/Engine/Plugins/Experimental/DynamicMaterial/Source/DynamicMaterial/Private/Components/DMRenderTargetRenderer.cpp b/Engine/Plugins/Experimental/DynamicMaterial/Source/DynamicMaterial/Private/Components/DMRenderTargetRenderer.cpp
--- a/Engine/Plugins/Experimental/DynamicMaterial/Source/DynamicMaterial/Private/Components/DMRenderTargetRenderer.cpp
+++ b/Engine/Plugins/Experimental/DynamicMaterial/Source/DynamicMaterial/Private/Components/DMRenderTargetRenderer.cpp
@@ -65,3 +65,13 @@ void UDMRenderTargetRenderer::FlushUpdateRenderTarget()
 		UpdateRenderTarget();
 	}
 }
+
+void UDMRenderTargetRenderer::PostLoad()
+{
+	Super::PostLoad();
+
+	if (UDMMaterialValueRenderTarget* RenderTargetValue = GetRenderTargetValue())
+	{
+		RenderTargetValue->EnsureRenderTarget();
+	}
+}
diff --git a/Engine/Plugins/Experimental/DynamicMaterial/Source/DynamicMaterial/Private/Components/RenderTargetRenderers/DMRenderTargetTextRenderer.cpp b/Engine/Plugins/Experimental/DynamicMaterial/Source/DynamicMaterial/Private/Components/RenderTargetRenderers/DMRenderTargetTextRenderer.cpp
--- a/Engine/Plugins/Experimental/DynamicMaterial/Source/DynamicMaterial/Private/Components/RenderTargetRenderers/DMRenderTargetTextRenderer.cpp
+++ b/Engine/Plugins/Experimental/DynamicMaterial/Source/DynamicMaterial/Private/Components/RenderTargetRenderers/DMRenderTargetTextRenderer.cpp
@@ -57,16 +57,6 @@ UDMRenderTargetTextRenderer::UDMRenderTargetTextRenderer()
 #endif
 }
 
-void UDMRenderTargetTextRenderer::PostLoad()
-{
-	Super::PostLoad();
-
-	if (UDMMaterialValueRenderTarget* RenderTargetValue = GetRenderTargetValue())
-	{
-		RenderTargetValue->EnsureRenderTarget();
-	}
-}
-
 #if WITH_EDITOR
 void UDMRenderTargetTextRenderer::PostEditChangeProperty(FPropertyChangedEvent& InPropertyChangedEvent)
 {
diff --git a/Engine/Plugins/Experimental/DynamicMaterial/Source/DynamicMaterial/Private/Components/RenderTargetRenderers/DMRenderTargetWidgetRenderer.cpp b/Engine/Plugins/Experimental/DynamicMaterial/Source/DynamicMaterial/Private/Components/RenderTargetRenderers/DMRenderTargetWidgetRenderer.cpp
new file mode 100644
--- /dev/null
+++ b/Engine/Plugins/Experimental/DynamicMaterial/Source/DynamicMaterial/Private/Components/RenderTargetRenderers/DMRenderTargetWidgetRenderer.cpp
@@ -0,0 +1,98 @@
+// Copyright Epic Games, Inc. All Rights Reserved.
+
+#include "Components/RenderTargetRenderers/DMRenderTargetWidgetRenderer.h"
+#include "Components/Widget.h"
+#include "Components/MaterialValues/DMMaterialValueRenderTarget.h"
+#include "Engine/Engine.h"
+#include "Engine/TextureRenderTarget2D.h"
+#include "Slate/WidgetRenderer.h"
+
+#define LOCWidget_NAMESPACE "DMRenderTargetWidgetRenderer"
+
+UDMRenderTargetWidgetRenderer::UDMRenderTargetWidgetRenderer()
+{
+#if WITH_EDITOR
+	EditableProperties.Add(GET_MEMBER_NAME_CHECKED(UDMRenderTargetWidgetRenderer, WidgetClass));
+#endif
+
+	WidgetRenderer = MakeShared<FWidgetRenderer>(/* Gamma correction */ false);
+	WidgetRenderer->SetIsPrepassNeeded(true);
+	WidgetRenderer->SetShouldClearTarget(true);
+}
+
+void UDMRenderTargetWidgetRenderer::SetWidgetClass(TSubclassOf<UWidget> InWidgetClass)
+{
+	if (InWidgetClass == WidgetClass)
+	{
+		return;
+	}
+
+	WidgetClass = InWidgetClass;
+	CreateWidgetInstance();
+	AsyncUpdateRenderTarget();
+}
+
+#if WITH_EDITOR
+void UDMRenderTargetWidgetRenderer::PostEditChangeProperty(FPropertyChangedEvent& InPropertyChangedEvent)
+{
+	Super::PostEditChangeProperty(InPropertyChangedEvent);
+
+	const FName PropertyName = InPropertyChangedEvent.GetMemberPropertyName();
+
+	if (PropertyName == GET_MEMBER_NAME_CHECKED(UDMRenderTargetWidgetRenderer, WidgetClass))
+	{
+		CreateWidgetInstance();
+		AsyncUpdateRenderTarget();
+	}
+}
+#endif
+
+void UDMRenderTargetWidgetRenderer::CreateWidgetInstance()
+{
+	UClass* WidgetClassLocal = WidgetClass.Get();
+
+	if (!WidgetClassLocal)
+	{
+		return;
+	}
+
+	WidgetInstance = NewObject<UWidget>(this, WidgetClassLocal, NAME_None, RF_Transient);
+}
+
+void UDMRenderTargetWidgetRenderer::UpdateRenderTarget_Internal()
+{
+	Super::UpdateRenderTarget_Internal();
+
+	UDMMaterialValueRenderTarget* RenderTargetValue = GetRenderTargetValue();
+
+	if (!RenderTargetValue)
+	{
+		return;
+	}
+
+	RenderTargetValue->EnsureRenderTarget(/* Async */ false);
+
+	UTextureRenderTarget2D* RenderTarget = RenderTargetValue->GetRenderTarget();
+
+	if (!RenderTarget)
+	{
+		return;
+	}
+
+	if (!WidgetInstance)
+	{
+		CreateWidgetInstance();
+
+		if (!WidgetInstance)
+		{
+			return;
+		}
+	}
+
+	WidgetRenderer->DrawWidget(
+		RenderTarget,
+		WidgetInstance->TakeWidget(),
+		{(double)RenderTarget->SizeX, (double)RenderTarget->SizeX},
+		/* Delta Time */ 0.f
+	);
+}
diff --git a/Engine/Plugins/Experimental/DynamicMaterial/Source/DynamicMaterial/Public/Components/DMRenderTargetRenderer.h b/Engine/Plugins/Experimental/DynamicMaterial/Source/DynamicMaterial/Public/Components/DMRenderTargetRenderer.h
--- a/Engine/Plugins/Experimental/DynamicMaterial/Source/DynamicMaterial/Public/Components/DMRenderTargetRenderer.h
+++ b/Engine/Plugins/Experimental/DynamicMaterial/Source/DynamicMaterial/Public/Components/DMRenderTargetRenderer.h
@@ -42,6 +42,10 @@ public:
 	UFUNCTION(BlueprintPure, Category = "Material Designer")
 	bool IsUpdating() const { return bUpdating; }
 
+	//~ Begin UObject
+	virtual void PostLoad() override;
+	//~ End UObject
+
 protected:
 	bool bUpdating = false;
 
diff --git a/Engine/Plugins/Experimental/DynamicMaterial/Source/DynamicMaterial/Public/Components/RenderTargetRenderers/DMRenderTargetTextRenderer.h b/Engine/Plugins/Experimental/DynamicMaterial/Source/DynamicMaterial/Public/Components/RenderTargetRenderers/DMRenderTargetTextRenderer.h
--- a/Engine/Plugins/Experimental/DynamicMaterial/Source/DynamicMaterial/Public/Components/RenderTargetRenderers/DMRenderTargetTextRenderer.h
+++ b/Engine/Plugins/Experimental/DynamicMaterial/Source/DynamicMaterial/Public/Components/RenderTargetRenderers/DMRenderTargetTextRenderer.h
@@ -26,7 +26,7 @@ struct FDMTextLine
 	float Width = 0.f;
 };
 
-UCLASS(BlueprintType, ClassGroup = "Material Designer", meta = (DisplayName = "Material Designer Render Target Renderer"))
+UCLASS(BlueprintType, ClassGroup = "Material Designer", meta = (DisplayName = "Material Designer Render Target Text Renderer"))
 class DYNAMICMATERIAL_API UDMRenderTargetTextRenderer : public UDMRenderTargetRenderer
 {
 	GENERATED_BODY()
@@ -168,12 +168,11 @@ public:
 	UFUNCTION(BlueprintCallable, Category = "Material Designer")
 	void SetGlowOuterRadius(const FVector2D& InGlowOuterRadius);
 
-	//~ Begin UObject
-	virtual void PostLoad() override;
 #if WITH_EDITOR
+	//~ Begin UObject
 	virtual void PostEditChangeProperty(FPropertyChangedEvent& InPropertyChangedEvent) override;
-#endif
 	//~ End UObject
+#endif
 
 protected:
 	UPROPERTY(EditInstanceOnly, BlueprintReadWrite, Setter = SetFont, BlueprintSetter = SetFont, Category = "Material Designer",
diff --git a/Engine/Plugins/Experimental/DynamicMaterial/Source/DynamicMaterial/Public/Components/RenderTargetRenderers/DMRenderTargetWidgetRenderer.h b/Engine/Plugins/Experimental/DynamicMaterial/Source/DynamicMaterial/Public/Components/RenderTargetRenderers/DMRenderTargetWidgetRenderer.h
new file mode 100644
--- /dev/null
+++ b/Engine/Plugins/Experimental/DynamicMaterial/Source/DynamicMaterial/Public/Components/RenderTargetRenderers/DMRenderTargetWidgetRenderer.h
@@ -0,0 +1,46 @@
+// Copyright Epic Games, Inc. All Rights Reserved.
+
+#pragma once
+
+#include "Components/DMRenderTargetRenderer.h"
+#include "DMRenderTargetWidgetRenderer.generated.h"
+
+class FWidgetRenderer;
+class UWidget;
+
+UCLASS(BlueprintType, ClassGroup = "Material Designer", meta = (DisplayName = "Material Designer Render Target Widget Renderer"))
+class DYNAMICMATERIAL_API UDMRenderTargetWidgetRenderer : public UDMRenderTargetRenderer
+{
+	GENERATED_BODY()
+
+public:
+	UDMRenderTargetWidgetRenderer();
+
+	UFUNCTION(BlueprintPure, Category = "Material Designer")
+	TSubclassOf<UWidget> GetWidgetClass() const { return WidgetClass; }
+
+	UFUNCTION(BlueprintCallable, Category = "Material Designer")
+	void SetWidgetClass(TSubclassOf<UWidget> InWidgetClass);
+
+#if WITH_EDITOR
+	//~ Begin UObject
+	virtual void PostEditChangeProperty(FPropertyChangedEvent& InPropertyChangedEvent) override;
+	//~ End UObject
+#endif
+
+protected:
+	UPROPERTY(EditInstanceOnly, BlueprintReadWrite, Setter = SetWidgetClass, BlueprintSetter = SetWidgetClass, Category = "Material Designer",
+		meta = (NotKeyframeable, AllowPrivateAccess = "true"))
+	TSubclassOf<UWidget> WidgetClass;
+
+	UPROPERTY()
+	TObjectPtr<UWidget> WidgetInstance;
+
+	TSharedPtr<FWidgetRenderer> WidgetRenderer;
+
+	virtual void CreateWidgetInstance();
+
+	//~ Begin UDMRenderTargetRenderer
+	virtual void UpdateRenderTarget_Internal() override;
+	//~ End UDMRenderTargetRenderer
+};
diff --git a/Engine/Plugins/Experimental/DynamicMaterial/Source/DynamicMaterialEditor/Private/Menus/DMMaterialSlotLayerMenus.cpp b/Engine/Plugins/Experimental/DynamicMaterial/Source/DynamicMaterialEditor/Private/Menus/DMMaterialSlotLayerMenus.cpp
--- a/Engine/Plugins/Experimental/DynamicMaterial/Source/DynamicMaterialEditor/Private/Menus/DMMaterialSlotLayerMenus.cpp
+++ b/Engine/Plugins/Experimental/DynamicMaterial/Source/DynamicMaterialEditor/Private/Menus/DMMaterialSlotLayerMenus.cpp
@@ -10,6 +10,7 @@
 #include "Components/MaterialStageExpressions/DMMSETextureSample.h"
 #include "Components/MaterialStageExpressions/DMMSETextureSampleEdgeColor.h"
 #include "Components/RenderTargetRenderers/DMRenderTargetTextRenderer.h"
+#include "Components/RenderTargetRenderers/DMRenderTargetWidgetRenderer.h"
 #include "DMDefs.h"
 #include "DMMaterialSlotLayerAddEffectMenus.h"
 #include "DMValueDefinition.h"
@@ -564,6 +565,18 @@ void FDMMaterialSlotLayerMenus::AddAddLayerSection(UToolMenu* InMenu)
 		))
 	);
 
+	NewSection.AddMenuEntry(
+		NAME_None,
+		LOCTEXT("AddWidget", "Widget"),
+		LOCTEXT("AddWidgetTooltip", "Add a Material Stage based on a Widget Renderer."),
+		FSlateIcon(),
+		FUIAction(FExecuteAction::CreateSP(
+			EditorWidget->GetActiveSlotWidget().Get(),
+			&SDMSlot::AddNewLayer_Renderer,
+			TSubclassOf<UDMRenderTargetRenderer>(UDMRenderTargetWidgetRenderer::StaticClass())
+		))
+	);
+
 	if constexpr (UE::DynamicMaterialEditor::bAdvancedSlotsEnabled)
 	{
 		if (bHasValidSlot)
diff --git a/Engine/Plugins/Experimental/DynamicMaterial/Source/DynamicMaterialEditor/Private/Menus/DMMaterialStageSourceMenus.cpp b/Engine/Plugins/Experimental/DynamicMaterial/Source/DynamicMaterialEditor/Private/Menus/DMMaterialStageSourceMenus.cpp
--- a/Engine/Plugins/Experimental/DynamicMaterial/Source/DynamicMaterialEditor/Private/Menus/DMMaterialStageSourceMenus.cpp
+++ b/Engine/Plugins/Experimental/DynamicMaterial/Source/DynamicMaterialEditor/Private/Menus/DMMaterialStageSourceMenus.cpp
@@ -23,6 +23,7 @@
 #include "Components/MaterialStageInputs/DMMSISlot.h"
 #include "Components/MaterialStageInputs/DMMSIValue.h"
 #include "Components/RenderTargetRenderers/DMRenderTargetTextRenderer.h"
+#include "Components/RenderTargetRenderers/DMRenderTargetWidgetRenderer.h"
 #include "DMBlueprintFunctionLibrary.h"
 #include "DMDefs.h"
 #include "DMValueDefinition.h"
@@ -129,7 +130,7 @@ namespace UE::DynamicMaterialEditor::Private
 							}
 							else
 							{
-								ensure(false);
+								ensureMsgf(false, TEXT("Invalid stage type (%s)"), *StageSource->GetClass()->GetName());
 							}
 
 							if (TSharedPtr<SDMEditor> EditorWidget = MenuContext->GetEditorWidget())
@@ -231,7 +232,7 @@ namespace UE::DynamicMaterialEditor::Private
 						}
 						else
 						{
-							ensure(false);
+							ensureMsgf(false, TEXT("Invalid stage type (%s)"), *StageSource->GetClass()->GetName());
 						}
 
 						if (TSharedPtr<SDMEditor> EditorWidget = MenuContext->GetEditorWidget())
@@ -327,7 +328,7 @@ namespace UE::DynamicMaterialEditor::Private
 							}
 							else
 							{
-								ensure(false);
+								ensureMsgf(false, TEXT("Invalid stage type (%s)"), *StageSource->GetClass()->GetName());
 							}
 
 							if (TSharedPtr<SDMEditor> EditorWidget = MenuContext->GetEditorWidget())
@@ -452,7 +453,7 @@ namespace UE::DynamicMaterialEditor::Private
 							}
 							else
 							{
-								ensure(false);
+								ensureMsgf(false, TEXT("Invalid stage type (%s)"), *StageSource->GetClass()->GetName());
 							}
 
 							if (TSharedPtr<SDMEditor> EditorWidget = MenuContext->GetEditorWidget())
@@ -598,7 +599,7 @@ namespace UE::DynamicMaterialEditor::Private
 								}
 								else
 								{
-									ensure(false);
+									ensureMsgf(false, TEXT("Invalid stage type (%s)"), *StageSource->GetClass()->GetName());
 								}
 
 								if (TSharedPtr<SDMEditor> EditorWidget = MenuContext->GetEditorWidget())
@@ -701,7 +702,7 @@ namespace UE::DynamicMaterialEditor::Private
 							}
 							else
 							{
-								ensure(false);
+								ensureMsgf(false, TEXT("Invalid stage type (%s)"), *StageSource->GetClass()->GetName());
 							}
 
 							if (TSharedPtr<SDMEditor> EditorWidget = MenuContext->GetEditorWidget())
@@ -766,7 +767,7 @@ namespace UE::DynamicMaterialEditor::Private
 		}
 		else
 		{
-			ensure(false);
+			ensureMsgf(false, TEXT("Invalid stage type (%s)"), *StageSource->GetClass()->GetName());
 		}
 
 		if (TSharedPtr<SDMEditor> EditorWidget = InMenuContext->GetEditorWidget())
@@ -824,7 +825,7 @@ namespace UE::DynamicMaterialEditor::Private
 		}
 		else
 		{
-			ensure(false);
+			ensureMsgf(false, TEXT("Invalid stage type (%s)"), *StageSource->GetClass()->GetName());
 		}
 
 		if (TSharedPtr<SDMEditor> EditorWidget = InMenuContext->GetEditorWidget())
@@ -882,7 +883,7 @@ namespace UE::DynamicMaterialEditor::Private
 		}
 		else
 		{
-			ensure(false);
+			ensureMsgf(false, TEXT("Invalid stage type (%s)"), *StageSource->GetClass()->GetName());
 		}
 
 		if (TSharedPtr<SDMEditor> EditorWidget = InMenuContext->GetEditorWidget())
@@ -949,7 +950,7 @@ namespace UE::DynamicMaterialEditor::Private
 		}
 		else
 		{
-			ensure(false);
+			ensureMsgf(false, TEXT("Invalid stage type (%s)"), *StageSource->GetClass()->GetName());
 		}
 
 		if (UDMMaterialStage* MaskStage = Layer->GetStage(EDMMaterialLayerStage::Mask))
@@ -1038,7 +1039,7 @@ namespace UE::DynamicMaterialEditor::Private
 		}
 		else
 		{
-			ensure(false);
+			ensureMsgf(false, TEXT("Invalid stage type (%s)"), *StageSource->GetClass()->GetName());
 		}
 
 		if (TSharedPtr<SDMEditor> EditorWidget = InMenuContext->GetEditorWidget())
@@ -1123,7 +1124,7 @@ namespace UE::DynamicMaterialEditor::Private
 		}
 		else
 		{
-			ensure(false);
+			ensureMsgf(false, TEXT("Invalid stage type (%s)"), *StageSource->GetClass()->GetName());
 		}
 
 		if (TSharedPtr<SDMEditor> EditorWidget = InMenuContext->GetEditorWidget())
@@ -1176,7 +1177,60 @@ namespace UE::DynamicMaterialEditor::Private
 		}
 		else
 		{
-			ensure(false);
+			ensureMsgf(false, TEXT("Invalid stage type (%s)"), *StageSource->GetClass()->GetName());
+		}
+
+		if (TSharedPtr<SDMEditor> EditorWidget = InMenuContext->GetEditorWidget())
+		{
+			EditorWidget->InvalidateComponentEditWidget();
+		}
+	}
+
+	void ChangeSourceToWidgetFromContext(UDMMenuContext* InMenuContext)
+	{
+		if (!IsValid(InMenuContext))
+		{
+			return;
+		}
+
+		UDMMaterialStageSource* const StageSource = InMenuContext->GetStageSource();
+
+		if (!StageSource)
+		{
+			return;
+		}
+
+		UDMMaterialStage* const Stage = InMenuContext->GetStage();
+
+		if (!Stage)
+		{
+			return;
+		}
+
+		UDMMaterialLayerObject* const Layer = Stage->GetLayer();
+
+		if (!Layer)
+		{
+			return;
+		}
+
+		if (StageSource->IsA<UDMMaterialStageBlend>())
+		{
+			FScopedTransaction Transaction(LOCTEXT("SetStageInputBase", "Set Material Designer Base Source"));
+			Stage->Modify();
+
+			UDMBlueprintFunctionLibrary::SetStageInputToRenderer(Stage, UDMRenderTargetWidgetRenderer::StaticClass(), UDMMaterialStageBlend::InputB);
+		}
+		else if (StageSource->IsA<UDMMaterialStageThroughputLayerBlend>())
+		{
+			FScopedTransaction Transaction(LOCTEXT("SetStageInputMask", "Set Material Designer Mask Source"));
+			Stage->Modify();
+
+			UDMBlueprintFunctionLibrary::SetStageInputToRenderer(Stage, UDMRenderTargetWidgetRenderer::StaticClass(), UDMMaterialStageThroughputLayerBlend::InputMaskSource);
+		}
+		else
+		{
+			ensureMsgf(false, TEXT("Invalid stage type (%s)"), *StageSource->GetClass()->GetName());
 		}
 
 		if (TSharedPtr<SDMEditor> EditorWidget = InMenuContext->GetEditorWidget())
@@ -1311,6 +1365,18 @@ namespace UE::DynamicMaterialEditor::Private
 			)
 		);
 
+		InSection.AddMenuEntry("Widget",
+			LOCTEXT("ChangeSourceWidget", "Widget"),
+			LOCTEXT("ChangeSourceWidgetTooltip", "Change the source of this stage to a Widget Renderer."),
+			FSlateIcon(),
+			FUIAction(
+				FExecuteAction::CreateStatic(
+					&ChangeSourceToWidgetFromContext,
+					MenuContext
+				)
+			)
+		);
+
 		if constexpr (UE::DynamicMaterialEditor::bGlobalValuesEnabled)
 		{
 			const TArray<UDMMaterialValue*>& Values = MaterialModel->GetValues();
