Material Designer
- Effect preset system

diff --git a/Engine/Plugins/Experimental/DynamicMaterial/Source/DynamicMaterial/DynamicMaterial.Build.cs b/Engine/Plugins/Experimental/DynamicMaterial/Source/DynamicMaterial/DynamicMaterial.Build.cs
--- a/Engine/Plugins/Experimental/DynamicMaterial/Source/DynamicMaterial/DynamicMaterial.Build.cs
+++ b/Engine/Plugins/Experimental/DynamicMaterial/Source/DynamicMaterial/DynamicMaterial.Build.cs
@@ -23,6 +23,8 @@ public class DynamicMaterial : ModuleRules
 			new string[]
 			{
 				"Engine",
+				"Json",
+				"JsonUtilities",
 				"RenderCore"
 			}
 		);
diff --git a/Engine/Plugins/Experimental/DynamicMaterial/Source/DynamicMaterial/Private/Components/MaterialValues/DMMaterialValueBool.cpp b/Engine/Plugins/Experimental/DynamicMaterial/Source/DynamicMaterial/Private/Components/MaterialValues/DMMaterialValueBool.cpp
--- a/Engine/Plugins/Experimental/DynamicMaterial/Source/DynamicMaterial/Private/Components/MaterialValues/DMMaterialValueBool.cpp
+++ b/Engine/Plugins/Experimental/DynamicMaterial/Source/DynamicMaterial/Private/Components/MaterialValues/DMMaterialValueBool.cpp
@@ -73,6 +73,24 @@ void UDMMaterialValueBool::ResetDefaultValue()
 	bDefaultValue = false;
 }
 
+TSharedPtr<FJsonValue> UDMMaterialValueBool::JsonSerialize() const
+{
+	return FDMJsonUtils::Serialize(Value);
+}
+
+bool UDMMaterialValueBool::JsonDeserialize(const TSharedPtr<FJsonValue>& InJsonValue)
+{
+	bool bValueJson;
+
+	if (FDMJsonUtils::Deserialize(InJsonValue, bValueJson))
+	{
+		SetValue(bValueJson);
+		return true;
+	}
+
+	return false;
+}
+
 void UDMMaterialValueBool::SetDefaultValue(bool bInDefaultValue)
 {
 	bDefaultValue = bInDefaultValue;
diff --git a/Engine/Plugins/Experimental/DynamicMaterial/Source/DynamicMaterial/Private/Components/MaterialValues/DMMaterialValueColorAtlas.cpp b/Engine/Plugins/Experimental/DynamicMaterial/Source/DynamicMaterial/Private/Components/MaterialValues/DMMaterialValueColorAtlas.cpp
--- a/Engine/Plugins/Experimental/DynamicMaterial/Source/DynamicMaterial/Private/Components/MaterialValues/DMMaterialValueColorAtlas.cpp
+++ b/Engine/Plugins/Experimental/DynamicMaterial/Source/DynamicMaterial/Private/Components/MaterialValues/DMMaterialValueColorAtlas.cpp
@@ -5,10 +5,11 @@
 #include "Materials/MaterialInstanceDynamic.h"
 
 #if WITH_EDITOR
+#include "Curves/CurveLinearColor.h"
+#include "Curves/CurveLinearColorAtlas.h"
 #include "Materials/Material.h"
 #include "Materials/MaterialExpression.h"
 #include "Materials/MaterialExpressionAppendVector.h"
-#include "Materials/MaterialExpressionComponentMask.h"
 #include "Materials/MaterialExpressionCurveAtlasRowParameter.h"
 #include "Model/IDMMaterialBuildStateInterface.h"
 #include "Model/IDMMaterialBuildUtilsInterface.h"
@@ -135,6 +136,73 @@ void UDMMaterialValueColorAtlas::ResetDefaultValue()
 	DefaultValue = 0.f;
 }
 
+TSharedPtr<FJsonValue> UDMMaterialValueColorAtlas::JsonSerialize() const
+{
+	return FDMJsonUtils::Serialize({
+		{GET_MEMBER_NAME_STRING_CHECKED(ThisClass, Value), FDMJsonUtils::Serialize(Value)},
+		{GET_MEMBER_NAME_STRING_CHECKED(ThisClass, Atlas), FDMJsonUtils::Serialize(Atlas)},
+		{GET_MEMBER_NAME_STRING_CHECKED(ThisClass, Curve), FDMJsonUtils::Serialize(Curve)}
+	});
+}
+
+bool UDMMaterialValueColorAtlas::JsonDeserialize(const TSharedPtr<FJsonValue>& InJsonValue)
+{
+	TMap<FString, TSharedPtr<FJsonValue>> Data;
+
+	if (!FDMJsonUtils::Deserialize(InJsonValue, Data))
+	{
+		return false;
+	}
+
+	bool bSuccess = false;
+	EDMUpdateType UpdateType = EDMUpdateType::Value;
+
+	if (const TSharedPtr<FJsonValue>* JsonValue = Data.Find(GET_MEMBER_NAME_STRING_CHECKED(ThisClass, Atlas)))
+	{
+		UCurveLinearColorAtlas* AtlasJson = nullptr;
+
+		if (FDMJsonUtils::Deserialize(*JsonValue, AtlasJson))
+		{
+			const FDMUpdateGuard Guard;
+			SetAtlas(AtlasJson);
+			UpdateType = EDMUpdateType::Structure;
+			bSuccess = true;
+		}
+	}
+
+	if (const TSharedPtr<FJsonValue>* JsonValue = Data.Find(GET_MEMBER_NAME_STRING_CHECKED(ThisClass, Curve)))
+	{
+		UCurveLinearColor* CurveJson = nullptr;
+
+		if (FDMJsonUtils::Deserialize(*JsonValue, CurveJson))
+		{
+			const FDMUpdateGuard Guard;
+			SetCurve(CurveJson);
+			UpdateType = EDMUpdateType::Structure;
+			bSuccess = true;
+		}
+	}
+
+	if (const TSharedPtr<FJsonValue>* JsonValue = Data.Find(GET_MEMBER_NAME_STRING_CHECKED(ThisClass, Value)))
+	{
+		float ValueJson;
+
+		if (FDMJsonUtils::Deserialize(*JsonValue, ValueJson))
+		{
+			const FDMUpdateGuard Guard;
+			SetValue(ValueJson);
+			bSuccess = true;
+		}
+	}
+
+	if (bSuccess)
+	{
+		Update(UpdateType);
+	}
+
+	return bSuccess;
+}
+
 void UDMMaterialValueColorAtlas::SetDefaultValue(float InDefaultValue)
 {
 	DefaultValue = InDefaultValue;
diff --git a/Engine/Plugins/Experimental/DynamicMaterial/Source/DynamicMaterial/Private/Components/MaterialValues/DMMaterialValueFloat1.cpp b/Engine/Plugins/Experimental/DynamicMaterial/Source/DynamicMaterial/Private/Components/MaterialValues/DMMaterialValueFloat1.cpp
--- a/Engine/Plugins/Experimental/DynamicMaterial/Source/DynamicMaterial/Private/Components/MaterialValues/DMMaterialValueFloat1.cpp
+++ b/Engine/Plugins/Experimental/DynamicMaterial/Source/DynamicMaterial/Private/Components/MaterialValues/DMMaterialValueFloat1.cpp
@@ -1,6 +1,7 @@
 // Copyright Epic Games, Inc. All Rights Reserved.
  
 #include "Components/MaterialValues/DMMaterialValueFloat1.h"
+
 #include "Materials/MaterialInstanceDynamic.h"
 
 #if WITH_EDITOR
@@ -54,6 +55,24 @@ void UDMMaterialValueFloat1::ResetDefaultValue()
 	DefaultValue = 0.f;
 }
 
+TSharedPtr<FJsonValue> UDMMaterialValueFloat1::JsonSerialize() const
+{
+	return FDMJsonUtils::Serialize(Value);
+}
+
+bool UDMMaterialValueFloat1::JsonDeserialize(const TSharedPtr<FJsonValue>& InJsonValue)
+{
+	float ValueJson;
+
+	if (FDMJsonUtils::Deserialize(InJsonValue, ValueJson))
+	{
+		SetValue(ValueJson);
+		return true;
+	}
+
+	return false;
+}
+
 void UDMMaterialValueFloat1::SetDefaultValue(float InDefaultValue)
 {
 	DefaultValue = InDefaultValue;
diff --git a/Engine/Plugins/Experimental/DynamicMaterial/Source/DynamicMaterial/Private/Components/MaterialValues/DMMaterialValueFloat2.cpp b/Engine/Plugins/Experimental/DynamicMaterial/Source/DynamicMaterial/Private/Components/MaterialValues/DMMaterialValueFloat2.cpp
--- a/Engine/Plugins/Experimental/DynamicMaterial/Source/DynamicMaterial/Private/Components/MaterialValues/DMMaterialValueFloat2.cpp
+++ b/Engine/Plugins/Experimental/DynamicMaterial/Source/DynamicMaterial/Private/Components/MaterialValues/DMMaterialValueFloat2.cpp
@@ -64,6 +64,24 @@ void UDMMaterialValueFloat2::ResetDefaultValue()
 	DefaultValue = FVector2D::ZeroVector;
 }
 
+TSharedPtr<FJsonValue> UDMMaterialValueFloat2::JsonSerialize() const
+{
+	return FDMJsonUtils::Serialize(Value);
+}
+
+bool UDMMaterialValueFloat2::JsonDeserialize(const TSharedPtr<FJsonValue>& InJsonValue)
+{
+	FVector2D ValueJson;
+
+	if (FDMJsonUtils::Deserialize(InJsonValue, ValueJson))
+	{
+		SetValue(ValueJson);
+		return true;
+	}
+
+	return false;
+}
+
 void UDMMaterialValueFloat2::SetDefaultValue(FVector2D InDefaultValue)
 {
 	DefaultValue = InDefaultValue;
diff --git a/Engine/Plugins/Experimental/DynamicMaterial/Source/DynamicMaterial/Private/Components/MaterialValues/DMMaterialValueFloat3RGB.cpp b/Engine/Plugins/Experimental/DynamicMaterial/Source/DynamicMaterial/Private/Components/MaterialValues/DMMaterialValueFloat3RGB.cpp
--- a/Engine/Plugins/Experimental/DynamicMaterial/Source/DynamicMaterial/Private/Components/MaterialValues/DMMaterialValueFloat3RGB.cpp
+++ b/Engine/Plugins/Experimental/DynamicMaterial/Source/DynamicMaterial/Private/Components/MaterialValues/DMMaterialValueFloat3RGB.cpp
@@ -61,6 +61,24 @@ void UDMMaterialValueFloat3RGB::ResetDefaultValue()
 	DefaultValue = FLinearColor(0, 0, 0, 1);
 }
 
+TSharedPtr<FJsonValue> UDMMaterialValueFloat3RGB::JsonSerialize() const
+{
+	return FDMJsonUtils::Serialize(Value);
+}
+
+bool UDMMaterialValueFloat3RGB::JsonDeserialize(const TSharedPtr<FJsonValue>& InJsonValue)
+{
+	FLinearColor ValueJson;
+
+	if (FDMJsonUtils::Deserialize(InJsonValue, ValueJson))
+	{
+		SetValue(ValueJson);
+		return true;
+	}
+
+	return false;
+}
+
 void UDMMaterialValueFloat3RGB::SetDefaultValue(FLinearColor InDefaultValue)
 {
 	DefaultValue = InDefaultValue;
diff --git a/Engine/Plugins/Experimental/DynamicMaterial/Source/DynamicMaterial/Private/Components/MaterialValues/DMMaterialValueFloat3RPY.cpp b/Engine/Plugins/Experimental/DynamicMaterial/Source/DynamicMaterial/Private/Components/MaterialValues/DMMaterialValueFloat3RPY.cpp
--- a/Engine/Plugins/Experimental/DynamicMaterial/Source/DynamicMaterial/Private/Components/MaterialValues/DMMaterialValueFloat3RPY.cpp
+++ b/Engine/Plugins/Experimental/DynamicMaterial/Source/DynamicMaterial/Private/Components/MaterialValues/DMMaterialValueFloat3RPY.cpp
@@ -61,6 +61,24 @@ void UDMMaterialValueFloat3RPY::ResetDefaultValue()
 	DefaultValue = FRotator::ZeroRotator;
 }
 
+TSharedPtr<FJsonValue> UDMMaterialValueFloat3RPY::JsonSerialize() const
+{
+	return FDMJsonUtils::Serialize(Value);
+}
+
+bool UDMMaterialValueFloat3RPY::JsonDeserialize(const TSharedPtr<FJsonValue>& InJsonValue)
+{
+	FRotator ValueJson;
+
+	if (FDMJsonUtils::Deserialize(InJsonValue, ValueJson))
+	{
+		SetValue(ValueJson);
+		return true;
+	}
+
+	return false;
+}
+
 void UDMMaterialValueFloat3RPY::SetDefaultValue(FRotator InDefaultValue)
 {
 	DefaultValue = InDefaultValue;
diff --git a/Engine/Plugins/Experimental/DynamicMaterial/Source/DynamicMaterial/Private/Components/MaterialValues/DMMaterialValueFloat3XYZ.cpp b/Engine/Plugins/Experimental/DynamicMaterial/Source/DynamicMaterial/Private/Components/MaterialValues/DMMaterialValueFloat3XYZ.cpp
--- a/Engine/Plugins/Experimental/DynamicMaterial/Source/DynamicMaterial/Private/Components/MaterialValues/DMMaterialValueFloat3XYZ.cpp
+++ b/Engine/Plugins/Experimental/DynamicMaterial/Source/DynamicMaterial/Private/Components/MaterialValues/DMMaterialValueFloat3XYZ.cpp
@@ -61,6 +61,24 @@ void UDMMaterialValueFloat3XYZ::ResetDefaultValue()
 	DefaultValue = FVector::ZeroVector;
 }
 
+TSharedPtr<FJsonValue> UDMMaterialValueFloat3XYZ::JsonSerialize() const
+{
+	return FDMJsonUtils::Serialize(Value);
+}
+
+bool UDMMaterialValueFloat3XYZ::JsonDeserialize(const TSharedPtr<FJsonValue>& InJsonValue)
+{
+	FVector ValueJson;
+
+	if (FDMJsonUtils::Deserialize(InJsonValue, ValueJson))
+	{
+		SetValue(ValueJson);
+		return true;
+	}
+
+	return false;
+}
+
 void UDMMaterialValueFloat3XYZ::SetDefaultValue(FVector InDefaultValue)
 {
 	DefaultValue = InDefaultValue;
diff --git a/Engine/Plugins/Experimental/DynamicMaterial/Source/DynamicMaterial/Private/Components/MaterialValues/DMMaterialValueFloat4.cpp b/Engine/Plugins/Experimental/DynamicMaterial/Source/DynamicMaterial/Private/Components/MaterialValues/DMMaterialValueFloat4.cpp
--- a/Engine/Plugins/Experimental/DynamicMaterial/Source/DynamicMaterial/Private/Components/MaterialValues/DMMaterialValueFloat4.cpp
+++ b/Engine/Plugins/Experimental/DynamicMaterial/Source/DynamicMaterial/Private/Components/MaterialValues/DMMaterialValueFloat4.cpp
@@ -62,6 +62,24 @@ void UDMMaterialValueFloat4::ResetDefaultValue()
 	DefaultValue = FLinearColor(0, 0, 0, 1);
 }
 
+TSharedPtr<FJsonValue> UDMMaterialValueFloat4::JsonSerialize() const
+{
+	return FDMJsonUtils::Serialize(Value);
+}
+
+bool UDMMaterialValueFloat4::JsonDeserialize(const TSharedPtr<FJsonValue>& InJsonValue)
+{
+	FLinearColor ValueJson;
+
+	if (FDMJsonUtils::Deserialize(InJsonValue, ValueJson))
+	{
+		SetValue(ValueJson);
+		return true;
+	}
+
+	return false;
+}
+
 void UDMMaterialValueFloat4::SetDefaultValue(FLinearColor InDefaultValue)
 {
 	DefaultValue = InDefaultValue;
diff --git a/Engine/Plugins/Experimental/DynamicMaterial/Source/DynamicMaterial/Private/Components/MaterialValues/DMMaterialValueRenderTarget.cpp b/Engine/Plugins/Experimental/DynamicMaterial/Source/DynamicMaterial/Private/Components/MaterialValues/DMMaterialValueRenderTarget.cpp
--- a/Engine/Plugins/Experimental/DynamicMaterial/Source/DynamicMaterial/Private/Components/MaterialValues/DMMaterialValueRenderTarget.cpp
+++ b/Engine/Plugins/Experimental/DynamicMaterial/Source/DynamicMaterial/Private/Components/MaterialValues/DMMaterialValueRenderTarget.cpp
@@ -151,6 +151,83 @@ void UDMMaterialValueRenderTarget::Update(EDMUpdateType InUpdateType)
 }
 
 #if WITH_EDITOR
+TSharedPtr<FJsonValue> UDMMaterialValueRenderTarget::JsonSerialize() const
+{
+	return FDMJsonUtils::Serialize({
+		{GET_MEMBER_NAME_STRING_CHECKED(ThisClass, TextureSize), FDMJsonUtils::Serialize(TextureSize)},
+		{GET_MEMBER_NAME_STRING_CHECKED(ThisClass, TextureFormat), FDMJsonUtils::Serialize<ETextureRenderTargetFormat>(TextureFormat)},
+		{GET_MEMBER_NAME_STRING_CHECKED(ThisClass, ClearColor), FDMJsonUtils::Serialize(ClearColor)},
+		{GET_MEMBER_NAME_STRING_CHECKED(ThisClass, Renderer), FDMJsonUtils::Serialize(Renderer.Get())}
+	});
+}
+
+bool UDMMaterialValueRenderTarget::JsonDeserialize(const TSharedPtr<FJsonValue>& InJsonValue)
+{
+	TMap<FString, TSharedPtr<FJsonValue>> Data;
+
+	if (!FDMJsonUtils::Deserialize(InJsonValue, Data))
+	{
+		return false;
+	}
+
+	bool bSuccess = false;
+
+	if (const TSharedPtr<FJsonValue>* JsonValue = Data.Find(GET_MEMBER_NAME_STRING_CHECKED(ThisClass, TextureSize)))
+	{
+		FIntPoint TextureSizeJson = FIntPoint::ZeroValue;
+
+		if (FDMJsonUtils::Deserialize(*JsonValue, TextureSizeJson))
+		{
+			const FDMUpdateGuard Guard;
+			SetTextureSize(TextureSizeJson);
+			bSuccess = true;
+		}
+	}
+
+	if (const TSharedPtr<FJsonValue>* JsonValue = Data.Find(GET_MEMBER_NAME_STRING_CHECKED(ThisClass, TextureFormat)))
+	{
+		ETextureRenderTargetFormat TextureFormatJson = RTF_RGBA16f;
+
+		if (FDMJsonUtils::Deserialize(*JsonValue, TextureFormatJson))
+		{
+			const FDMUpdateGuard Guard;
+			SetTextureFormat(TextureFormatJson);
+			bSuccess = true;
+		}
+	}
+
+	if (const TSharedPtr<FJsonValue>* JsonClearColor = Data.Find(GET_MEMBER_NAME_STRING_CHECKED(ThisClass, ClearColor)))
+	{
+		FLinearColor ClearColorJson = FLinearColor::Black;
+
+		if (FDMJsonUtils::Deserialize(*JsonClearColor, ClearColorJson))
+		{
+			const FDMUpdateGuard Guard;
+			SetClearColor(ClearColorJson);
+			bSuccess = true;
+		}
+	}
+
+	if (const TSharedPtr<FJsonValue>* JsonRenderer = Data.Find(GET_MEMBER_NAME_STRING_CHECKED(ThisClass, Renderer)))
+	{
+		UDMRenderTargetRenderer* RendererJson;
+
+		if (FDMJsonUtils::Deserialize(*JsonRenderer, RendererJson, this))
+		{
+			const FDMUpdateGuard Guard;
+			SetRenderer(RendererJson);
+			bSuccess = true;
+		}
+	}
+
+	if (bSuccess)
+	{
+		Update(EDMUpdateType::Value);
+	}
+
+	return bSuccess;
+}
+
 void UDMMaterialValueRenderTarget::PostEditorDuplicate(UDynamicMaterialModel* InMaterialModel, UDMMaterialComponent* InParent)
 {
 	Super::PostEditorDuplicate(InMaterialModel, InParent);
diff --git a/Engine/Plugins/Experimental/DynamicMaterial/Source/DynamicMaterial/Private/Components/MaterialValues/DMMaterialValueTexture.cpp b/Engine/Plugins/Experimental/DynamicMaterial/Source/DynamicMaterial/Private/Components/MaterialValues/DMMaterialValueTexture.cpp
--- a/Engine/Plugins/Experimental/DynamicMaterial/Source/DynamicMaterial/Private/Components/MaterialValues/DMMaterialValueTexture.cpp
+++ b/Engine/Plugins/Experimental/DynamicMaterial/Source/DynamicMaterial/Private/Components/MaterialValues/DMMaterialValueTexture.cpp
@@ -99,16 +99,16 @@ UDMMaterialValueTexture* UDMMaterialValueTexture::CreateMaterialValueTexture(UOb
 	return TextureValue;
 }
  
-void UDMMaterialValueTexture::PreEditChange(FProperty* PropertyAboutToChange)
+void UDMMaterialValueTexture::PreEditChange(FProperty* InPropertyAboutToChange)
 {
-	Super::PreEditChange(PropertyAboutToChange);
+	Super::PreEditChange(InPropertyAboutToChange);
  
 	if (!IsComponentValid())
 	{
 		return;
 	}
 
-	if (PropertyAboutToChange->GetFName() == ValueName)
+	if (InPropertyAboutToChange->GetFName() == ValueName)
 	{
 		OldValue = GetValue();
 	}
@@ -119,17 +119,17 @@ bool UDMMaterialValueTexture::IsDefaultValue() const
 	return Value == DefaultValue;
 }
 
-void UDMMaterialValueTexture::PostEditChangeProperty(FPropertyChangedEvent& PropertyChangedEvent)
+void UDMMaterialValueTexture::PostEditChangeProperty(FPropertyChangedEvent& InPropertyChangedEvent)
 {
 	// Skip parent class because we need to do extra logic.
-	Super::Super::PostEditChangeProperty(PropertyChangedEvent);
+	Super::Super::PostEditChangeProperty(InPropertyChangedEvent);
 
 	if (!IsComponentValid())
 	{
 		return;
 	}
 
-	FName MemberPropertyName = PropertyChangedEvent.GetMemberPropertyName();
+	FName MemberPropertyName = InPropertyChangedEvent.GetMemberPropertyName();
 
 	if (MemberPropertyName == NAME_None)
 	{
@@ -179,6 +179,24 @@ void UDMMaterialValueTexture::ResetDefaultValue()
 	}
 }
 
+TSharedPtr<FJsonValue> UDMMaterialValueTexture::JsonSerialize() const
+{
+	return FDMJsonUtils::Serialize(Value);
+}
+
+bool UDMMaterialValueTexture::JsonDeserialize(const TSharedPtr<FJsonValue>& InJsonValue)
+{
+	UTexture* ValueJson;
+
+	if (FDMJsonUtils::Deserialize(InJsonValue, ValueJson))
+	{
+		SetValue(ValueJson);
+		return true;
+	}
+
+	return false;
+}
+
 void UDMMaterialValueTexture::SetDefaultValue(UTexture* InDefaultValue)
 {
 	DefaultValue = InDefaultValue;
diff --git a/Engine/Plugins/Experimental/DynamicMaterial/Source/DynamicMaterial/Private/Components/RenderTargetRenderers/DMRenderTargetTextRenderer.cpp b/Engine/Plugins/Experimental/DynamicMaterial/Source/DynamicMaterial/Private/Components/RenderTargetRenderers/DMRenderTargetTextRenderer.cpp
--- a/Engine/Plugins/Experimental/DynamicMaterial/Source/DynamicMaterial/Private/Components/RenderTargetRenderers/DMRenderTargetTextRenderer.cpp
+++ b/Engine/Plugins/Experimental/DynamicMaterial/Source/DynamicMaterial/Private/Components/RenderTargetRenderers/DMRenderTargetTextRenderer.cpp
@@ -9,6 +9,10 @@
 #include "Engine/TextureRenderTarget2D.h"
 #include "TextureResource.h"
 
+#if WITH_EDITOR
+#include "Dom/JsonValue.h"
+#endif
+
 #define LOCTEXT_NAMESPACE "DMRenderTargetTextRenderer"
 
 namespace UE::DynamicMaterial::Private
@@ -58,6 +62,16 @@ UDMRenderTargetTextRenderer::UDMRenderTargetTextRenderer()
 }
 
 #if WITH_EDITOR
+TSharedPtr<FJsonValue> UDMRenderTargetTextRenderer::JsonSerialize() const
+{
+	return MakeShared<FJsonValueNull>();
+}
+
+bool UDMRenderTargetTextRenderer::JsonDeserialize(const TSharedPtr<FJsonValue>& InJsonValue)
+{
+	return true;
+}
+
 void UDMRenderTargetTextRenderer::PostEditChangeProperty(FPropertyChangedEvent& InPropertyChangedEvent)
 {
 	Super::PostEditChangeProperty(InPropertyChangedEvent);
diff --git a/Engine/Plugins/Experimental/DynamicMaterial/Source/DynamicMaterial/Private/Components/RenderTargetRenderers/DMRenderTargetWidgetRenderer.cpp b/Engine/Plugins/Experimental/DynamicMaterial/Source/DynamicMaterial/Private/Components/RenderTargetRenderers/DMRenderTargetWidgetRenderer.cpp
--- a/Engine/Plugins/Experimental/DynamicMaterial/Source/DynamicMaterial/Private/Components/RenderTargetRenderers/DMRenderTargetWidgetRenderer.cpp
+++ b/Engine/Plugins/Experimental/DynamicMaterial/Source/DynamicMaterial/Private/Components/RenderTargetRenderers/DMRenderTargetWidgetRenderer.cpp
@@ -1,8 +1,8 @@
 // Copyright Epic Games, Inc. All Rights Reserved.
 
 #include "Components/RenderTargetRenderers/DMRenderTargetWidgetRenderer.h"
-#include "Components/Widget.h"
 #include "Components/MaterialValues/DMMaterialValueRenderTarget.h"
+#include "Components/Widget.h"
 #include "Engine/Engine.h"
 #include "Engine/TextureRenderTarget2D.h"
 #include "Slate/WidgetRenderer.h"
@@ -33,6 +33,24 @@ void UDMRenderTargetWidgetRenderer::SetWidgetClass(TSubclassOf<UWidget> InWidget
 }
 
 #if WITH_EDITOR
+TSharedPtr<FJsonValue> UDMRenderTargetWidgetRenderer::JsonSerialize() const
+{
+	return FDMJsonUtils::Serialize(WidgetClass.Get());
+}
+
+bool UDMRenderTargetWidgetRenderer::JsonDeserialize(const TSharedPtr<FJsonValue>& InJsonValue)
+{
+	TSubclassOf<UWidget> WidgetClassJson;
+
+	if (FDMJsonUtils::Deserialize(InJsonValue, WidgetClassJson))
+	{
+		SetWidgetClass(WidgetClassJson);
+		return true;
+	}
+
+	return false;
+}
+
 void UDMRenderTargetWidgetRenderer::PostEditChangeProperty(FPropertyChangedEvent& InPropertyChangedEvent)
 {
 	Super::PostEditChangeProperty(InPropertyChangedEvent);
diff --git a/Engine/Plugins/Experimental/DynamicMaterial/Source/DynamicMaterial/Private/Utils/DMJsonUtils.cpp b/Engine/Plugins/Experimental/DynamicMaterial/Source/DynamicMaterial/Private/Utils/DMJsonUtils.cpp
new file mode 100644
--- /dev/null
+++ b/Engine/Plugins/Experimental/DynamicMaterial/Source/DynamicMaterial/Private/Utils/DMJsonUtils.cpp
@@ -0,0 +1,390 @@
+// Copyright Epic Games, Inc. All Rights Reserved.
+
+#if WITH_EDITOR
+
+#include "Utils/DMJsonUtils.h"
+
+#include "DMDefs.h"
+#include "Dom/JsonValue.h"
+#include "DynamicMaterialModule.h"
+#include "JsonObjectConverter.h"
+
+namespace UE::DynamicMaterial::Private
+{
+	const FString JsonKey_Class = TEXT("Class");
+	const FString JsonKey_Data = TEXT("Data");
+}
+
+TSharedPtr<FJsonValue> FDMJsonUtils::SerializeNumber(double InNumber)
+{
+	return MakeShared<FJsonValueNumber>(InNumber);
+}
+
+TSharedPtr<FJsonValue> FDMJsonUtils::Serialize(bool bInValue)
+{
+	return MakeShared<FJsonValueBoolean>(bInValue);
+}
+
+TSharedPtr<FJsonValue> FDMJsonUtils::Serialize(const FString& InString)
+{
+	return MakeShared<FJsonValueString>(InString);
+}
+
+TSharedPtr<FJsonValue> FDMJsonUtils::Serialize(const FText& InText)
+{
+	return MakeShared<FJsonValueString>(InText.ToString());
+}
+
+TSharedPtr<FJsonValue> FDMJsonUtils::Serialize(const FName& InName)
+{
+	return MakeShared<FJsonValueString>(InName.GetPlainNameString());
+}
+
+TSharedPtr<FJsonValue> FDMJsonUtils::Serialize(const UClass* InClass)
+{
+	if (IsValid(InClass))
+	{
+		return Serialize(InClass->GetClassPathName().ToString());
+	}
+
+	return MakeShared<FJsonValueNull>();
+}
+
+TSharedPtr<FJsonValue> FDMJsonUtils::Serialize(const UScriptStruct* InScriptStruct, const void* InData)
+{
+	if (!InData)
+	{
+		return MakeShared<FJsonValueNull>();
+	}
+
+	if (!IsValid(InScriptStruct))
+	{
+		UE_LOG(LogDynamicMaterial, Error, TEXT("Invalid script struct."));
+		return nullptr;
+	}
+
+	TSharedRef<FJsonObject> JsonStruct = MakeShared<FJsonObject>();
+
+	if (FJsonObjectConverter::UStructToJsonObject(InScriptStruct, InData, JsonStruct))
+	{
+		return MakeShared<FJsonValueObject>(JsonStruct);
+	}
+
+	UE_LOG(LogDynamicMaterial, Error,TEXT("Failed to convert struct to json. [%s]"), *InScriptStruct->GetName());
+	return nullptr;
+}
+
+TSharedPtr<FJsonValue> FDMJsonUtils::Serialize(const UObject* InObject)
+{
+	if (!IsValid(InObject))
+	{
+		return MakeShared<FJsonValueNull>();
+	}
+
+	if (InObject->IsAsset())
+	{
+		return Serialize(InObject->GetPathName());
+	}
+
+	using namespace UE::DynamicMaterial::Private;
+
+	if (InObject->Implements<UDMJsonSerializable>())
+	{
+		TSharedPtr<FJsonValue> JsonObject = Cast<IDMJsonSerializable>(InObject)->JsonSerialize();
+
+		if (JsonObject.IsValid())
+		{
+			return Serialize({
+				{JsonKey_Class, Serialize(InObject->GetClass())},
+				{JsonKey_Data, JsonObject}
+			});
+		}
+	}
+	else
+	{
+		TSharedRef<FJsonObject> JsonObject = MakeShared<FJsonObject>();
+
+		if (FJsonObjectConverter::UStructToJsonObject(InObject->GetClass(), InObject, JsonObject))
+		{
+			return Serialize({
+				{JsonKey_Class, Serialize(InObject->GetClass())},
+				{JsonKey_Data, MakeShared<FJsonValueObject>(JsonObject)}
+			});
+		}
+	}
+
+	UE_LOG(LogDynamicMaterial, Error, TEXT("Failed to convert object to json. [%s] [%s]"), 
+		*InObject->GetClass()->GetName(), *InObject->GetPathName());
+
+	return nullptr;
+}
+
+TSharedPtr<FJsonValue> FDMJsonUtils::Serialize(const FObjectPtr& InObject)
+{
+	return Serialize(InObject.Get());
+}
+
+TSharedPtr<FJsonValue> FDMJsonUtils::Serialize(const TMap<FString, TSharedPtr<FJsonValue>>& InMap)
+{
+	return MakeShared<FJsonValueObject>(MakeShared<FJsonObject>(InMap));
+}
+
+bool FDMJsonUtils::DeserializeNumber(const TSharedPtr<FJsonValue>& InJsonValue, double& OutNumber)
+{
+	if (InJsonValue->TryGetNumber(OutNumber))
+	{
+		return true;
+	}
+
+	UE_LOG(LogDynamicMaterial, Error, TEXT("Failed to read number from json."));
+	return false;
+}
+
+bool FDMJsonUtils::Deserialize(const TSharedPtr<FJsonValue>& InJsonValue, bool& bOutValue)
+{
+	if (InJsonValue->TryGetBool(bOutValue))
+	{
+		return true;
+	}
+
+	UE_LOG(LogDynamicMaterial, Error,TEXT("Failed to read bool from json."));
+	return false;
+}
+
+bool FDMJsonUtils::Deserialize(const TSharedPtr<FJsonValue>& InJsonValue, FString& OutString)
+{
+	if (InJsonValue->TryGetString(OutString))
+	{
+		return true;
+	}
+
+	UE_LOG(LogDynamicMaterial, Error, TEXT("Failed to read string from json."));
+	return false;
+}
+
+bool FDMJsonUtils::Deserialize(const TSharedPtr<FJsonValue>& InJsonValue, FText& OutText)
+{
+	FString String;
+
+	if (InJsonValue->TryGetString(String))
+	{
+		OutText = FText::FromString(String);
+		return true;
+	}
+
+	UE_LOG(LogDynamicMaterial, Error, TEXT("Failed to read text from json."));
+	return false;
+}
+
+bool FDMJsonUtils::Deserialize(const TSharedPtr<FJsonValue>& InJsonValue, FName& OutName)
+{
+	FString String;
+
+	if (Deserialize(InJsonValue, String))
+	{
+		OutName = FName(*String);
+		return true;
+	}
+
+	UE_LOG(LogDynamicMaterial, Error, TEXT("Failed to read name from json."));
+	return false;
+}
+
+bool FDMJsonUtils::Deserialize(const TSharedPtr<FJsonValue>& InJsonValue, UClass*& OutClass)
+{
+	FString ClassString;
+
+	if (Deserialize(InJsonValue, ClassString))
+	{
+		if (UClass* Class = LoadClass<UObject>(nullptr, *ClassString))
+		{
+			OutClass = Class;
+			return true;
+		}
+
+		UE_LOG(LogDynamicMaterial, Error, TEXT("Failed to find class. [%s]"), *ClassString);
+	}
+	else
+	{
+		UE_LOG(LogDynamicMaterial, Error, TEXT("Failed to read class path."));
+	}
+
+	return false;
+}
+
+bool FDMJsonUtils::Deserialize(const TSharedPtr<FJsonValue>& InJsonValue, const UScriptStruct* InScriptStruct, void* OutData)
+{
+	if (!IsValid(InScriptStruct))
+	{
+		return false;
+	}
+
+	const TSharedPtr<FJsonObject>* JsonObject;
+
+	if (InJsonValue->TryGetObject(JsonObject))
+	{
+		if (FJsonObjectConverter::JsonObjectToUStruct((*JsonObject).ToSharedRef(), InScriptStruct, OutData))
+		{
+			return true;
+		}
+
+		UE_LOG(LogDynamicMaterial, Error, TEXT("Failed to find deserialize struct. [%s]"), *InScriptStruct->GetName());
+	}
+	else
+	{
+		UE_LOG(LogDynamicMaterial, Error, TEXT("Failed to read struct data. [%s]"), *InScriptStruct->GetName());
+	}
+
+	return false;
+}
+
+bool FDMJsonUtils::Deserialize(const TSharedPtr<FJsonValue>& InJsonValue, FObjectPtr& OutObject, UObject* InOuter)
+{
+	UObject* Object;
+
+	if (Deserialize(InJsonValue, Object, InOuter))
+	{
+		OutObject = Object;
+		return true;
+	}
+
+	return false;
+}
+
+bool FDMJsonUtils::Deserialize(const TSharedPtr<FJsonValue>& InJsonValue, const UEnum* InEnum, int64& OutValue)
+{
+	if (!InEnum)
+	{
+		UE_LOG(LogDynamicMaterial, Error, TEXT("Invalid enum."));
+		return false;
+	}
+
+	int64 Value;
+
+	if (Deserialize(InJsonValue, Value))
+	{
+		if (InEnum->IsValidEnumValueOrBitfield(Value))
+		{
+			OutValue = Value;
+			return true;
+		}
+
+		UE_LOG(LogDynamicMaterial, Error, TEXT("Invalid enum value. [%s] [%lli]"), *InEnum->GetName(), Value);
+	}
+	else
+	{
+		UE_LOG(LogDynamicMaterial, Error, TEXT("Failed to read data from json."));
+	}
+
+	return false;
+}
+
+bool FDMJsonUtils::Deserialize(const TSharedPtr<FJsonValue>& InJsonValue, UObject*& OutObject, UObject* InOuter)
+{
+	if (InJsonValue->IsNull())
+	{
+		OutObject = nullptr;
+		return true;
+	}
+
+	FString ObjectPath;
+
+	if (InJsonValue->TryGetString(ObjectPath))
+	{
+		// Objects in packages have no outer, even if one is provided
+		if (UObject* Object = LoadObject<UObject>(nullptr, *ObjectPath))
+		{
+			OutObject = Object;
+			return true;
+		}
+
+		UE_LOG(LogDynamicMaterial, Error, TEXT("Missing object. [%s]"), *ObjectPath);
+		return false;
+	}
+
+	do
+	{
+		using namespace UE::DynamicMaterial::Private;
+
+		TMap<FString, TSharedPtr<FJsonValue>> Data;
+
+		if (!Deserialize(InJsonValue, Data))
+		{
+			UE_LOG(LogDynamicMaterial, Error, TEXT("Failed to read data from json."));
+			break;
+		}
+
+		if (!Data.Contains(JsonKey_Class) || !Data.Contains(JsonKey_Data))
+		{
+			UE_LOG(LogDynamicMaterial, Error, TEXT("Missing data in json."));
+			break;
+		}
+
+		UClass* Class = nullptr;
+
+		if (!Deserialize(Data[JsonKey_Class], Class))
+		{
+			UE_LOG(LogDynamicMaterial, Error, TEXT("Failed to read class from json."));
+			break;
+		}
+
+		const TSharedPtr<FJsonObject>* ObjectData = nullptr;
+
+		if (!Data[JsonKey_Data]->TryGetObject(ObjectData))
+		{
+			UE_LOG(LogDynamicMaterial, Error, TEXT("Failed to read object data from json."));
+			break;
+		}
+
+		UObject* Object = NewObject<UObject>(InOuter, Class, NAME_None, RF_Transactional);
+
+		if (!Object)
+		{
+			UE_LOG(LogDynamicMaterial, Error, TEXT("Failed to find instantiate class. [%s]"), *Class->GetName());
+			break;
+		}
+
+		if (Object->Implements<UDMJsonSerializable>())
+		{
+			if (!Cast<IDMJsonSerializable>(Object)->JsonDeserialize(Data[JsonKey_Data]))
+			{
+				UE_LOG(LogDynamicMaterial, Error, TEXT("Failed to find deserialize object."));
+				break;
+			}
+		}
+		else
+		{
+			if (!FJsonObjectConverter::JsonObjectToUStruct(ObjectData->ToSharedRef(), Class, Object))
+			{
+				UE_LOG(LogDynamicMaterial, Error, TEXT("Failed to find deserialize object."));
+				break;
+			}
+		}
+
+		if (InOuter)
+		{
+			Object->Rename(/* New name */ nullptr, InOuter, UE::DynamicMaterial::RenameFlags);
+		}
+
+		return true;
+	}
+	while (false);
+
+	return false;
+}
+
+bool FDMJsonUtils::Deserialize(const TSharedPtr<FJsonValue>& InJsonValue, TMap<FString, TSharedPtr<FJsonValue>>& OutMap)
+{
+	const TSharedPtr<FJsonObject>* JsonObject;
+
+	if (InJsonValue->TryGetObject(JsonObject))
+	{
+		OutMap = (*JsonObject)->Values;
+		return true;
+	}
+
+	UE_LOG(LogDynamicMaterial, Error, TEXT("Failed to read object from json."));		
+	return false;
+}
+
+#endif
diff --git a/Engine/Plugins/Experimental/DynamicMaterial/Source/DynamicMaterial/Public/Components/DMMaterialValue.h b/Engine/Plugins/Experimental/DynamicMaterial/Source/DynamicMaterial/Public/Components/DMMaterialValue.h
--- a/Engine/Plugins/Experimental/DynamicMaterial/Source/DynamicMaterial/Public/Components/DMMaterialValue.h
+++ b/Engine/Plugins/Experimental/DynamicMaterial/Source/DynamicMaterial/Public/Components/DMMaterialValue.h
@@ -6,8 +6,9 @@
 #include "DMDefs.h"
 #include "Templates/SharedPointer.h"
 #include "UObject/StrongObjectPtr.h"
+#include "Utils/DMJsonUtils.h"
 #include "DMMaterialValue.generated.h"
- 
+
 class UDMMaterialParameter;
 class UDynamicMaterialModel;
 class UMaterial;
@@ -29,7 +30,7 @@ class UDMMaterialValue;
  * A value used inside a material. Can be exported as a material parameter.
  */
 UCLASS(BlueprintType, Abstract, ClassGroup = "Material Designer", meta = (DisplayName = "Material Designer Value"))
-class DYNAMICMATERIAL_API UDMMaterialValue : public UDMMaterialLinkedComponent
+class DYNAMICMATERIAL_API UDMMaterialValue : public UDMMaterialLinkedComponent, public IDMJsonSerializable
 {
 	GENERATED_BODY()
  
diff --git a/Engine/Plugins/Experimental/DynamicMaterial/Source/DynamicMaterial/Public/Components/DMRenderTargetRenderer.h b/Engine/Plugins/Experimental/DynamicMaterial/Source/DynamicMaterial/Public/Components/DMRenderTargetRenderer.h
--- a/Engine/Plugins/Experimental/DynamicMaterial/Source/DynamicMaterial/Public/Components/DMRenderTargetRenderer.h
+++ b/Engine/Plugins/Experimental/DynamicMaterial/Source/DynamicMaterial/Public/Components/DMRenderTargetRenderer.h
@@ -3,13 +3,14 @@
 #pragma once
 
 #include "DMMaterialComponent.h"
+#include "Utils/DMJsonUtils.h"
 #include "DMRenderTargetRenderer.generated.h"
 
 class UDMMaterialValueRenderTarget;
 template <typename T> class TSubclassOf;
 
 UCLASS(Abstract, BlueprintType, Blueprintable, ClassGroup = "Material Designer", meta = (DisplayName = "Material Designer Render Target Renderer"))
-class DYNAMICMATERIAL_API UDMRenderTargetRenderer : public UDMMaterialComponent
+class DYNAMICMATERIAL_API UDMRenderTargetRenderer : public UDMMaterialComponent, public IDMJsonSerializable
 {
 	GENERATED_BODY()
 
diff --git a/Engine/Plugins/Experimental/DynamicMaterial/Source/DynamicMaterial/Public/Components/MaterialValues/DMMaterialValueBool.h b/Engine/Plugins/Experimental/DynamicMaterial/Source/DynamicMaterial/Public/Components/MaterialValues/DMMaterialValueBool.h
--- a/Engine/Plugins/Experimental/DynamicMaterial/Source/DynamicMaterial/Public/Components/MaterialValues/DMMaterialValueBool.h
+++ b/Engine/Plugins/Experimental/DynamicMaterial/Source/DynamicMaterial/Public/Components/MaterialValues/DMMaterialValueBool.h
@@ -24,6 +24,13 @@ public:
 	UFUNCTION(BlueprintCallable, Category = "Material Designer")
 	void SetDefaultValue(bool bInDefaultValue);
 #endif
+
+#if WITH_EDITOR
+	//~ Begin IDMJsonSerializable
+	virtual TSharedPtr<FJsonValue> JsonSerialize() const override;
+	virtual bool JsonDeserialize(const TSharedPtr<FJsonValue>& InJsonValue) override;
+	//~ End IDMJsonSerializable
+#endif
  
 	//~ Begin UDMMaterialValue
 	virtual void SetMIDParameter(UMaterialInstanceDynamic* InMID) const override;
diff --git a/Engine/Plugins/Experimental/DynamicMaterial/Source/DynamicMaterial/Public/Components/MaterialValues/DMMaterialValueColorAtlas.h b/Engine/Plugins/Experimental/DynamicMaterial/Source/DynamicMaterial/Public/Components/MaterialValues/DMMaterialValueColorAtlas.h
--- a/Engine/Plugins/Experimental/DynamicMaterial/Source/DynamicMaterial/Public/Components/MaterialValues/DMMaterialValueColorAtlas.h
+++ b/Engine/Plugins/Experimental/DynamicMaterial/Source/DynamicMaterial/Public/Components/MaterialValues/DMMaterialValueColorAtlas.h
@@ -44,6 +44,13 @@ public:
 	virtual bool IsWholeLayerValue() const override { return true; }
 #endif
 
+#if WITH_EDITOR
+	//~ Begin IDMJsonSerializable
+	virtual TSharedPtr<FJsonValue> JsonSerialize() const override;
+	virtual bool JsonDeserialize(const TSharedPtr<FJsonValue>& InJsonValue) override;
+	//~ End IDMJsonSerializable
+#endif
+
 	//~ Begin UDMMaterialValue
 	virtual void SetMIDParameter(UMaterialInstanceDynamic* InMID) const override;
 #if WITH_EDITOR
diff --git a/Engine/Plugins/Experimental/DynamicMaterial/Source/DynamicMaterial/Public/Components/MaterialValues/DMMaterialValueFloat1.h b/Engine/Plugins/Experimental/DynamicMaterial/Source/DynamicMaterial/Public/Components/MaterialValues/DMMaterialValueFloat1.h
--- a/Engine/Plugins/Experimental/DynamicMaterial/Source/DynamicMaterial/Public/Components/MaterialValues/DMMaterialValueFloat1.h
+++ b/Engine/Plugins/Experimental/DynamicMaterial/Source/DynamicMaterial/Public/Components/MaterialValues/DMMaterialValueFloat1.h
@@ -24,6 +24,13 @@ public:
 	UFUNCTION(BlueprintCallable, Category = "Material Designer")
 	void SetDefaultValue(float InDefaultValue);
 #endif
+
+#if WITH_EDITOR
+	//~ Begin IDMJsonSerializable
+	virtual TSharedPtr<FJsonValue> JsonSerialize() const override;
+	virtual bool JsonDeserialize(const TSharedPtr<FJsonValue>& InJsonValue) override;
+	//~ End IDMJsonSerializable
+#endif
  
 	//~ Begin UDMMaterialValue
 	virtual void SetMIDParameter(UMaterialInstanceDynamic* InMID) const override;
diff --git a/Engine/Plugins/Experimental/DynamicMaterial/Source/DynamicMaterial/Public/Components/MaterialValues/DMMaterialValueFloat2.h b/Engine/Plugins/Experimental/DynamicMaterial/Source/DynamicMaterial/Public/Components/MaterialValues/DMMaterialValueFloat2.h
--- a/Engine/Plugins/Experimental/DynamicMaterial/Source/DynamicMaterial/Public/Components/MaterialValues/DMMaterialValueFloat2.h
+++ b/Engine/Plugins/Experimental/DynamicMaterial/Source/DynamicMaterial/Public/Components/MaterialValues/DMMaterialValueFloat2.h
@@ -24,6 +24,13 @@ public:
 	UFUNCTION(BlueprintCallable, Category = "Material Designer")
 	void SetDefaultValue(FVector2D InDefaultValue);
 #endif
+
+#if WITH_EDITOR
+	//~ Begin IDMJsonSerializable
+	virtual TSharedPtr<FJsonValue> JsonSerialize() const override;
+	virtual bool JsonDeserialize(const TSharedPtr<FJsonValue>& InJsonValue) override;
+	//~ End IDMJsonSerializable
+#endif
  
 	//~ Begin UDMMaterialValue
 	virtual void SetMIDParameter(UMaterialInstanceDynamic* InMID) const override;
diff --git a/Engine/Plugins/Experimental/DynamicMaterial/Source/DynamicMaterial/Public/Components/MaterialValues/DMMaterialValueFloat3RGB.h b/Engine/Plugins/Experimental/DynamicMaterial/Source/DynamicMaterial/Public/Components/MaterialValues/DMMaterialValueFloat3RGB.h
--- a/Engine/Plugins/Experimental/DynamicMaterial/Source/DynamicMaterial/Public/Components/MaterialValues/DMMaterialValueFloat3RGB.h
+++ b/Engine/Plugins/Experimental/DynamicMaterial/Source/DynamicMaterial/Public/Components/MaterialValues/DMMaterialValueFloat3RGB.h
@@ -24,6 +24,13 @@ public:
 	UFUNCTION(BlueprintCallable, Category = "Material Designer")
 	void SetDefaultValue(FLinearColor InDefaultValue);
 #endif
+
+#if WITH_EDITOR
+	//~ Begin IDMJsonSerializable
+	virtual TSharedPtr<FJsonValue> JsonSerialize() const override;
+	virtual bool JsonDeserialize(const TSharedPtr<FJsonValue>& InJsonValue) override;
+	//~ End IDMJsonSerializable
+#endif
  
 	//~ Begin UDMMaterialValue
 	virtual void SetMIDParameter(UMaterialInstanceDynamic* InMID) const override;
diff --git a/Engine/Plugins/Experimental/DynamicMaterial/Source/DynamicMaterial/Public/Components/MaterialValues/DMMaterialValueFloat3RPY.h b/Engine/Plugins/Experimental/DynamicMaterial/Source/DynamicMaterial/Public/Components/MaterialValues/DMMaterialValueFloat3RPY.h
--- a/Engine/Plugins/Experimental/DynamicMaterial/Source/DynamicMaterial/Public/Components/MaterialValues/DMMaterialValueFloat3RPY.h
+++ b/Engine/Plugins/Experimental/DynamicMaterial/Source/DynamicMaterial/Public/Components/MaterialValues/DMMaterialValueFloat3RPY.h
@@ -24,6 +24,13 @@ public:
 	UFUNCTION(BlueprintCallable, Category = "Material Designer")
 	void SetDefaultValue(FRotator InDefaultValue);
 #endif
+
+#if WITH_EDITOR
+	//~ Begin IDMJsonSerializable
+	virtual TSharedPtr<FJsonValue> JsonSerialize() const override;
+	virtual bool JsonDeserialize(const TSharedPtr<FJsonValue>& InJsonValue) override;
+	//~ End IDMJsonSerializable
+#endif
  
 	//~ Begin UDMMaterialValue
 	virtual void SetMIDParameter(UMaterialInstanceDynamic* InMID) const override;
diff --git a/Engine/Plugins/Experimental/DynamicMaterial/Source/DynamicMaterial/Public/Components/MaterialValues/DMMaterialValueFloat3XYZ.h b/Engine/Plugins/Experimental/DynamicMaterial/Source/DynamicMaterial/Public/Components/MaterialValues/DMMaterialValueFloat3XYZ.h
--- a/Engine/Plugins/Experimental/DynamicMaterial/Source/DynamicMaterial/Public/Components/MaterialValues/DMMaterialValueFloat3XYZ.h
+++ b/Engine/Plugins/Experimental/DynamicMaterial/Source/DynamicMaterial/Public/Components/MaterialValues/DMMaterialValueFloat3XYZ.h
@@ -24,6 +24,13 @@ public:
 	UFUNCTION(BlueprintCallable, Category = "Material Designer")
 	void SetDefaultValue(FVector InDefaultValue);
 #endif
+
+#if WITH_EDITOR
+	//~ Begin IDMJsonSerializable
+	virtual TSharedPtr<FJsonValue> JsonSerialize() const override;
+	virtual bool JsonDeserialize(const TSharedPtr<FJsonValue>& InJsonValue) override;
+	//~ End IDMJsonSerializable
+#endif
  
 	//~ Begin UDMMaterialValue
 	virtual void SetMIDParameter(UMaterialInstanceDynamic* InMID) const override;
diff --git a/Engine/Plugins/Experimental/DynamicMaterial/Source/DynamicMaterial/Public/Components/MaterialValues/DMMaterialValueFloat4.h b/Engine/Plugins/Experimental/DynamicMaterial/Source/DynamicMaterial/Public/Components/MaterialValues/DMMaterialValueFloat4.h
--- a/Engine/Plugins/Experimental/DynamicMaterial/Source/DynamicMaterial/Public/Components/MaterialValues/DMMaterialValueFloat4.h
+++ b/Engine/Plugins/Experimental/DynamicMaterial/Source/DynamicMaterial/Public/Components/MaterialValues/DMMaterialValueFloat4.h
@@ -24,6 +24,13 @@ public:
 	UFUNCTION(BlueprintCallable, Category = "Material Designer")
 	void SetDefaultValue(FLinearColor InDefaultValue);
 #endif
+
+#if WITH_EDITOR
+	//~ Begin IDMJsonSerializable
+	virtual TSharedPtr<FJsonValue> JsonSerialize() const override;
+	virtual bool JsonDeserialize(const TSharedPtr<FJsonValue>& InJsonValue) override;
+	//~ End IDMJsonSerializable
+#endif
  
 	//~ Begin UDMMaterialValue
 	virtual void SetMIDParameter(UMaterialInstanceDynamic* InMID) const override;
diff --git a/Engine/Plugins/Experimental/DynamicMaterial/Source/DynamicMaterial/Public/Components/MaterialValues/DMMaterialValueRenderTarget.h b/Engine/Plugins/Experimental/DynamicMaterial/Source/DynamicMaterial/Public/Components/MaterialValues/DMMaterialValueRenderTarget.h
--- a/Engine/Plugins/Experimental/DynamicMaterial/Source/DynamicMaterial/Public/Components/MaterialValues/DMMaterialValueRenderTarget.h
+++ b/Engine/Plugins/Experimental/DynamicMaterial/Source/DynamicMaterial/Public/Components/MaterialValues/DMMaterialValueRenderTarget.h
@@ -60,6 +60,11 @@ public:
 	void FlushCreateRenderTarget();
 
 #if WITH_EDITOR
+	//~ Begin IDMJsonSerializable
+	virtual TSharedPtr<FJsonValue> JsonSerialize() const override;
+	virtual bool JsonDeserialize(const TSharedPtr<FJsonValue>& InJsonValue) override;
+	//~ End IDMJsonSerializable
+
 	//~ Begin UDMMaterialValue
 	/** Render target is handled internally. */
 	virtual bool AllowEditValue() const override { return false; }
diff --git a/Engine/Plugins/Experimental/DynamicMaterial/Source/DynamicMaterial/Public/Components/MaterialValues/DMMaterialValueTexture.h b/Engine/Plugins/Experimental/DynamicMaterial/Source/DynamicMaterial/Public/Components/MaterialValues/DMMaterialValueTexture.h
--- a/Engine/Plugins/Experimental/DynamicMaterial/Source/DynamicMaterial/Public/Components/MaterialValues/DMMaterialValueTexture.h
+++ b/Engine/Plugins/Experimental/DynamicMaterial/Source/DynamicMaterial/Public/Components/MaterialValues/DMMaterialValueTexture.h
@@ -46,13 +46,12 @@ public:
 	void SetDefaultValue(UTexture* InDefaultValue);
 #endif // WITH_EDITOR
 
-	//~ Begin UObject
 #if WITH_EDITOR
-	virtual void PreEditChange(FProperty* PropertyAboutToChange) override;
-	virtual void PostEditChangeProperty(FPropertyChangedEvent& PropertyChangedEvent) override;
+	//~ Begin IDMJsonSerializable
+	virtual TSharedPtr<FJsonValue> JsonSerialize() const override;
+	virtual bool JsonDeserialize(const TSharedPtr<FJsonValue>& InJsonValue) override;
+	//~ End IDMJsonSerializable
 #endif
-	virtual void PostLoad() override;
-	//~ End UObject
 
 	//~ Begin UDMMaterialValue
 	virtual void SetMIDParameter(UMaterialInstanceDynamic* InMID) const override;
@@ -65,6 +64,14 @@ public:
 #endif
 	//~ End UDMMaterialValue
 
+	//~ Begin UObject
+#if WITH_EDITOR
+	virtual void PreEditChange(FProperty* InPropertyAboutToChange) override;
+	virtual void PostEditChangeProperty(FPropertyChangedEvent& InPropertyChangedEvent) override;
+#endif
+	virtual void PostLoad() override;
+	//~ End UObject
+
 protected:
 	UDMMaterialValueTexture();
  
diff --git a/Engine/Plugins/Experimental/DynamicMaterial/Source/DynamicMaterial/Public/Components/RenderTargetRenderers/DMRenderTargetTextRenderer.h b/Engine/Plugins/Experimental/DynamicMaterial/Source/DynamicMaterial/Public/Components/RenderTargetRenderers/DMRenderTargetTextRenderer.h
--- a/Engine/Plugins/Experimental/DynamicMaterial/Source/DynamicMaterial/Public/Components/RenderTargetRenderers/DMRenderTargetTextRenderer.h
+++ b/Engine/Plugins/Experimental/DynamicMaterial/Source/DynamicMaterial/Public/Components/RenderTargetRenderers/DMRenderTargetTextRenderer.h
@@ -169,6 +169,11 @@ public:
 	void SetGlowOuterRadius(const FVector2D& InGlowOuterRadius);
 
 #if WITH_EDITOR
+	//~ Begin IDMJsonSerializable
+	virtual TSharedPtr<FJsonValue> JsonSerialize() const override;
+	virtual bool JsonDeserialize(const TSharedPtr<FJsonValue>& InJsonValue) override;
+	//~ End IDMJsonSerializable
+
 	//~ Begin UObject
 	virtual void PostEditChangeProperty(FPropertyChangedEvent& InPropertyChangedEvent) override;
 	//~ End UObject
diff --git a/Engine/Plugins/Experimental/DynamicMaterial/Source/DynamicMaterial/Public/Components/RenderTargetRenderers/DMRenderTargetWidgetRenderer.h b/Engine/Plugins/Experimental/DynamicMaterial/Source/DynamicMaterial/Public/Components/RenderTargetRenderers/DMRenderTargetWidgetRenderer.h
--- a/Engine/Plugins/Experimental/DynamicMaterial/Source/DynamicMaterial/Public/Components/RenderTargetRenderers/DMRenderTargetWidgetRenderer.h
+++ b/Engine/Plugins/Experimental/DynamicMaterial/Source/DynamicMaterial/Public/Components/RenderTargetRenderers/DMRenderTargetWidgetRenderer.h
@@ -23,6 +23,11 @@ public:
 	void SetWidgetClass(TSubclassOf<UWidget> InWidgetClass);
 
 #if WITH_EDITOR
+	//~ Begin IDMJsonSerializable
+	virtual TSharedPtr<FJsonValue> JsonSerialize() const override;
+	virtual bool JsonDeserialize(const TSharedPtr<FJsonValue>& InJsonValue) override;
+	//~ End IDMJsonSerializable
+
 	//~ Begin UObject
 	virtual void PostEditChangeProperty(FPropertyChangedEvent& InPropertyChangedEvent) override;
 	//~ End UObject
diff --git a/Engine/Plugins/Experimental/DynamicMaterial/Source/DynamicMaterial/Public/Utils/DMJsonUtils.h b/Engine/Plugins/Experimental/DynamicMaterial/Source/DynamicMaterial/Public/Utils/DMJsonUtils.h
new file mode 100644
--- /dev/null
+++ b/Engine/Plugins/Experimental/DynamicMaterial/Source/DynamicMaterial/Public/Utils/DMJsonUtils.h
@@ -0,0 +1,168 @@
+// Copyright Epic Games, Inc. All Rights Reserved.
+
+#pragma once
+
+#include "Concepts/BaseStructureProvider.h"
+#include "Math/MathFwd.h"
+#include "Templates/SharedPointerFwd.h"
+#include "UObject/Interface.h"
+#include "DMJsonUtils.generated.h"
+
+class FJsonValue;
+class UClass;
+class UTexture;
+
+/** Offers custom serialization options for class */
+UINTERFACE(MinimalAPI, Blueprintable, BlueprintType)
+class UDMJsonSerializable : public UInterface
+{
+	GENERATED_BODY()
+};
+
+class IDMJsonSerializable
+{
+	GENERATED_BODY()
+
+public:
+	virtual TSharedPtr<FJsonValue> JsonSerialize() const PURE_VIRTUAL(UDMMaterialValue::JsonSerialize, return nullptr;)
+	virtual bool JsonDeserialize(const TSharedPtr<FJsonValue>& InJsonValue) PURE_VIRTUAL(UDMMaterialValue::JsonDeserialize, return false;)
+};
+
+#if WITH_EDITOR
+class DYNAMICMATERIAL_API FDMJsonUtils
+{
+public:
+	static TSharedPtr<FJsonValue> Serialize(bool bInValue);
+	static TSharedPtr<FJsonValue> Serialize(const FString& InString);
+	static TSharedPtr<FJsonValue> Serialize(const FText& InText);
+	static TSharedPtr<FJsonValue> Serialize(const FName& InName);
+	static TSharedPtr<FJsonValue> Serialize(const UClass* InClass);
+	static TSharedPtr<FJsonValue> Serialize(const UScriptStruct* InScriptStruct, const void* InData);
+	static TSharedPtr<FJsonValue> Serialize(const UObject* InObject);
+	static TSharedPtr<FJsonValue> Serialize(const FObjectPtr& InObject);
+	static TSharedPtr<FJsonValue> Serialize(const TMap<FString, TSharedPtr<FJsonValue>>& InMap);
+
+	template<typename InArithmeticType
+		UE_REQUIRES(TIsArithmetic<InArithmeticType>::Value)>
+	static TSharedPtr<FJsonValue> Serialize(InArithmeticType InValue)
+	{
+		return SerializeNumber(static_cast<double>(InValue));
+	}
+
+	template<typename InEnumType
+		UE_REQUIRES(TIsEnum<InEnumType>::Value)>
+	static TSharedPtr<FJsonValue> Serialize(InEnumType InValue)
+	{
+		return SerializeNumber(static_cast<double>(InValue));
+	}
+
+	template<typename InStructType
+		UE_REQUIRES(TModels_V<CBaseStructureProvider, InStructType>)>
+	static TSharedPtr<FJsonValue> Serialize(const InStructType& InStruct)
+	{
+		return Serialize(TBaseStructure<InStructType>::Get(), static_cast<const void*>(&InStruct));
+	}
+
+	template<typename InClassType
+		UE_REQUIRES(TModels_V<CStaticClassProvider, InClassType>)>
+	static TSharedPtr<FJsonValue> Serialize(const TSubclassOf<InClassType>& InClass)
+	{
+		return Serialize(InClass.Get());
+	}
+
+	static bool Deserialize(const TSharedPtr<FJsonValue>& InJsonValue, bool& bOutValue);
+	static bool Deserialize(const TSharedPtr<FJsonValue>& InJsonValue, FString& OutString);
+	static bool Deserialize(const TSharedPtr<FJsonValue>& InJsonValue, FText& OutText);
+	static bool Deserialize(const TSharedPtr<FJsonValue>& InJsonValue, FName& OutName);
+	static bool Deserialize(const TSharedPtr<FJsonValue>& InJsonValue, UClass*& OutClass);
+	static bool Deserialize(const TSharedPtr<FJsonValue>& InJsonValue, const UScriptStruct* InScriptStruct, void* OutData);
+	static bool Deserialize(const TSharedPtr<FJsonValue>& InJsonValue, const UEnum* InEnum, int64& OutValue);
+	static bool Deserialize(const TSharedPtr<FJsonValue>& InJsonValue, UObject*& OutObject, UObject* InOuter = nullptr);
+	static bool Deserialize(const TSharedPtr<FJsonValue>& InJsonValue, FObjectPtr& OutObject, UObject* InOuter = nullptr);
+	static bool Deserialize(const TSharedPtr<FJsonValue>& InJsonValue, TMap<FString, TSharedPtr<FJsonValue>>& OutMap);
+
+	template<typename InArithmeticType
+		UE_REQUIRES(TIsArithmetic<InArithmeticType>::Value)>
+	static bool Deserialize(const TSharedPtr<FJsonValue>& InJsonValue, InArithmeticType& OutValue)
+	{
+		double Value;
+
+		if (DeserializeNumber(InJsonValue, Value))
+		{
+			if constexpr (TIsIntegral<InArithmeticType>::Value)
+			{
+				OutValue = FMath::RoundToInt(Value);
+			}
+			else
+			{
+				OutValue = static_cast<InArithmeticType>(Value);
+			}
+			return true;
+		}
+
+		return false;
+	}
+
+	template<typename InEnumType
+		UE_REQUIRES(TIsEnum<InEnumType>::Value)>
+	static bool Deserialize(const TSharedPtr<FJsonValue>& InJsonValue, InEnumType& OutValue)
+	{
+		int64 Value;
+
+		if (Deserialize(InJsonValue, StaticEnum<InEnumType>(), Value))
+		{
+			OutValue = static_cast<InEnumType>(Value);
+			return true;
+		}
+
+		return false;
+	}
+
+	template<typename InStructType
+		UE_REQUIRES(TModels_V<CBaseStructureProvider, InStructType>)>
+	static bool Deserialize(const TSharedPtr<FJsonValue>& InJsonValue, InStructType& OutStruct)
+	{
+		return Deserialize(InJsonValue, TBaseStructure<InStructType>::Get(), static_cast<void*>(&OutStruct));
+	}
+
+	template<typename InClassType
+		UE_REQUIRES(TModels_V<CStaticClassProvider, InClassType>)>
+	static bool Deserialize(const TSharedPtr<FJsonValue>& InJsonValue, InClassType*& OutObject, UObject* InOuter = nullptr)
+	{
+		UObject* Object;
+
+		if (Deserialize(InJsonValue, Object, InOuter))
+		{
+			if (!Object || Object->IsA<InClassType>())
+			{
+				OutObject = Cast<InClassType>(Object);
+				return true;
+			}
+		}
+
+		return false;
+	}
+
+	template<typename InClassType
+		UE_REQUIRES(TModels_V<CStaticClassProvider, InClassType>)>
+	static bool Deserialize(const TSharedPtr<FJsonValue>& InJsonValue, TSubclassOf<InClassType>& OutSubclass)
+	{
+		UClass* Class = nullptr;
+
+		if (Deserialize(InJsonValue, Class))
+		{
+			if (!Class || Class->IsChildOf(InClassType::StaticClass()))
+			{
+				OutSubclass = Class;
+				return true;
+			}
+		}
+
+		return false;
+	}
+
+private:
+	static TSharedPtr<FJsonValue> SerializeNumber(double InNumber);
+	static bool DeserializeNumber(const TSharedPtr<FJsonValue>& InJsonValue, double& OutNumber);
+};
+#endif
diff --git a/Engine/Plugins/Experimental/DynamicMaterial/Source/DynamicMaterialEditor/DynamicMaterialEditor.Build.cs b/Engine/Plugins/Experimental/DynamicMaterial/Source/DynamicMaterialEditor/DynamicMaterialEditor.Build.cs
--- a/Engine/Plugins/Experimental/DynamicMaterial/Source/DynamicMaterialEditor/DynamicMaterialEditor.Build.cs
+++ b/Engine/Plugins/Experimental/DynamicMaterial/Source/DynamicMaterialEditor/DynamicMaterialEditor.Build.cs
@@ -29,6 +29,7 @@ public class DynamicMaterialEditor : ModuleRules
 				"CustomDetailsView",
 				"DeveloperSettings",
 				"DynamicMaterialShadersEditor",
+				"EditorSubsystem",
 				"EditorWidgets",
 				"InputCore",
 				"Json",
diff --git a/Engine/Plugins/Experimental/DynamicMaterial/Source/DynamicMaterialEditor/Private/Components/DMMaterialEffectFunction.cpp b/Engine/Plugins/Experimental/DynamicMaterial/Source/DynamicMaterialEditor/Private/Components/DMMaterialEffectFunction.cpp
--- a/Engine/Plugins/Experimental/DynamicMaterial/Source/DynamicMaterialEditor/Private/Components/DMMaterialEffectFunction.cpp
+++ b/Engine/Plugins/Experimental/DynamicMaterial/Source/DynamicMaterialEditor/Private/Components/DMMaterialEffectFunction.cpp
@@ -8,7 +8,8 @@
 #include "Components/MaterialValues/DMMaterialValueFloat3XYZ.h"
 #include "Components/MaterialValues/DMMaterialValueFloat4.h"
 #include "DMComponentPath.h"
-#include "DMPrivate.h"
+#include "Dom/JsonObject.h"
+#include "Dom/JsonValue.h"
 #include "DynamicMaterialEditorModule.h"
 #include "Materials/MaterialExpressionComponentMask.h"
 #include "Materials/MaterialExpressionFunctionInput.h"
@@ -18,6 +19,8 @@
 #include "MaterialValueType.h"
 #include "Model/DMMaterialBuildState.h"
 #include "Model/DynamicMaterialModelEditorOnlyData.h"
+#include "Utils/DMMaterialFunctionLibrary.h"
+#include "Utils/DMPrivate.h"
 
 #define LOCTEXT_NAMESPACE "DMMaterialEffectFunction"
 
@@ -76,6 +79,110 @@ const TArray<TObjectPtr<UDMMaterialValue>>& UDMMaterialEffectFunction::GetInputV
 	return InputValues;
 }
 
+TSharedPtr<FJsonValue> UDMMaterialEffectFunction::JsonSerialize() const
+{
+	TArray<TSharedPtr<FJsonValue>> ValueArray;
+
+	if (MaterialFunctionPtr)
+	{
+		for (int32 InputIndex = 0; InputIndex < InputValues.Num(); ++InputIndex)
+		{
+			// Index 0 is always null
+			if (InputIndex == 0)
+			{
+				ValueArray.Add(MakeShared<FJsonValueNull>());
+			}
+			else if (InputValues[InputIndex])
+			{
+				ValueArray.Add(InputValues[InputIndex]->JsonSerialize());
+			}
+			else
+			{
+				UE::DynamicMaterialEditor::Private::LogError(TEXT("Null input found when serializing material effect function."), true, this);
+			}
+		}
+	}
+
+	return FDMJsonUtils::Serialize({
+		{GET_MEMBER_NAME_STRING_CHECKED(ThisClass, bEnabled), FDMJsonUtils::Serialize(bEnabled)},
+		{GET_MEMBER_NAME_STRING_CHECKED(ThisClass, MaterialFunctionPtr), FDMJsonUtils::Serialize(MaterialFunctionPtr)},
+		{GET_MEMBER_NAME_STRING_CHECKED(ThisClass, InputValues), MakeShared<FJsonValueArray>(ValueArray)},
+	});
+}
+
+bool UDMMaterialEffectFunction::JsonDeserialize(const TSharedPtr<FJsonValue>& InJsonValue)
+{
+	TMap<FString, TSharedPtr<FJsonValue>> Data;
+
+	if (!FDMJsonUtils::Deserialize(InJsonValue, Data))
+	{
+		return false;
+	}
+
+	bool bSuccess = false;
+
+	if (const TSharedPtr<FJsonValue>* JsonValue = Data.Find(GET_MEMBER_NAME_STRING_CHECKED(ThisClass, bEnabled)))
+	{
+		bool bEnabledJson = false;
+
+		if (FDMJsonUtils::Deserialize(*JsonValue, bEnabledJson))
+		{
+			const FDMUpdateGuard Guard;
+			SetEnabled(bEnabledJson);
+			bSuccess = true;
+		}
+	}
+
+	if (const TSharedPtr<FJsonValue>* JsonValue = Data.Find(GET_MEMBER_NAME_STRING_CHECKED(ThisClass, MaterialFunctionPtr)))
+	{
+		UMaterialFunctionInterface* MaterialFunctionJson = nullptr;
+
+		if (FDMJsonUtils::Deserialize(*JsonValue, MaterialFunctionJson))
+		{
+			const FDMUpdateGuard Guard;
+			SetMaterialFunction(MaterialFunctionJson);
+			bSuccess = true;
+		}
+	}
+
+	// If there's no effect function, there's no input values
+	if (MaterialFunctionPtr)
+	{
+		if (const TSharedPtr<FJsonValue>* JsonValue = Data.Find(GET_MEMBER_NAME_STRING_CHECKED(ThisClass, InputValues)))
+		{
+			const TArray<TSharedPtr<FJsonValue>>* InputValuesJson = nullptr;
+
+			if ((*JsonValue)->TryGetArray(InputValuesJson))
+			{
+				if (InputValuesJson->Num() != InputValues.Num())
+				{
+					UE::DynamicMaterialEditor::Private::LogError(TEXT("Mismatched input value count deserializing effect function."), true, this);
+					bSuccess = false;
+				}
+				else
+				{
+					// Index 0 is always nullptr, ignore it.
+					for (int32 InputIndex = 1; InputIndex < InputValues.Num(); ++InputIndex)
+					{
+						if (!InputValues[InputIndex]->JsonDeserialize((*InputValuesJson)[InputIndex]))
+						{
+							UE::DynamicMaterialEditor::Private::LogError(TEXT("Unable to deserialize input value while deserializing material effect function."), true, this);
+							bSuccess = false;
+						}
+					}
+				}
+			}
+		}
+	}
+
+	if (bSuccess)
+	{
+		Update(EDMUpdateType::Structure);
+	}
+
+	return bSuccess;
+}
+
 FText UDMMaterialEffectFunction::GetEffectName() const
 {
 	if (UMaterialFunctionInterface* MaterialFunction = MaterialFunctionPtr.Get())
@@ -107,6 +214,18 @@ FText UDMMaterialEffectFunction::GetEffectDescription() const
 	return FText::GetEmpty();
 }
 
+bool UDMMaterialEffectFunction::IsCompatibleWith(UDMMaterialEffect* InEffect) const
+{
+	if (!MaterialFunctionPtr)
+	{
+		return false;
+	}
+
+	UDMMaterialEffectFunction* EffectFunction = Cast<UDMMaterialEffectFunction>(InEffect);
+
+	return !EffectFunction || EffectFunction->GetMaterialFunction() != MaterialFunctionPtr;
+}
+
 void UDMMaterialEffectFunction::ApplyTo(const TSharedRef<FDMMaterialBuildState>& InBuildState, TArray<UMaterialExpression*>& InOutStageExpressions,
 	int32& InOutLastExpressionOutputChannel, int32& InOutLastExpressionOutputIndex) const
 {
diff --git a/Engine/Plugins/Experimental/DynamicMaterial/Source/DynamicMaterialEditor/Private/Components/DMMaterialEffectStack.cpp b/Engine/Plugins/Experimental/DynamicMaterial/Source/DynamicMaterialEditor/Private/Components/DMMaterialEffectStack.cpp
--- a/Engine/Plugins/Experimental/DynamicMaterial/Source/DynamicMaterialEditor/Private/Components/DMMaterialEffectStack.cpp
+++ b/Engine/Plugins/Experimental/DynamicMaterial/Source/DynamicMaterialEditor/Private/Components/DMMaterialEffectStack.cpp
@@ -1,11 +1,13 @@
 // Copyright Epic Games, Inc. All Rights Reserved.
 
 #include "Components/DMMaterialEffectStack.h"
+
 #include "Components/DMMaterialEffect.h"
 #include "Components/DMMaterialLayer.h"
 #include "Components/DMMaterialSlot.h"
 #include "DMComponentPath.h"
 #include "UObject/Package.h"
+#include "Utils/DMPrivate.h"
 
 #define LOCTEXT_NAMESPACE "DMMaterialEffectStack"
 
@@ -121,6 +123,22 @@ bool UDMMaterialEffectStack::AddEffect(UDMMaterialEffect* InEffect)
 		InEffect->Modify();
 	}
 
+	TArray<UDMMaterialEffect*> IncompatibleEffects = GetIncompatibleEffects(InEffect);
+
+	if (IncompatibleEffects.Num() == 1)
+	{
+		for (int32 ReplaceIndex = 0; ReplaceIndex < Effects.Num(); ++ReplaceIndex)
+		{
+			if (Effects[ReplaceIndex] == IncompatibleEffects[0])
+			{
+				SetEffect(ReplaceIndex, InEffect);
+				return true;
+			}
+		}
+	}
+
+	RemoveIncompatibleEffects(InEffect);
+
 	if (UDMMaterialEffectStack* OldStack = InEffect->GetEffectStack())
 	{
 		if (GUndo)
@@ -145,40 +163,40 @@ bool UDMMaterialEffectStack::AddEffect(UDMMaterialEffect* InEffect)
 	return true;
 }
 
-bool UDMMaterialEffectStack::SetEffect(int32 InIndex, UDMMaterialEffect* InEffect)
+UDMMaterialEffect* UDMMaterialEffectStack::SetEffect(int32 InIndex, UDMMaterialEffect* InEffect)
 {
-	if (!IsValid(InEffect))
+	if (!Effects.IsValidIndex(InIndex) || !IsValid(InEffect))
 	{
-		return false;
+		return nullptr;
 	}
 
-	if (GUndo)
-	{
-		InEffect->Modify();
-	}
+	UDMMaterialEffect* OldEffect = Effects[InIndex];
 
-	if (UDMMaterialEffectStack* OldStack = InEffect->GetEffectStack())
+	if (OldEffect)
 	{
 		if (GUndo)
 		{
-			OldStack->Modify();
+			OldEffect->Modify();
 		}
 
-		OldStack->RemoveEffect(InEffect);
+		OldEffect->SetEnabled(false);
+		OldEffect->Rename(nullptr, GetTransientPackage(), UE::DynamicMaterial::RenameFlags);
+		OldEffect->SetComponentState(EDMComponentLifetimeState::Removed);
 	}
 
-	Effects.Add(InEffect);
-
-	InEffect->Rename(nullptr, this, UE::DynamicMaterial::RenameFlags);
-
-	if (IsComponentAdded())
+	if (GUndo)
 	{
-		InEffect->SetComponentState(EDMComponentLifetimeState::Added);
+		InEffect->Modify();
 	}
 
-	InEffect->Update(EDMUpdateType::Structure);
+	InEffect->SetEnabled(true);
+	InEffect->Rename(nullptr, this, UE::DynamicMaterial::RenameFlags);
+	InEffect->SetComponentState(EDMComponentLifetimeState::Added);
 
-	return true;
+	Effects[InIndex] = InEffect;
+	Effects[InIndex]->Update(EDMUpdateType::Structure);
+
+	return OldEffect;
 }
 
 bool UDMMaterialEffectStack::MoveEffect(int32 InIndex, int32 InNewIndex)
@@ -231,11 +249,11 @@ bool UDMMaterialEffectStack::MoveEffect(UDMMaterialEffect* InEffect, int32 InNew
 	return false;
 }
 
-bool UDMMaterialEffectStack::RemoveEffect(int32 InIndex)
+UDMMaterialEffect* UDMMaterialEffectStack::RemoveEffect(int32 InIndex)
 {
 	if (!Effects.IsValidIndex(InIndex))
 	{
-		return false;
+		return nullptr;
 	}
 
 	UDMMaterialEffect* Effect = Effects[InIndex];
@@ -253,7 +271,7 @@ bool UDMMaterialEffectStack::RemoveEffect(int32 InIndex)
 
 	Update(EDMUpdateType::Structure);
 
-	return true;
+	return Effect;
 }
 
 bool UDMMaterialEffectStack::RemoveEffect(UDMMaterialEffect* InEffect)
@@ -272,7 +290,7 @@ bool UDMMaterialEffectStack::RemoveEffect(UDMMaterialEffect* InEffect)
 
 	if (EffectIndex != INDEX_NONE)
 	{
-		return RemoveEffect(EffectIndex);
+		return !!RemoveEffect(EffectIndex);
 	}
 
 	return false;
@@ -298,6 +316,49 @@ bool UDMMaterialEffectStack::ApplyEffects(const TSharedRef<FDMMaterialBuildState
 	return bAppliedEffect;
 }
 
+FDMMaterialEffectStackJson UDMMaterialEffectStack::CreatePreset()
+{
+	FDMMaterialEffectStackJson Preset;
+	Preset.bEnabled = bEnabled;
+
+	Preset.Effects.Reserve(Effects.Num());
+
+	for (const TObjectPtr<UDMMaterialEffect>& Effect : Effects)
+	{
+		FDMMaterialEffectJson& EffectJson = Preset.Effects.AddDefaulted_GetRef();
+		EffectJson.Class = Effect->GetClass();
+		EffectJson.Data = Effect->JsonSerialize();
+	}
+
+	return Preset;
+}
+
+void UDMMaterialEffectStack::ApplyPreset(const FDMMaterialEffectStackJson& InPreset)
+{
+	SetEnabled(InPreset.bEnabled);
+
+	for (const FDMMaterialEffectJson& EffectJson : InPreset.Effects)
+	{
+		if (!EffectJson.Class.Get())
+		{
+			UE::DynamicMaterialEditor::Private::LogError(TEXT("Invalid class when applying effect preset."), true, this);
+			continue;
+		}
+
+		UDMMaterialEffect* Effect = UDMMaterialEffect::CreateEffect(this, EffectJson.Class);
+
+		if (!Effect)
+		{
+			UE::DynamicMaterialEditor::Private::LogError(TEXT("Failed creating class when applying effect preset."), true, this);
+			continue;
+		}
+
+		Effect->JsonDeserialize(EffectJson.Data);
+
+		AddEffect(Effect);
+	}
+}
+
 UDMMaterialComponent* UDMMaterialEffectStack::GetParentComponent() const
 {
 	UObject* Outer = GetOuterSafe();
@@ -401,6 +462,49 @@ void UDMMaterialEffectStack::PostEditUndo()
 	Update(EDMUpdateType::Structure);
 }
 
+TArray<UDMMaterialEffect*> UDMMaterialEffectStack::GetIncompatibleEffects(UDMMaterialEffect* InEffect)
+{
+	if (!InEffect)
+	{
+		return {};
+	}
+
+	TArray<UDMMaterialEffect*> IncompatibleEffects;
+
+	for (const TObjectPtr<UDMMaterialEffect>& Effect : Effects)
+	{
+		if (!Effect->IsCompatibleWith(InEffect))
+		{
+			IncompatibleEffects.Add(Effect);
+		}
+	}
+
+	return IncompatibleEffects;
+}
+
+TArray<UDMMaterialEffect*> UDMMaterialEffectStack::RemoveIncompatibleEffects(UDMMaterialEffect* InEffect)
+{
+	if (!InEffect)
+	{
+		return {};
+	}
+
+	TArray<UDMMaterialEffect*> IncompatibleEffects;
+
+	for (int32 Index = Effects.Num() - 1; Index >= 0; --Index)
+	{
+		UDMMaterialEffect* Effect = Effects[Index];
+
+		if (!Effect->IsCompatibleWith(InEffect))
+		{
+			IncompatibleEffects.Add(Effect);
+			RemoveEffect(Index);
+		}
+	}
+
+	return IncompatibleEffects;
+}
+
 UDMMaterialComponent* UDMMaterialEffectStack::GetSubComponentByPath(FDMComponentPath& InPath, const FDMComponentPathSegment& InPathSegment) const
 {
 	if (InPathSegment.GetToken() == EffectsPathToken)
diff --git a/Engine/Plugins/Experimental/DynamicMaterial/Source/DynamicMaterialEditor/Private/Components/DMMaterialLayer.cpp b/Engine/Plugins/Experimental/DynamicMaterial/Source/DynamicMaterialEditor/Private/Components/DMMaterialLayer.cpp
--- a/Engine/Plugins/Experimental/DynamicMaterial/Source/DynamicMaterialEditor/Private/Components/DMMaterialLayer.cpp
+++ b/Engine/Plugins/Experimental/DynamicMaterial/Source/DynamicMaterialEditor/Private/Components/DMMaterialLayer.cpp
@@ -6,7 +6,6 @@
 #include "Components/DMMaterialSlot.h"
 #include "Components/DMMaterialStage.h"
 #include "DMComponentPath.h"
-#include "DMPrivate.h"
 #include "Dom/JsonObject.h"
 #include "DynamicMaterialEditorModule.h"
 #include "Factories.h"
@@ -17,6 +16,7 @@
 #include "Model/DynamicMaterialModelEditorOnlyData.h"
 #include "Serialization/JsonSerializer.h"
 #include "Serialization/JsonWriter.h"
+#include "Utils/DMPrivate.h"
 
 #define LOCTEXT_NAMESPACE "DMMaterialLayer"
 
diff --git a/Engine/Plugins/Experimental/DynamicMaterial/Source/DynamicMaterialEditor/Private/Components/DMMaterialProperty.cpp b/Engine/Plugins/Experimental/DynamicMaterial/Source/DynamicMaterialEditor/Private/Components/DMMaterialProperty.cpp
--- a/Engine/Plugins/Experimental/DynamicMaterial/Source/DynamicMaterialEditor/Private/Components/DMMaterialProperty.cpp
+++ b/Engine/Plugins/Experimental/DynamicMaterial/Source/DynamicMaterialEditor/Private/Components/DMMaterialProperty.cpp
@@ -9,7 +9,6 @@
 #include "Components/MaterialStageInputs/DMMSIExpression.h"
 #include "Components/MaterialStageInputs/DMMSIValue.h"
 #include "Components/MaterialValues/DMMaterialValueTexture.h"
-#include "DMPrivate.h"
 #include "DMValueDefinition.h"
 #include "DynamicMaterialEditorSettings.h"
 #include "Materials/MaterialExpressionConstant.h"
@@ -24,6 +23,7 @@
 #include "Model/DMMaterialBuildState.h"
 #include "Model/DMMaterialBuildUtils.h"
 #include "Model/DynamicMaterialModelEditorOnlyData.h"
+#include "Utils/DMPrivate.h"
 
 #define LOCTEXT_NAMESPACE "DMMaterialProperty"
 
diff --git a/Engine/Plugins/Experimental/DynamicMaterial/Source/DynamicMaterialEditor/Private/Components/DMMaterialSlot.cpp b/Engine/Plugins/Experimental/DynamicMaterial/Source/DynamicMaterialEditor/Private/Components/DMMaterialSlot.cpp
--- a/Engine/Plugins/Experimental/DynamicMaterial/Source/DynamicMaterialEditor/Private/Components/DMMaterialSlot.cpp
+++ b/Engine/Plugins/Experimental/DynamicMaterial/Source/DynamicMaterialEditor/Private/Components/DMMaterialSlot.cpp
@@ -10,20 +10,20 @@
 #include "CoreGlobals.h"
 #include "DMComponentPath.h"
 #include "DMDefs.h"
-#include "DMPrivate.h"
 #include "DynamicMaterialEditorModule.h"
 #include "Factories/MaterialFactoryNew.h"
-#include "MaterialValueType.h"
 #include "Materials/Material.h"
 #include "Materials/MaterialExpression.h"
 #include "Materials/MaterialExpressionMax.h"
 #include "Materials/MaterialExpressionMultiply.h"
+#include "MaterialValueType.h"
 #include "Model/DMMaterialBuildState.h"
 #include "Model/DMMaterialBuildUtils.h"
 #include "Model/DynamicMaterialModel.h"
 #include "Model/DynamicMaterialModelEditorOnlyData.h"
 #include "UObject/Package.h"
 #include "UObject/UObjectGlobals.h"
+#include "Utils/DMPrivate.h"
 
 #define LOCTEXT_NAMESPACE "DMMaterialStage"
 
diff --git a/Engine/Plugins/Experimental/DynamicMaterial/Source/DynamicMaterialEditor/Private/Components/DMMaterialStageBlend.cpp b/Engine/Plugins/Experimental/DynamicMaterial/Source/DynamicMaterialEditor/Private/Components/DMMaterialStageBlend.cpp
--- a/Engine/Plugins/Experimental/DynamicMaterial/Source/DynamicMaterialEditor/Private/Components/DMMaterialStageBlend.cpp
+++ b/Engine/Plugins/Experimental/DynamicMaterial/Source/DynamicMaterialEditor/Private/Components/DMMaterialStageBlend.cpp
@@ -10,13 +10,13 @@
 #include "Components/MaterialStageInputs/DMMSITextureUV.h"
 #include "Components/MaterialStageInputs/DMMSIValue.h"
 #include "Components/MaterialValues/DMMaterialValueFloat1.h"
-#include "DMPrivate.h"
-#include "Helpers/DMInputNodeBuilder.h"
 #include "Materials/MaterialExpressionComponentMask.h"
 #include "Model/DMMaterialBuildState.h"
 #include "Model/DMMaterialBuildUtils.h"
 #include "Model/DynamicMaterialModel.h"
 #include "Model/DynamicMaterialModelEditorOnlyData.h"
+#include "Utils/DMInputNodeBuilder.h"
+#include "Utils/DMPrivate.h"
  
 #define LOCTEXT_NAMESPACE "DMMaterialProperty"
  
diff --git a/Engine/Plugins/Experimental/DynamicMaterial/Source/DynamicMaterialEditor/Private/Components/DMMaterialStageBlendFunction.cpp b/Engine/Plugins/Experimental/DynamicMaterial/Source/DynamicMaterialEditor/Private/Components/DMMaterialStageBlendFunction.cpp
--- a/Engine/Plugins/Experimental/DynamicMaterial/Source/DynamicMaterialEditor/Private/Components/DMMaterialStageBlendFunction.cpp
+++ b/Engine/Plugins/Experimental/DynamicMaterial/Source/DynamicMaterialEditor/Private/Components/DMMaterialStageBlendFunction.cpp
@@ -2,12 +2,12 @@
 
 #include "Components/DMMaterialStageBlendFunction.h"
 #include "DMDefs.h"
-#include "DMMaterialFunctionLibrary.h"
 #include "Materials/MaterialExpressionFunctionInput.h"
 #include "Materials/MaterialExpressionMaterialFunctionCall.h"
 #include "Materials/MaterialFunctionInterface.h"
 #include "Model/DMMaterialBuildState.h"
 #include "Model/DMMaterialBuildUtils.h"
+#include "Utils/DMMaterialFunctionLibrary.h"
 
 #define LOCTEXT_NAMESPACE "DMMaterialStageBlendFunction"
 
diff --git a/Engine/Plugins/Experimental/DynamicMaterial/Source/DynamicMaterialEditor/Private/Components/DMMaterialStageFunction.cpp b/Engine/Plugins/Experimental/DynamicMaterial/Source/DynamicMaterialEditor/Private/Components/DMMaterialStageFunction.cpp
--- a/Engine/Plugins/Experimental/DynamicMaterial/Source/DynamicMaterialEditor/Private/Components/DMMaterialStageFunction.cpp
+++ b/Engine/Plugins/Experimental/DynamicMaterial/Source/DynamicMaterialEditor/Private/Components/DMMaterialStageFunction.cpp
@@ -11,7 +11,6 @@
 #include "Components/MaterialValues/DMMaterialValueFloat2.h"
 #include "Components/MaterialValues/DMMaterialValueFloat3XYZ.h"
 #include "Components/MaterialValues/DMMaterialValueFloat4.h"
-#include "DMPrivate.h"
 #include "DynamicMaterialEditorModule.h"
 #include "Materials/MaterialExpressionFunctionInput.h"
 #include "Materials/MaterialExpressionFunctionOutput.h"
@@ -21,6 +20,7 @@
 #include "Model/DMMaterialBuildState.h"
 #include "Model/DMMaterialBuildUtils.h"
 #include "Model/DynamicMaterialModelEditorOnlyData.h"
+#include "Utils/DMPrivate.h"
 
 #define LOCTEXT_NAMESPACE "DMMaterialStageFunction"
 
diff --git a/Engine/Plugins/Experimental/DynamicMaterial/Source/DynamicMaterialEditor/Private/Components/DMMaterialStageThroughput.cpp b/Engine/Plugins/Experimental/DynamicMaterial/Source/DynamicMaterialEditor/Private/Components/DMMaterialStageThroughput.cpp
--- a/Engine/Plugins/Experimental/DynamicMaterial/Source/DynamicMaterialEditor/Private/Components/DMMaterialStageThroughput.cpp
+++ b/Engine/Plugins/Experimental/DynamicMaterial/Source/DynamicMaterialEditor/Private/Components/DMMaterialStageThroughput.cpp
@@ -11,9 +11,6 @@
 #include "Components/MaterialStageInputs/DMMSIValue.h"
 #include "Containers/ArrayView.h"
 #include "DMDefs.h"
-#include "DMMaterialFunctionLibrary.h"
-#include "DMPrivate.h"
-#include "Helpers/DMInputNodeBuilder.h"
 #include "Materials/MaterialExpression.h"
 #include "Materials/MaterialExpressionAppendVector.h"
 #include "Materials/MaterialExpressionComponentMask.h"
@@ -25,6 +22,9 @@
 #include "Model/DMMaterialBuildState.h"
 #include "Model/DMMaterialBuildUtils.h"
 #include "Model/DynamicMaterialModelEditorOnlyData.h"
+#include "Utils/DMInputNodeBuilder.h"
+#include "Utils/DMMaterialFunctionLibrary.h"
+#include "Utils/DMPrivate.h"
 
 #define LOCTEXT_NAMESPACE "DMMaterialStageThroughput"
 
diff --git a/Engine/Plugins/Experimental/DynamicMaterial/Source/DynamicMaterialEditor/Private/Components/DMMaterialStageThroughputLayerBlend.cpp b/Engine/Plugins/Experimental/DynamicMaterial/Source/DynamicMaterialEditor/Private/Components/DMMaterialStageThroughputLayerBlend.cpp
--- a/Engine/Plugins/Experimental/DynamicMaterial/Source/DynamicMaterialEditor/Private/Components/DMMaterialStageThroughputLayerBlend.cpp
+++ b/Engine/Plugins/Experimental/DynamicMaterial/Source/DynamicMaterialEditor/Private/Components/DMMaterialStageThroughputLayerBlend.cpp
@@ -10,10 +10,7 @@
 #include "Components/MaterialStageInputs/DMMSIValue.h"
 #include "Components/MaterialValues/DMMaterialValueTexture.h"
 #include "DMDefs.h"
-#include "DMMaterialFunctionLibrary.h"
-#include "DMPrivate.h"
 #include "DynamicMaterialEditorSettings.h"
-#include "Helpers/DMInputNodeBuilder.h"
 #include "Materials/MaterialExpressionComponentMask.h"
 #include "Materials/MaterialExpressionConstant3Vector.h"
 #include "Materials/MaterialExpressionFunctionInput.h"
@@ -25,6 +22,9 @@
 #include "Model/DMMaterialBuildUtils.h"
 #include "Model/DynamicMaterialModel.h"
 #include "Model/DynamicMaterialModelEditorOnlyData.h"
+#include "Utils/DMInputNodeBuilder.h"
+#include "Utils/DMMaterialFunctionLibrary.h"
+#include "Utils/DMPrivate.h"
  
 #define LOCTEXT_NAMESPACE "DMMaterialStageThroughputLayerBlend"
 
diff --git a/Engine/Plugins/Experimental/DynamicMaterial/Source/DynamicMaterialEditor/Private/Components/DMMaterialSubStage.cpp b/Engine/Plugins/Experimental/DynamicMaterial/Source/DynamicMaterialEditor/Private/Components/DMMaterialSubStage.cpp
--- a/Engine/Plugins/Experimental/DynamicMaterial/Source/DynamicMaterialEditor/Private/Components/DMMaterialSubStage.cpp
+++ b/Engine/Plugins/Experimental/DynamicMaterial/Source/DynamicMaterialEditor/Private/Components/DMMaterialSubStage.cpp
@@ -4,7 +4,7 @@
 
 #include "Components/DMMaterialLayer.h"
 #include "Components/MaterialStageInputs/DMMSIThroughput.h"
-#include "DMPrivate.h"
+#include "Utils/DMPrivate.h"
 
 UDMMaterialSubStage* UDMMaterialSubStage::CreateMaterialSubStage(UDMMaterialStage* InParentStage)
 {
diff --git a/Engine/Plugins/Experimental/DynamicMaterial/Source/DynamicMaterialEditor/Private/Components/MaterialStageExpressions/DMMSEMathBase.cpp b/Engine/Plugins/Experimental/DynamicMaterial/Source/DynamicMaterialEditor/Private/Components/MaterialStageExpressions/DMMSEMathBase.cpp
--- a/Engine/Plugins/Experimental/DynamicMaterial/Source/DynamicMaterialEditor/Private/Components/MaterialStageExpressions/DMMSEMathBase.cpp
+++ b/Engine/Plugins/Experimental/DynamicMaterial/Source/DynamicMaterialEditor/Private/Components/MaterialStageExpressions/DMMSEMathBase.cpp
@@ -5,12 +5,12 @@
 #include "Components/DMMaterialSlot.h"
 #include "Components/DMMaterialStage.h"
 #include "Components/DMMaterialStageInput.h"
-#include "DMPrivate.h"
-#include "Helpers/DMInputNodeBuilder.h"
 #include "Materials/MaterialExpression.h"
 #include "Model/DMMaterialBuildState.h"
 #include "Model/DMMaterialBuildUtils.h"
 #include "Model/DynamicMaterialModelEditorOnlyData.h"
+#include "Utils/DMInputNodeBuilder.h"
+#include "Utils/DMPrivate.h"
 
 #define LOCTEXT_NAMESPACE "DMMaterialStageExpressionMathBase"
 
diff --git a/Engine/Plugins/Experimental/DynamicMaterial/Source/DynamicMaterialEditor/Private/Components/MaterialStageExpressions/DMMSETextureSampleEdgeColor.cpp b/Engine/Plugins/Experimental/DynamicMaterial/Source/DynamicMaterialEditor/Private/Components/MaterialStageExpressions/DMMSETextureSampleEdgeColor.cpp
--- a/Engine/Plugins/Experimental/DynamicMaterial/Source/DynamicMaterialEditor/Private/Components/MaterialStageExpressions/DMMSETextureSampleEdgeColor.cpp
+++ b/Engine/Plugins/Experimental/DynamicMaterial/Source/DynamicMaterialEditor/Private/Components/MaterialStageExpressions/DMMSETextureSampleEdgeColor.cpp
@@ -7,8 +7,8 @@
 #include "Components/MaterialStageInputs/DMMSIExpression.h"
 #include "Components/MaterialStageInputs/DMMSIValue.h"
 #include "Components/MaterialValues/DMMaterialValueFloat2.h"
-#include "DMPrivate.h"
 #include "Materials/MaterialExpressionTextureSample.h"
+#include "Utils/DMPrivate.h"
 
 #define LOCTEXT_NAMESPACE "DMMaterialStageExpressionTextureSample"
 
diff --git a/Engine/Plugins/Experimental/DynamicMaterial/Source/DynamicMaterialEditor/Private/Components/MaterialStageGradients/DMMSGLinear.cpp b/Engine/Plugins/Experimental/DynamicMaterial/Source/DynamicMaterialEditor/Private/Components/MaterialStageGradients/DMMSGLinear.cpp
--- a/Engine/Plugins/Experimental/DynamicMaterial/Source/DynamicMaterialEditor/Private/Components/MaterialStageGradients/DMMSGLinear.cpp
+++ b/Engine/Plugins/Experimental/DynamicMaterial/Source/DynamicMaterialEditor/Private/Components/MaterialStageGradients/DMMSGLinear.cpp
@@ -2,7 +2,6 @@
 
 #include "Components/MaterialStageGradients/DMMSGLinear.h"
 #include "DMDefs.h"
-#include "DMMaterialFunctionLibrary.h"
 #include "Materials/MaterialExpressionComponentMask.h"
 #include "Materials/MaterialExpressionConstant.h"
 #include "Materials/MaterialExpressionFmod.h"
@@ -11,6 +10,7 @@
 #include "Materials/MaterialExpressionSubtract.h"
 #include "Model/DMMaterialBuildState.h"
 #include "Model/DMMaterialBuildUtils.h"
+#include "Utils/DMMaterialFunctionLibrary.h"
 
 #define LOCTEXT_NAMESPACE "DMMaterialStageGradientLinear"
 
diff --git a/Engine/Plugins/Experimental/DynamicMaterial/Source/DynamicMaterialEditor/Private/Components/MaterialStageGradients/DMMSGRadial.cpp b/Engine/Plugins/Experimental/DynamicMaterial/Source/DynamicMaterialEditor/Private/Components/MaterialStageGradients/DMMSGRadial.cpp
--- a/Engine/Plugins/Experimental/DynamicMaterial/Source/DynamicMaterialEditor/Private/Components/MaterialStageGradients/DMMSGRadial.cpp
+++ b/Engine/Plugins/Experimental/DynamicMaterial/Source/DynamicMaterialEditor/Private/Components/MaterialStageGradients/DMMSGRadial.cpp
@@ -1,7 +1,7 @@
 // Copyright Epic Games, Inc. All Rights Reserved.
 
 #include "Components/MaterialStageGradients/DMMSGRadial.h"
-#include "DMMaterialFunctionLibrary.h"
+#include "Utils/DMMaterialFunctionLibrary.h"
 #include "Materials/MaterialExpressionMaterialFunctionCall.h"
 #include "Model/DMMaterialBuildState.h"
 #include "Model/DMMaterialBuildUtils.h"
diff --git a/Engine/Plugins/Experimental/DynamicMaterial/Source/DynamicMaterialEditor/Private/Components/MaterialStageInputs/DMMSISlot.cpp b/Engine/Plugins/Experimental/DynamicMaterial/Source/DynamicMaterialEditor/Private/Components/MaterialStageInputs/DMMSISlot.cpp
--- a/Engine/Plugins/Experimental/DynamicMaterial/Source/DynamicMaterialEditor/Private/Components/MaterialStageInputs/DMMSISlot.cpp
+++ b/Engine/Plugins/Experimental/DynamicMaterial/Source/DynamicMaterialEditor/Private/Components/MaterialStageInputs/DMMSISlot.cpp
@@ -7,11 +7,11 @@
 #include "Components/DMMaterialStage.h"
 #include "Components/DMMaterialStageThroughput.h"
 #include "DMComponentPath.h"
-#include "DMPrivate.h"
 #include "DMValueDefinition.h"
 #include "Model/DMMaterialBuildState.h"
 #include "Model/DynamicMaterialModel.h"
 #include "Model/DynamicMaterialModelEditorOnlyData.h"
+#include "Utils/DMPrivate.h"
 
 #define LOCTEXT_NAMESPACE "DMMaterialStageInputSlot"
 
diff --git a/Engine/Plugins/Experimental/DynamicMaterial/Source/DynamicMaterialEditor/Private/Components/MaterialStageInputs/DMMSITextureUV.cpp b/Engine/Plugins/Experimental/DynamicMaterial/Source/DynamicMaterialEditor/Private/Components/MaterialStageInputs/DMMSITextureUV.cpp
--- a/Engine/Plugins/Experimental/DynamicMaterial/Source/DynamicMaterialEditor/Private/Components/MaterialStageInputs/DMMSITextureUV.cpp
+++ b/Engine/Plugins/Experimental/DynamicMaterial/Source/DynamicMaterialEditor/Private/Components/MaterialStageInputs/DMMSITextureUV.cpp
@@ -11,7 +11,6 @@
 #include "Components/DMTextureUV.h"
 #include "DMComponentPath.h"
 #include "DMDefs.h"
-#include "DMMaterialFunctionLibrary.h"
 #include "DynamicMaterialEditorModule.h"
 #include "DynamicMaterialModule.h"
 #include "Materials/MaterialExpressionComponentMask.h"
@@ -23,6 +22,7 @@
 #include "Model/DMMaterialBuildUtils.h"
 #include "Model/DynamicMaterialModel.h"
 #include "Model/DynamicMaterialModelEditorOnlyData.h"
+#include "Utils/DMMaterialFunctionLibrary.h"
 
 #define LOCTEXT_NAMESPACE "DMMaterialStageInputTextureUV"
 
diff --git a/Engine/Plugins/Experimental/DynamicMaterial/Source/DynamicMaterialEditor/Private/Components/MaterialStageInputs/DMMSIThroughput.cpp b/Engine/Plugins/Experimental/DynamicMaterial/Source/DynamicMaterialEditor/Private/Components/MaterialStageInputs/DMMSIThroughput.cpp
--- a/Engine/Plugins/Experimental/DynamicMaterial/Source/DynamicMaterialEditor/Private/Components/MaterialStageInputs/DMMSIThroughput.cpp
+++ b/Engine/Plugins/Experimental/DynamicMaterial/Source/DynamicMaterialEditor/Private/Components/MaterialStageInputs/DMMSIThroughput.cpp
@@ -6,12 +6,12 @@
 #include "Components/DMMaterialStageThroughput.h"
 #include "Components/DMMaterialSubStage.h"
 #include "DMComponentPath.h"
-#include "DMPrivate.h"
 #include "Materials/MaterialExpressionComponentMask.h"
 #include "Model/DMMaterialBuildState.h"
 #include "Model/DMMaterialBuildUtils.h"
 #include "Model/DynamicMaterialModel.h"
 #include "Model/DynamicMaterialModelEditorOnlyData.h"
+#include "Utils/DMPrivate.h"
 
 #define LOCTEXT_NAMESPACE "DMMaterialStageInputThroughput"
 
diff --git a/Engine/Plugins/Experimental/DynamicMaterial/Source/DynamicMaterialEditor/Private/DMBlueprintFunctionLibrary.cpp b/Engine/Plugins/Experimental/DynamicMaterial/Source/DynamicMaterialEditor/Private/DMBlueprintFunctionLibrary.cpp
deleted file mode 100644
diff --git a/Engine/Plugins/Experimental/DynamicMaterial/Source/DynamicMaterialEditor/Private/DMBlueprintFunctionLibrary.h b/Engine/Plugins/Experimental/DynamicMaterial/Source/DynamicMaterialEditor/Private/DMBlueprintFunctionLibrary.h
deleted file mode 100644
diff --git a/Engine/Plugins/Experimental/DynamicMaterial/Source/DynamicMaterialEditor/Private/DMMaterialFunctionLibrary.cpp b/Engine/Plugins/Experimental/DynamicMaterial/Source/DynamicMaterialEditor/Private/DMMaterialFunctionLibrary.cpp
deleted file mode 100644
diff --git a/Engine/Plugins/Experimental/DynamicMaterial/Source/DynamicMaterialEditor/Private/DMMaterialFunctionLibrary.h b/Engine/Plugins/Experimental/DynamicMaterial/Source/DynamicMaterialEditor/Private/DMMaterialFunctionLibrary.h
deleted file mode 100644
diff --git a/Engine/Plugins/Experimental/DynamicMaterial/Source/DynamicMaterialEditor/Private/DMPrivate.cpp b/Engine/Plugins/Experimental/DynamicMaterial/Source/DynamicMaterialEditor/Private/DMPrivate.cpp
deleted file mode 100644
diff --git a/Engine/Plugins/Experimental/DynamicMaterial/Source/DynamicMaterialEditor/Private/DMPrivate.h b/Engine/Plugins/Experimental/DynamicMaterial/Source/DynamicMaterialEditor/Private/DMPrivate.h
deleted file mode 100644
diff --git a/Engine/Plugins/Experimental/DynamicMaterial/Source/DynamicMaterialEditor/Private/DynamicMaterialEditorModule.cpp b/Engine/Plugins/Experimental/DynamicMaterial/Source/DynamicMaterialEditor/Private/DynamicMaterialEditorModule.cpp
--- a/Engine/Plugins/Experimental/DynamicMaterial/Source/DynamicMaterialEditor/Private/DynamicMaterialEditorModule.cpp
+++ b/Engine/Plugins/Experimental/DynamicMaterial/Source/DynamicMaterialEditor/Private/DynamicMaterialEditorModule.cpp
@@ -21,7 +21,6 @@
 #include "DetailsPanel/DMPropertyTypeCustomizer.h"
 #include "DetailsPanel/DMValueDetailsRowExtensions.h"
 #include "DetailsPanel/Slate/SDMMaterialListExtensionWidget.h"
-#include "DMMaterialFunctionLibrary.h"
 #include "DMWorldSubsystem.h"
 #include "DynamicMaterialEditorCommands.h"
 #include "DynamicMaterialEditorSettings.h"
@@ -55,6 +54,7 @@
 #include "Slate/Properties/Generators/DMTextureUVPropertyRowGenerator.h"
 #include "Slate/Properties/Generators/DMThroughputPropertyRowGenerator.h"
 #include "Slate/SDMEditor.h"
+#include "Utils/DMMaterialFunctionLibrary.h"
 
 DEFINE_LOG_CATEGORY(LogDynamicMaterialEditor);
 
diff --git a/Engine/Plugins/Experimental/DynamicMaterial/Source/DynamicMaterialEditor/Private/Helpers/DMInputNodeBuilder.cpp b/Engine/Plugins/Experimental/DynamicMaterial/Source/DynamicMaterialEditor/Private/Helpers/DMInputNodeBuilder.cpp
deleted file mode 100644
diff --git a/Engine/Plugins/Experimental/DynamicMaterial/Source/DynamicMaterialEditor/Private/Helpers/DMInputNodeBuilder.h b/Engine/Plugins/Experimental/DynamicMaterial/Source/DynamicMaterialEditor/Private/Helpers/DMInputNodeBuilder.h
deleted file mode 100644
diff --git a/Engine/Plugins/Experimental/DynamicMaterial/Source/DynamicMaterialEditor/Private/Helpers/DMMaterialSnapshotLibrary.cpp b/Engine/Plugins/Experimental/DynamicMaterial/Source/DynamicMaterialEditor/Private/Helpers/DMMaterialSnapshotLibrary.cpp
deleted file mode 100644
diff --git a/Engine/Plugins/Experimental/DynamicMaterial/Source/DynamicMaterialEditor/Private/Helpers/DMMaterialSnapshotLibrary.h b/Engine/Plugins/Experimental/DynamicMaterial/Source/DynamicMaterialEditor/Private/Helpers/DMMaterialSnapshotLibrary.h
deleted file mode 100644
diff --git a/Engine/Plugins/Experimental/DynamicMaterial/Source/DynamicMaterialEditor/Private/Material/DynamicMaterialInstanceFactory.cpp b/Engine/Plugins/Experimental/DynamicMaterial/Source/DynamicMaterialEditor/Private/Material/DynamicMaterialInstanceFactory.cpp
--- a/Engine/Plugins/Experimental/DynamicMaterial/Source/DynamicMaterialEditor/Private/Material/DynamicMaterialInstanceFactory.cpp
+++ b/Engine/Plugins/Experimental/DynamicMaterial/Source/DynamicMaterialEditor/Private/Material/DynamicMaterialInstanceFactory.cpp
@@ -2,7 +2,6 @@
 
 #include "Material/DynamicMaterialInstanceFactory.h"
 #include "AssetToolsModule.h"
-#include "DMPrivate.h"
 #include "EngineAnalytics.h"
 #include "GameFramework/Actor.h"
 #include "IAssetTools.h"
@@ -11,6 +10,7 @@
 #include "Model/DynamicMaterialModelFactory.h"
 #include "Model/IDynamicMaterialModelEditorOnlyDataInterface.h"
 #include "Modules/ModuleManager.h"
+#include "Utils/DMPrivate.h"
 
 #define LOCTEXT_NAMESPACE "MaterialDesignerInstanceFactory"
 
diff --git a/Engine/Plugins/Experimental/DynamicMaterial/Source/DynamicMaterialEditor/Private/Menus/DMMaterialSlotLayerAddEffectMenus.cpp b/Engine/Plugins/Experimental/DynamicMaterial/Source/DynamicMaterialEditor/Private/Menus/DMMaterialSlotLayerAddEffectMenus.cpp
--- a/Engine/Plugins/Experimental/DynamicMaterial/Source/DynamicMaterialEditor/Private/Menus/DMMaterialSlotLayerAddEffectMenus.cpp
+++ b/Engine/Plugins/Experimental/DynamicMaterial/Source/DynamicMaterialEditor/Private/Menus/DMMaterialSlotLayerAddEffectMenus.cpp
@@ -7,10 +7,14 @@
 #include "Components/DMMaterialEffectStack.h"
 #include "Components/DMMaterialLayer.h"
 #include "DynamicMaterialEditorSettings.h"
+#include "Framework/Notifications/NotificationManager.h"
 #include "Materials/MaterialFunctionInterface.h"
 #include "Menus/DMMaterialSlotLayerAddEffectContext.h"
 #include "Menus/DMMenuContext.h"
 #include "ToolMenus.h"
+#include "Utils/DMMaterialEffectStackPresetSubsystem.h"
+#include "Widgets/Input/SEditableTextBox.h"
+#include "Widgets/Notifications/SNotificationList.h"
 #include "Widgets/SNullWidget.h"
 
 #define LOCTEXT_NAMESPACE "DMMaterialSlotLayerAddEffectMenus"
@@ -188,6 +192,274 @@ namespace UE::DynamicMaterialEditor::Private
 		}
 	}
 
+	void SavePreset(const FToolMenuContext& InContext, const FString& InPresetName)
+	{
+		const UDMMaterialLayerObject* Layer = nullptr;
+
+		if (UDMMenuContext* MenuContext = InContext.FindContext<UDMMenuContext>())
+		{
+			Layer = MenuContext->GetLayer();
+		}
+		else if (UDMMaterialSlotLayerAddEffectContext* SlotContext = InContext.FindContext<UDMMaterialSlotLayerAddEffectContext>())
+		{
+			Layer = SlotContext->GetLayer();
+		}
+
+		if (!Layer || !IsValid(Layer))
+		{
+			return;
+		}
+
+		UDMMaterialEffectStack* EffectStack = Layer->GetEffectStack();
+
+		if (!EffectStack)
+		{
+			return;
+		}
+
+		FDMMaterialEffectStackJson Preset = EffectStack->CreatePreset();
+
+		UDMMaterialEffectStackPresetSubsystem::Get()->SavePreset(InPresetName, Preset);
+	}
+
+	bool VerifyFileName(const FText& InValue, FText& OutErrorText)
+	{
+		static const FText TooShortError = LOCTEXT("TooShortError", "Min 3 characters.");
+		static const FText TooLongError = LOCTEXT("TooLongError", "Max 50 characters.");
+		static const FText InvalidCharacterError = LOCTEXT("InvalidCharacterError", "Valid characters are A-Z, A-z, space, _ and -");
+
+		static const TArray<TPair<TCHAR, TCHAR>> ValidCharacterRanges = {
+			{'A', 'Z'},
+			{'a', 'z'},
+			{'0', '9'},
+			{' ', ' '},
+			{'-', '-'},
+			{'_', '_'}
+		};
+
+		const FString ValueStr = InValue.ToString();
+		const int32 ValueLen = ValueStr.Len();
+
+		if (ValueLen < 3)
+		{
+			OutErrorText = TooShortError;
+			return false;
+		}
+
+		if (ValueLen > 50)
+		{
+			OutErrorText = TooLongError;
+			return false;
+		}
+
+		for (int32 Index = 0; Index < ValueLen; ++Index)
+		{
+			bool bInRange = false;
+
+			for (const TPair<TCHAR, TCHAR>& ValidCharacterRange : ValidCharacterRanges)
+			{
+				if (ValueStr[Index] >= ValidCharacterRange.Key && ValueStr[Index] <= ValidCharacterRange.Value)
+				{
+					bInRange = true;
+					break;
+				}
+			}
+
+			if (!bInRange)
+			{
+				OutErrorText = InvalidCharacterError;
+				return false;
+			}
+		}
+
+		return true;
+	}
+
+	void GenerateSaveEffectsMenu(FToolMenuSection& InSection)
+	{
+		FNewToolMenuChoice RenameChoice(FNewToolMenuWidget::CreateLambda(
+			[](const FToolMenuContext& InContext) -> TSharedRef<SWidget>
+			{
+				const UDMMaterialLayerObject* Layer = nullptr;
+
+				if (UDMMenuContext* MenuContext = InContext.FindContext<UDMMenuContext>())
+				{
+					Layer = MenuContext->GetLayer();
+				}
+				else if (UDMMaterialSlotLayerAddEffectContext* SlotContext = InContext.FindContext<UDMMaterialSlotLayerAddEffectContext>())
+				{
+					Layer = SlotContext->GetLayer();
+				}
+
+				if (!Layer || !IsValid(Layer))
+				{
+					return SNullWidget::NullWidget;
+				}
+
+				UDMMaterialEffectStack* EffectStack = Layer->GetEffectStack();
+
+				if (!EffectStack)
+				{
+					return SNullWidget::NullWidget;
+				}
+
+				TWeakObjectPtr<UDMMaterialEffectStack> EffectStackWeak = EffectStack;
+
+				return SNew(SEditableTextBox)
+					.Text(LOCTEXT("NewPreset", "New Preset"))
+					.OnVerifyTextChanged_Static(&VerifyFileName)
+					.AllowContextMenu(false)
+					.ClearKeyboardFocusOnCommit(true)
+					.MinDesiredWidth(100.f)
+					.OnTextCommitted_Lambda(
+						[EffectStackWeak](const FText& InText, ETextCommit::Type InCommitType)
+						{
+							if (InCommitType != ETextCommit::OnEnter)
+							{
+								return;
+							}
+
+							UDMMaterialEffectStack* EffectStack = EffectStackWeak.Get();
+
+							if (!IsValid(EffectStack))
+							{
+								return;
+							}
+
+							FText Unused;
+
+							if (!VerifyFileName(InText, Unused))
+							{
+								return;
+							}
+
+							FDMMaterialEffectStackJson Preset = EffectStack->CreatePreset();
+
+							if (UDMMaterialEffectStackPresetSubsystem::Get()->SavePreset(InText.ToString(), Preset))
+							{
+								FNotificationInfo Info(LOCTEXT("PresetSaved", "Preset saved!"));
+								Info.ExpireDuration = 5.0f;
+								FSlateNotificationManager::Get().AddNotification(Info);
+							}
+							else
+							{
+								FNotificationInfo Info(LOCTEXT("PresetNotSaved", "Failed to save preset!"));
+								Info.ExpireDuration = 5.0f;
+								FSlateNotificationManager::Get().AddNotification(Info);
+							}
+						});
+			}));
+
+		InSection.AddSubMenu(
+			"SavePreset",
+			LOCTEXT("SavePreset", "Save Preset"),
+			TAttribute<FText>(),
+			RenameChoice
+		);
+	}
+
+	void LoadPreset(const FToolMenuContext& InContext, FString InPresetName)
+	{
+		const UDMMaterialLayerObject* Layer = nullptr;
+
+		if (UDMMenuContext* MenuContext = InContext.FindContext<UDMMenuContext>())
+		{
+			Layer = MenuContext->GetLayer();
+		}
+		else if (UDMMaterialSlotLayerAddEffectContext* SlotContext = InContext.FindContext<UDMMaterialSlotLayerAddEffectContext>())
+		{
+			Layer = SlotContext->GetLayer();
+		}
+
+		if (!Layer || !IsValid(Layer))
+		{
+			return;
+		}
+
+		UDMMaterialEffectStack* EffectStack = Layer->GetEffectStack();
+
+		if (!EffectStack)
+		{
+			return;
+		}
+
+		FDMMaterialEffectStackJson Preset;
+
+		if (UDMMaterialEffectStackPresetSubsystem::Get()->LoadPreset(InPresetName, Preset))
+		{
+			EffectStack->ApplyPreset(Preset);
+
+			FNotificationInfo Info(LOCTEXT("PresetApplied", "Preset applied!"));
+			Info.ExpireDuration = 5.0f;
+			FSlateNotificationManager::Get().AddNotification(Info);
+		}
+		else
+		{
+			FNotificationInfo Info(LOCTEXT("PresetNotApplied", "Failed to apply preset!"));
+			Info.ExpireDuration = 5.0f;
+			FSlateNotificationManager::Get().AddNotification(Info);
+		}
+	}
+
+	void GenerateLoadEffectsMenu(UToolMenu* InMenu)
+	{
+		if (!InMenu)
+		{
+			return;
+		}
+
+		const UDMMaterialEffectStackPresetSubsystem* PresetSubsystem = UDMMaterialEffectStackPresetSubsystem::Get();
+
+		if (!PresetSubsystem)
+		{
+			return;
+		}
+
+		FToolMenuSection& Section = InMenu->AddSection("LoadPreset", LOCTEXT("LoadPreset", "Load Preset"));
+
+		TArray<FString> PresetList = PresetSubsystem->GetPresetNames();
+
+		for (int32 PresetIndex = 0; PresetIndex < PresetList.Num(); ++PresetIndex)
+		{
+			FToolUIAction LoadAction;
+			LoadAction.ExecuteAction = FToolMenuExecuteAction::CreateStatic(&LoadPreset, PresetList[PresetIndex]);
+
+			Section.AddMenuEntry(
+				*PresetList[PresetIndex],
+				FText::FromString(PresetList[PresetIndex]),
+				TAttribute<FText>(),
+				FSlateIcon(),
+				FToolUIActionChoice(LoadAction)
+			);
+		}
+	}
+
+	void GenerateEffectPresetMenu(UToolMenu* InMenu)
+	{
+		if (!InMenu)
+		{
+			return;
+		}
+
+		const UDynamicMaterialEditorSettings* Settings = GetDefault<UDynamicMaterialEditorSettings>();
+
+		if (!Settings)
+		{
+			return;
+		}
+
+		FToolMenuSection& Section = InMenu->AddSection("Presets", LOCTEXT("Presets", "Presets"));
+
+		GenerateSaveEffectsMenu(Section);
+
+		Section.AddSubMenu(
+			"LoadPreset",
+			LOCTEXT("LoadPreset", "Load Preset"),
+			TAttribute<FText>(),
+			FNewToolMenuChoice(FNewToolMenuDelegate::CreateStatic(&GenerateLoadEffectsMenu))
+		);
+	}
+
 	void RegisterAddEffectMenu()
 	{
 		using namespace UE::DynamicMaterialEditor::Private;
@@ -203,6 +475,7 @@ namespace UE::DynamicMaterialEditor::Private
 			[](UToolMenu* InMenu)
 			{
 				GenerateAddEffectMenu(InMenu);
+				GenerateEffectPresetMenu(InMenu);
 			}));
 	}
 }
diff --git a/Engine/Plugins/Experimental/DynamicMaterial/Source/DynamicMaterialEditor/Private/Menus/DMMaterialStageInputMenus.cpp b/Engine/Plugins/Experimental/DynamicMaterial/Source/DynamicMaterialEditor/Private/Menus/DMMaterialStageInputMenus.cpp
--- a/Engine/Plugins/Experimental/DynamicMaterial/Source/DynamicMaterialEditor/Private/Menus/DMMaterialStageInputMenus.cpp
+++ b/Engine/Plugins/Experimental/DynamicMaterial/Source/DynamicMaterialEditor/Private/Menus/DMMaterialStageInputMenus.cpp
@@ -16,7 +16,6 @@
 #include "Components/MaterialStageInputs/DMMSISlot.h"
 #include "Components/MaterialStageInputs/DMMSITextureUV.h"
 #include "Components/MaterialStageInputs/DMMSIValue.h"
-#include "DMPrivate.h"
 #include "DMValueDefinition.h"
 #include "DynamicMaterialEditorModule.h"
 #include "Framework/MultiBox/MultiBoxBuilder.h"
@@ -25,6 +24,7 @@
 #include "Model/DynamicMaterialModel.h"
 #include "Model/DynamicMaterialModelEditorOnlyData.h"
 #include "ScopedTransaction.h"
+#include "Utils/DMPrivate.h"
 #include "Widgets/SNullWidget.h"
 
 #define LOCTEXT_NAMESPACE "FDMMaterialStageInputMenus"
diff --git a/Engine/Plugins/Experimental/DynamicMaterial/Source/DynamicMaterialEditor/Private/Menus/DMMaterialStageSourceMenus.cpp b/Engine/Plugins/Experimental/DynamicMaterial/Source/DynamicMaterialEditor/Private/Menus/DMMaterialStageSourceMenus.cpp
--- a/Engine/Plugins/Experimental/DynamicMaterial/Source/DynamicMaterialEditor/Private/Menus/DMMaterialStageSourceMenus.cpp
+++ b/Engine/Plugins/Experimental/DynamicMaterial/Source/DynamicMaterialEditor/Private/Menus/DMMaterialStageSourceMenus.cpp
@@ -24,7 +24,6 @@
 #include "Components/MaterialStageInputs/DMMSIValue.h"
 #include "Components/RenderTargetRenderers/DMRenderTargetTextRenderer.h"
 #include "Components/RenderTargetRenderers/DMRenderTargetWidgetRenderer.h"
-#include "DMBlueprintFunctionLibrary.h"
 #include "DMDefs.h"
 #include "DMValueDefinition.h"
 #include "DynamicMaterialEditorModule.h"
@@ -37,6 +36,7 @@
 #include "Slate/SDMStage.h"
 #include "Templates/SharedPointer.h"
 #include "ToolMenus.h"
+#include "Utils/DMBlueprintFunctionLibrary.h"
 #include "Widgets/SNullWidget.h"
 
 #define LOCTEXT_NAMESPACE "FDMMaterialStageSourceMenus"
diff --git a/Engine/Plugins/Experimental/DynamicMaterial/Source/DynamicMaterialEditor/Private/Menus/DMToolBarMenus.cpp b/Engine/Plugins/Experimental/DynamicMaterial/Source/DynamicMaterialEditor/Private/Menus/DMToolBarMenus.cpp
--- a/Engine/Plugins/Experimental/DynamicMaterial/Source/DynamicMaterialEditor/Private/Menus/DMToolBarMenus.cpp
+++ b/Engine/Plugins/Experimental/DynamicMaterial/Source/DynamicMaterialEditor/Private/Menus/DMToolBarMenus.cpp
@@ -3,8 +3,6 @@
 #include "DMToolBarMenus.h"
 #include "ContentBrowserModule.h"
 #include "DesktopPlatformModule.h"
-#include "DMBlueprintFunctionLibrary.h"
-#include "DMPrivate.h"
 #include "DynamicMaterialEditorModule.h"
 #include "DynamicMaterialEditorSettings.h"
 #include "Engine/Engine.h"
@@ -24,7 +22,9 @@
 #include "Slate/SDMSlot.h"
 #include "ToolMenu.h"
 #include "ToolMenus.h"
-#include "Helpers/DMMaterialSnapshotLibrary.h"
+#include "Utils/DMBlueprintFunctionLibrary.h"
+#include "Utils/DMMaterialSnapshotLibrary.h"
+#include "Utils/DMPrivate.h"
 #include "Widgets/SWidget.h"
 
 #define LOCTEXT_NAMESPACE "FDMToolBarMenus"
diff --git a/Engine/Plugins/Experimental/DynamicMaterial/Source/DynamicMaterialEditor/Private/Model/DMMaterialBuildUtils.cpp b/Engine/Plugins/Experimental/DynamicMaterial/Source/DynamicMaterialEditor/Private/Model/DMMaterialBuildUtils.cpp
--- a/Engine/Plugins/Experimental/DynamicMaterial/Source/DynamicMaterialEditor/Private/Model/DMMaterialBuildUtils.cpp
+++ b/Engine/Plugins/Experimental/DynamicMaterial/Source/DynamicMaterialEditor/Private/Model/DMMaterialBuildUtils.cpp
@@ -1,19 +1,19 @@
 // Copyright Epic Games, Inc. All Rights Reserved.
 
 #include "Model/DMMaterialBuildUtils.h"
-#include "Model/DMMaterialBuildState.h"
-#include "Materials/MaterialExpressionTextureSample.h"
-#include "Materials/MaterialExpressionTextureObjectParameter.h"
-#include "Materials/MaterialExpressionTextureObject.h"
-#include "Materials/MaterialExpressionConstant4Vector.h"
-#include "Materials/MaterialExpressionConstant.h"
-#include "Materials/MaterialExpressionComponentMask.h"
-#include "Materials/MaterialExpressionAppendVector.h"
-#include "Materials/MaterialExpression.h"
-#include "Materials/Material.h"
-#include "MaterialEditingLibrary.h"
-#include "DMPrivate.h"
 #include "Components/DMMaterialStageInput.h"
+#include "MaterialEditingLibrary.h"
+#include "Materials/Material.h"
+#include "Materials/MaterialExpression.h"
+#include "Materials/MaterialExpressionAppendVector.h"
+#include "Materials/MaterialExpressionComponentMask.h"
+#include "Materials/MaterialExpressionConstant.h"
+#include "Materials/MaterialExpressionConstant4Vector.h"
+#include "Materials/MaterialExpressionTextureObject.h"
+#include "Materials/MaterialExpressionTextureObjectParameter.h"
+#include "Materials/MaterialExpressionTextureSample.h"
+#include "Model/DMMaterialBuildState.h"
+#include "Utils/DMPrivate.h"
 
 FDMMaterialBuildUtils::FDMMaterialBuildUtils(FDMMaterialBuildState& InBuildState)
 	: BuildState(InBuildState)
diff --git a/Engine/Plugins/Experimental/DynamicMaterial/Source/DynamicMaterialEditor/Private/Model/DynamicMaterialModelEditorOnlyData.cpp b/Engine/Plugins/Experimental/DynamicMaterial/Source/DynamicMaterialEditor/Private/Model/DynamicMaterialModelEditorOnlyData.cpp
--- a/Engine/Plugins/Experimental/DynamicMaterial/Source/DynamicMaterialEditor/Private/Model/DynamicMaterialModelEditorOnlyData.cpp
+++ b/Engine/Plugins/Experimental/DynamicMaterial/Source/DynamicMaterialEditor/Private/Model/DynamicMaterialModelEditorOnlyData.cpp
@@ -30,8 +30,6 @@
 #include "CoreGlobals.h"
 #include "DMComponentPath.h"
 #include "DMDefs.h"
-#include "DMMaterialFunctionLibrary.h"
-#include "DMPrivate.h"
 #include "DynamicMaterialEditorModule.h"
 #include "DynamicMaterialEditorSettings.h"
 #include "DynamicMaterialModule.h"
@@ -49,6 +47,8 @@
 #include "Model/DMMaterialBuildUtils.h"
 #include "Model/DynamicMaterialModel.h"
 #include "UObject/Package.h"
+#include "Utils/DMMaterialFunctionLibrary.h"
+#include "Utils/DMPrivate.h"
 
 #define LOCTEXT_NAMESPACE "MaterialDesignerModel"
 
diff --git a/Engine/Plugins/Experimental/DynamicMaterial/Source/DynamicMaterialEditor/Private/Slate/Layers/SDMLayerEffectsItem.cpp b/Engine/Plugins/Experimental/DynamicMaterial/Source/DynamicMaterialEditor/Private/Slate/Layers/SDMLayerEffectsItem.cpp
--- a/Engine/Plugins/Experimental/DynamicMaterial/Source/DynamicMaterialEditor/Private/Slate/Layers/SDMLayerEffectsItem.cpp
+++ b/Engine/Plugins/Experimental/DynamicMaterial/Source/DynamicMaterialEditor/Private/Slate/Layers/SDMLayerEffectsItem.cpp
@@ -4,11 +4,11 @@
 #include "Components/DMMaterialEffect.h"
 #include "Components/DMMaterialEffectStack.h"
 #include "Components/DMMaterialStage.h"
-#include "DMPrivate.h"
 #include "DragDrop/DMLayerEffectsDragDropOperation.h"
 #include "DynamicMaterialEditorStyle.h"
 #include "Slate/SDMStage.h"
 #include "SlateOptMacros.h"
+#include "Utils/DMPrivate.h"
 #include "Widgets/Images/SImage.h"
 #include "Widgets/Input/SButton.h"
 #include "Widgets/Layout/SExpandableArea.h"
diff --git a/Engine/Plugins/Experimental/DynamicMaterial/Source/DynamicMaterialEditor/Private/Slate/Layers/SDMSlotLayerItem.cpp b/Engine/Plugins/Experimental/DynamicMaterial/Source/DynamicMaterialEditor/Private/Slate/Layers/SDMSlotLayerItem.cpp
--- a/Engine/Plugins/Experimental/DynamicMaterial/Source/DynamicMaterialEditor/Private/Slate/Layers/SDMSlotLayerItem.cpp
+++ b/Engine/Plugins/Experimental/DynamicMaterial/Source/DynamicMaterialEditor/Private/Slate/Layers/SDMSlotLayerItem.cpp
@@ -2,14 +2,15 @@
 
 #include "Slate/Layers/SDMSlotLayerItem.h"
 
-#include "ContentBrowserDataDragDropOp.h"
+#include "Components/DMMaterialEffectFunction.h"
+#include "Components/DMMaterialEffectStack.h"
 #include "Components/DMMaterialLayer.h"
 #include "Components/DMMaterialProperty.h"
 #include "Components/DMMaterialSlot.h"
 #include "Components/DMMaterialStage.h"
 #include "Components/DMMaterialStageThroughput.h"
 #include "Components/MaterialStageInputs/DMMSITextureUV.h"
-#include "DMPrivate.h"
+#include "ContentBrowserDataDragDropOp.h"
 #include "DragDrop/DMSlotLayerDragDropOperation.h"
 #include "DynamicMaterialEditorSettings.h"
 #include "DynamicMaterialEditorStyle.h"
@@ -21,9 +22,8 @@
 #include "Slate/SDMSlot.h"
 #include "Slate/SDMStage.h"
 #include "SlateOptMacros.h"
-#include "Components/DMMaterialEffectFunction.h"
-#include "Components/DMMaterialEffectStack.h"
 #include "Styling/StyleColors.h"
+#include "Utils/DMPrivate.h"
 #include "Widgets/Images/SImage.h"
 #include "Widgets/Input/SButton.h"
 #include "Widgets/Input/SEditableTextBox.h"
@@ -726,9 +726,11 @@ FReply SDMSlotLayerItem::OnLayerItemAcceptDrop(const FDragDropEvent& InDragDropE
 				if (UMaterialFunctionInterface* MaterialFunction = Cast<UMaterialFunctionInterface>(DraggedAsset.GetAsset()))
 				{
 					UDMMaterialEffectFunction* EffectFunction = UDMMaterialEffect::CreateEffect<UDMMaterialEffectFunction>(EffectStack);
-					EffectFunction->SetMaterialFunction(MaterialFunction);
 
 					// Will successfully set the function if it's valid
+					EffectFunction->SetMaterialFunction(MaterialFunction);
+
+					// If it was valid, its value will match.
 					if (EffectFunction->GetMaterialFunction() == MaterialFunction)
 					{
 						EffectStack->AddEffect(EffectFunction);
diff --git a/Engine/Plugins/Experimental/DynamicMaterial/Source/DynamicMaterialEditor/Private/Slate/Layers/SDMSlotLayerView.cpp b/Engine/Plugins/Experimental/DynamicMaterial/Source/DynamicMaterialEditor/Private/Slate/Layers/SDMSlotLayerView.cpp
--- a/Engine/Plugins/Experimental/DynamicMaterial/Source/DynamicMaterialEditor/Private/Slate/Layers/SDMSlotLayerView.cpp
+++ b/Engine/Plugins/Experimental/DynamicMaterial/Source/DynamicMaterialEditor/Private/Slate/Layers/SDMSlotLayerView.cpp
@@ -13,7 +13,6 @@
 #include "Components/MaterialStageInputs/DMMSIFunction.h"
 #include "Components/MaterialStageInputs/DMMSIValue.h"
 #include "Components/MaterialValues/DMMaterialValueTexture.h"
-#include "DMPrivate.h"
 #include "DragAndDrop/AssetDragDropOp.h"
 #include "DynamicMaterialEditorCommands.h"
 #include "DynamicMaterialEditorModule.h"
@@ -25,6 +24,7 @@
 #include "Slate/Layers/SDMSlotLayerItem.h"
 #include "Slate/SDMSlot.h"
 #include "ToolMenus.h"
+#include "Utils/DMPrivate.h"
 #include "Widgets/SNullWidget.h"
 
 #define LOCTEXT_NAMESPACE "SDMSlotLayerView"
diff --git a/Engine/Plugins/Experimental/DynamicMaterial/Source/DynamicMaterialEditor/Private/Slate/Properties/Editors/SDMPropertyEditFloat1Value.cpp b/Engine/Plugins/Experimental/DynamicMaterial/Source/DynamicMaterialEditor/Private/Slate/Properties/Editors/SDMPropertyEditFloat1Value.cpp
--- a/Engine/Plugins/Experimental/DynamicMaterial/Source/DynamicMaterialEditor/Private/Slate/Properties/Editors/SDMPropertyEditFloat1Value.cpp
+++ b/Engine/Plugins/Experimental/DynamicMaterial/Source/DynamicMaterialEditor/Private/Slate/Properties/Editors/SDMPropertyEditFloat1Value.cpp
@@ -2,8 +2,8 @@
  
 #include "SDMPropertyEditFloat1Value.h"
 #include "Components/MaterialValues/DMMaterialValueFloat1.h"
-#include "DMPrivate.h"
 #include "DynamicMaterialEditorSettings.h"
+#include "Utils/DMPrivate.h"
 #include "Widgets/Input/SSpinBox.h"
 #include "Widgets/SNullWidget.h"
  
diff --git a/Engine/Plugins/Experimental/DynamicMaterial/Source/DynamicMaterialEditor/Private/Slate/Properties/Editors/SDMPropertyEditFloat2Value.cpp b/Engine/Plugins/Experimental/DynamicMaterial/Source/DynamicMaterialEditor/Private/Slate/Properties/Editors/SDMPropertyEditFloat2Value.cpp
--- a/Engine/Plugins/Experimental/DynamicMaterial/Source/DynamicMaterialEditor/Private/Slate/Properties/Editors/SDMPropertyEditFloat2Value.cpp
+++ b/Engine/Plugins/Experimental/DynamicMaterial/Source/DynamicMaterialEditor/Private/Slate/Properties/Editors/SDMPropertyEditFloat2Value.cpp
@@ -2,8 +2,8 @@
  
 #include "SDMPropertyEditFloat2Value.h"
 #include "Components/MaterialValues/DMMaterialValueFloat2.h"
-#include "DMPrivate.h"
 #include "DynamicMaterialEditorSettings.h"
+#include "Utils/DMPrivate.h"
 #include "Widgets/Input/SSpinBox.h"
 #include "Widgets/SNullWidget.h"
  
diff --git a/Engine/Plugins/Experimental/DynamicMaterial/Source/DynamicMaterialEditor/Private/Slate/Properties/Editors/SDMPropertyEditFloat3RGBValue.cpp b/Engine/Plugins/Experimental/DynamicMaterial/Source/DynamicMaterialEditor/Private/Slate/Properties/Editors/SDMPropertyEditFloat3RGBValue.cpp
--- a/Engine/Plugins/Experimental/DynamicMaterial/Source/DynamicMaterialEditor/Private/Slate/Properties/Editors/SDMPropertyEditFloat3RGBValue.cpp
+++ b/Engine/Plugins/Experimental/DynamicMaterial/Source/DynamicMaterialEditor/Private/Slate/Properties/Editors/SDMPropertyEditFloat3RGBValue.cpp
@@ -2,7 +2,7 @@
  
 #include "SDMPropertyEditFloat3RGBValue.h"
 #include "Components/MaterialValues/DMMaterialValueFloat3RGB.h"
-#include "DMPrivate.h"
+#include "Utils/DMPrivate.h"
 #include "Widgets/SNullWidget.h"
 
 #define LOCTEXT_NAMESPACE "SDMPropertyEditFloat3RGBValue"
diff --git a/Engine/Plugins/Experimental/DynamicMaterial/Source/DynamicMaterialEditor/Private/Slate/Properties/Editors/SDMPropertyEditFloat3RPYValue.cpp b/Engine/Plugins/Experimental/DynamicMaterial/Source/DynamicMaterialEditor/Private/Slate/Properties/Editors/SDMPropertyEditFloat3RPYValue.cpp
--- a/Engine/Plugins/Experimental/DynamicMaterial/Source/DynamicMaterialEditor/Private/Slate/Properties/Editors/SDMPropertyEditFloat3RPYValue.cpp
+++ b/Engine/Plugins/Experimental/DynamicMaterial/Source/DynamicMaterialEditor/Private/Slate/Properties/Editors/SDMPropertyEditFloat3RPYValue.cpp
@@ -2,8 +2,8 @@
  
 #include "SDMPropertyEditFloat3RPYValue.h"
 #include "Components/MaterialValues/DMMaterialValueFloat3RPY.h"
-#include "DMPrivate.h"
 #include "DynamicMaterialEditorSettings.h"
+#include "Utils/DMPrivate.h"
 #include "Widgets/Input/SSpinBox.h"
 #include "Widgets/SNullWidget.h"
  
diff --git a/Engine/Plugins/Experimental/DynamicMaterial/Source/DynamicMaterialEditor/Private/Slate/Properties/Editors/SDMPropertyEditFloat3XYZValue.cpp b/Engine/Plugins/Experimental/DynamicMaterial/Source/DynamicMaterialEditor/Private/Slate/Properties/Editors/SDMPropertyEditFloat3XYZValue.cpp
--- a/Engine/Plugins/Experimental/DynamicMaterial/Source/DynamicMaterialEditor/Private/Slate/Properties/Editors/SDMPropertyEditFloat3XYZValue.cpp
+++ b/Engine/Plugins/Experimental/DynamicMaterial/Source/DynamicMaterialEditor/Private/Slate/Properties/Editors/SDMPropertyEditFloat3XYZValue.cpp
@@ -2,8 +2,8 @@
  
 #include "SDMPropertyEditFloat3XYZValue.h"
 #include "Components/MaterialValues/DMMaterialValueFloat3XYZ.h"
-#include "DMPrivate.h"
 #include "DynamicMaterialEditorSettings.h"
+#include "Utils/DMPrivate.h"
 #include "Widgets/Input/SSpinBox.h"
 #include "Widgets/SNullWidget.h"
  
diff --git a/Engine/Plugins/Experimental/DynamicMaterial/Source/DynamicMaterialEditor/Private/Slate/Properties/Editors/SDMPropertyEditFloat4Value.cpp b/Engine/Plugins/Experimental/DynamicMaterial/Source/DynamicMaterialEditor/Private/Slate/Properties/Editors/SDMPropertyEditFloat4Value.cpp
--- a/Engine/Plugins/Experimental/DynamicMaterial/Source/DynamicMaterialEditor/Private/Slate/Properties/Editors/SDMPropertyEditFloat4Value.cpp
+++ b/Engine/Plugins/Experimental/DynamicMaterial/Source/DynamicMaterialEditor/Private/Slate/Properties/Editors/SDMPropertyEditFloat4Value.cpp
@@ -2,7 +2,7 @@
  
 #include "SDMPropertyEditFloat4Value.h"
 #include "Components/MaterialValues/DMMaterialValueFloat4.h"
-#include "DMPrivate.h"
+#include "Utils/DMPrivate.h"
 #include "Widgets/SNullWidget.h"
 
 #define LOCTEXT_NAMESPACE "SDMPropertyEditFloat4Value"
diff --git a/Engine/Plugins/Experimental/DynamicMaterial/Source/DynamicMaterialEditor/Private/Slate/Properties/Editors/SDMPropertyEditTextureValue.cpp b/Engine/Plugins/Experimental/DynamicMaterial/Source/DynamicMaterialEditor/Private/Slate/Properties/Editors/SDMPropertyEditTextureValue.cpp
--- a/Engine/Plugins/Experimental/DynamicMaterial/Source/DynamicMaterialEditor/Private/Slate/Properties/Editors/SDMPropertyEditTextureValue.cpp
+++ b/Engine/Plugins/Experimental/DynamicMaterial/Source/DynamicMaterialEditor/Private/Slate/Properties/Editors/SDMPropertyEditTextureValue.cpp
@@ -2,8 +2,8 @@
  
 #include "SDMPropertyEditTextureValue.h"
 #include "Components/MaterialValues/DMMaterialValueTexture.h"
-#include "DMPrivate.h"
 #include "Engine/Texture.h"
+#include "Utils/DMPrivate.h"
 #include "Widgets/SNullWidget.h"
 
 #define LOCTEXT_NAMESPACE "SDMPropertyEditTextureValue"
diff --git a/Engine/Plugins/Experimental/DynamicMaterial/Source/DynamicMaterialEditor/Private/Slate/Properties/Generators/DMTextureUVPropertyRowGenerator.cpp b/Engine/Plugins/Experimental/DynamicMaterial/Source/DynamicMaterialEditor/Private/Slate/Properties/Generators/DMTextureUVPropertyRowGenerator.cpp
--- a/Engine/Plugins/Experimental/DynamicMaterial/Source/DynamicMaterialEditor/Private/Slate/Properties/Generators/DMTextureUVPropertyRowGenerator.cpp
+++ b/Engine/Plugins/Experimental/DynamicMaterial/Source/DynamicMaterialEditor/Private/Slate/Properties/Generators/DMTextureUVPropertyRowGenerator.cpp
@@ -4,7 +4,6 @@
 #include "Components/DMMaterialComponent.h"
 #include "Components/DMTextureUV.h"
 #include "DMEDefs.h"
-#include "DMPrivate.h"
 #include "Modules/ModuleManager.h"
 #include "PropertyCustomizationHelpers.h"
 #include "PropertyEditorModule.h"
@@ -13,6 +12,7 @@
 #include "Slate/SDMComponentEdit.h"
 #include "Slate/SDMEditor.h"
 #include "Styling/SlateIconFinder.h"
+#include "Utils/DMPrivate.h"
 #include "Widgets/Layout/SBox.h"
 
 #define LOCTEXT_NAMESPACE "DMTextureUVPropertyRowGenerator"
diff --git a/Engine/Plugins/Experimental/DynamicMaterial/Source/DynamicMaterialEditor/Private/Slate/Properties/SDMMaterialParameters.cpp b/Engine/Plugins/Experimental/DynamicMaterial/Source/DynamicMaterialEditor/Private/Slate/Properties/SDMMaterialParameters.cpp
--- a/Engine/Plugins/Experimental/DynamicMaterial/Source/DynamicMaterialEditor/Private/Slate/Properties/SDMMaterialParameters.cpp
+++ b/Engine/Plugins/Experimental/DynamicMaterial/Source/DynamicMaterialEditor/Private/Slate/Properties/SDMMaterialParameters.cpp
@@ -4,7 +4,6 @@
 #include "Components/DMMaterialParameter.h"
 #include "Components/DMMaterialValue.h"
 #include "DetailLayoutBuilder.h"
-#include "DMPrivate.h"
 #include "DMValueDefinition.h"
 #include "DynamicMaterialEditorModule.h"
 #include "DynamicMaterialEditorStyle.h"
@@ -13,6 +12,7 @@
 #include "Slate/Properties/SDMDetailsGrid.h"
 #include "SlateOptMacros.h"
 #include "Styling/StyleColors.h"
+#include "Utils/DMPrivate.h"
 #include "Widgets/Images/SImage.h"
 #include "Widgets/Input/SButton.h"
 #include "Widgets/Input/SEditableTextBox.h"
diff --git a/Engine/Plugins/Experimental/DynamicMaterial/Source/DynamicMaterialEditor/Private/Slate/Properties/SDMPropertyEdit.cpp b/Engine/Plugins/Experimental/DynamicMaterial/Source/DynamicMaterialEditor/Private/Slate/Properties/SDMPropertyEdit.cpp
--- a/Engine/Plugins/Experimental/DynamicMaterial/Source/DynamicMaterialEditor/Private/Slate/Properties/SDMPropertyEdit.cpp
+++ b/Engine/Plugins/Experimental/DynamicMaterial/Source/DynamicMaterialEditor/Private/Slate/Properties/SDMPropertyEdit.cpp
@@ -2,10 +2,9 @@
 
 #include "Slate/Properties/SDMPropertyEdit.h"
 #include "Components/DMMaterialValue.h"
-#include "DMPrivate.h"
-#include "DMWorldSubsystem.h"
 #include "DetailLayoutBuilder.h"
 #include "DetailRowMenuContext.h"
+#include "DMWorldSubsystem.h"
 #include "DynamicMaterialEditorModule.h"
 #include "DynamicMaterialEditorSettings.h"
 #include "Engine/Engine.h"
@@ -14,11 +13,12 @@
 #include "IDetailKeyframeHandler.h"
 #include "PropertyCustomizationHelpers.h"
 #include "PropertyHandle.h"
+#include "Slate/SDMComponentEdit.h"
 #include "Slate/SDMEditor.h"
 #include "Slate/SDMSlot.h"
-#include "Slate/SDMComponentEdit.h"
 #include "ToolMenuContext.h"
 #include "ToolMenus.h"
+#include "Utils/DMPrivate.h"
 #include "Widgets/Colors/SColorBlock.h"
 #include "Widgets/Colors/SColorPicker.h"
 #include "Widgets/Input/SCheckBox.h"
diff --git a/Engine/Plugins/Experimental/DynamicMaterial/Source/DynamicMaterialEditor/Private/Slate/Properties/SDMSourceBlendType.cpp b/Engine/Plugins/Experimental/DynamicMaterial/Source/DynamicMaterialEditor/Private/Slate/Properties/SDMSourceBlendType.cpp
--- a/Engine/Plugins/Experimental/DynamicMaterial/Source/DynamicMaterialEditor/Private/Slate/Properties/SDMSourceBlendType.cpp
+++ b/Engine/Plugins/Experimental/DynamicMaterial/Source/DynamicMaterialEditor/Private/Slate/Properties/SDMSourceBlendType.cpp
@@ -28,9 +28,6 @@
 #include "Components/MaterialStageBlends/DMMSBSoftLight.h"
 #include "Components/MaterialStageBlends/DMMSBSubtract.h"
 #include "Components/MaterialStageBlends/DMMSBVividLight.h"
-#include "DMEDefs.h"
-#include "DMPrivate.h"
-#include "DynamicMaterialEditorStyle.h"
 #include "ToolMenu.h"
 #include "ToolMenuDelegates.h"
 #include "ToolMenus.h"
diff --git a/Engine/Plugins/Experimental/DynamicMaterial/Source/DynamicMaterialEditor/Private/Slate/SDMComponentEdit.cpp b/Engine/Plugins/Experimental/DynamicMaterial/Source/DynamicMaterialEditor/Private/Slate/SDMComponentEdit.cpp
--- a/Engine/Plugins/Experimental/DynamicMaterial/Source/DynamicMaterialEditor/Private/Slate/SDMComponentEdit.cpp
+++ b/Engine/Plugins/Experimental/DynamicMaterial/Source/DynamicMaterialEditor/Private/Slate/SDMComponentEdit.cpp
@@ -19,7 +19,6 @@
 #include "CustomDetailsViewSequencer.h"
 #include "DetailLayoutBuilder.h"
 #include "DMEDefs.h"
-#include "DMPrivate.h"
 #include "DMWorldSubsystem.h"
 #include "DynamicMaterialEditorModule.h"
 #include "DynamicMaterialEditorStyle.h"
@@ -51,6 +50,7 @@
 #include "Styling/SlateTypes.h"
 #include "ToolMenu.h"
 #include "ToolMenus.h"
+#include "Utils/DMPrivate.h"
 #include "Widgets/Input/SComboButton.h"
 #include "Widgets/Layout/SBox.h"
 #include "Widgets/SNullWidget.h"
diff --git a/Engine/Plugins/Experimental/DynamicMaterial/Source/DynamicMaterialEditor/Private/Slate/SDMEditor.cpp b/Engine/Plugins/Experimental/DynamicMaterial/Source/DynamicMaterialEditor/Private/Slate/SDMEditor.cpp
--- a/Engine/Plugins/Experimental/DynamicMaterial/Source/DynamicMaterialEditor/Private/Slate/SDMEditor.cpp
+++ b/Engine/Plugins/Experimental/DynamicMaterial/Source/DynamicMaterialEditor/Private/Slate/SDMEditor.cpp
@@ -8,8 +8,6 @@
 #include "Components/MaterialStageExpressions/DMMSETextureSample.h"
 #include "Components/PrimitiveComponent.h"
 #include "DetailLayoutBuilder.h"
-#include "DMBlueprintFunctionLibrary.h"
-#include "DMPrivate.h"
 #include "DMWorldSubsystem.h"
 #include "DynamicMaterialEditorCommands.h"
 #include "DynamicMaterialEditorModule.h"
@@ -36,6 +34,8 @@
 #include "SlateOptMacros.h"
 #include "Styling/StyleColors.h"
 #include "ThumbnailRendering/ThumbnailManager.h"
+#include "Utils/DMBlueprintFunctionLibrary.h"
+#include "Utils/DMPrivate.h"
 #include "Widgets/Input/SButton.h"
 #include "Widgets/Input/SCheckBox.h"
 #include "Widgets/Layout/SBorder.h"
diff --git a/Engine/Plugins/Experimental/DynamicMaterial/Source/DynamicMaterialEditor/Private/Slate/SDMSlot.cpp b/Engine/Plugins/Experimental/DynamicMaterial/Source/DynamicMaterialEditor/Private/Slate/SDMSlot.cpp
--- a/Engine/Plugins/Experimental/DynamicMaterial/Source/DynamicMaterialEditor/Private/Slate/SDMSlot.cpp
+++ b/Engine/Plugins/Experimental/DynamicMaterial/Source/DynamicMaterialEditor/Private/Slate/SDMSlot.cpp
@@ -22,8 +22,6 @@
 #include "Components/MaterialStageInputs/DMMSIValue.h"
 #include "Components/MaterialValues/DMMaterialValueFloat1.h"
 #include "DetailLayoutBuilder.h"
-#include "DMBlueprintFunctionLibrary.h"
-#include "DMPrivate.h"
 #include "DynamicMaterialEditorModule.h"
 #include "DynamicMaterialEditorSettings.h"
 #include "DynamicMaterialEditorStyle.h"
@@ -43,6 +41,8 @@
 #include "SlateOptMacros.h"
 #include "Styling/AppStyle.h"
 #include "Styling/StyleColors.h"
+#include "Utils/DMBlueprintFunctionLibrary.h"
+#include "Utils/DMPrivate.h"
 #include "Widgets/Colors/SColorBlock.h"
 #include "Widgets/Input/SButton.h"
 #include "Widgets/Input/SComboButton.h"
diff --git a/Engine/Plugins/Experimental/DynamicMaterial/Source/DynamicMaterialEditor/Private/Slate/SDMStage.cpp b/Engine/Plugins/Experimental/DynamicMaterial/Source/DynamicMaterialEditor/Private/Slate/SDMStage.cpp
--- a/Engine/Plugins/Experimental/DynamicMaterial/Source/DynamicMaterialEditor/Private/Slate/SDMStage.cpp
+++ b/Engine/Plugins/Experimental/DynamicMaterial/Source/DynamicMaterialEditor/Private/Slate/SDMStage.cpp
@@ -12,7 +12,6 @@
 #include "Components/MaterialStageInputs/DMMSIExpression.h"
 #include "Components/MaterialStageInputs/DMMSIValue.h"
 #include "Components/MaterialValues/DMMaterialValueTexture.h"
-#include "DMPrivate.h"
 #include "DragAndDrop/AssetDragDropOp.h"
 #include "DragDrop/DMStageDragDropOperation.h"
 #include "DynamicMaterialEditorSettings.h"
@@ -23,6 +22,7 @@
 #include "Slate/SMaterialToolTip.h"
 #include "SlateOptMacros.h"
 #include "Styling/CoreStyle.h"
+#include "Utils/DMPrivate.h"
 #include "Widgets/Colors/SColorBlock.h"
 #include "Widgets/Layout/SBox.h"
 #include "Widgets/Layout/SScaleBox.h"
diff --git a/Engine/Plugins/Experimental/DynamicMaterial/Source/DynamicMaterialEditor/Private/Slate/SDMToolBar.cpp b/Engine/Plugins/Experimental/DynamicMaterial/Source/DynamicMaterialEditor/Private/Slate/SDMToolBar.cpp
--- a/Engine/Plugins/Experimental/DynamicMaterial/Source/DynamicMaterialEditor/Private/Slate/SDMToolBar.cpp
+++ b/Engine/Plugins/Experimental/DynamicMaterial/Source/DynamicMaterialEditor/Private/Slate/SDMToolBar.cpp
@@ -3,7 +3,6 @@
 #include "Slate/SDMToolBar.h"
 #include "AssetToolsModule.h"
 #include "ContentBrowserModule.h"
-#include "DMBlueprintFunctionLibrary.h"
 #include "DMEDefs.h"
 #include "DynamicMaterialEditorModule.h"
 #include "DynamicMaterialEditorSettings.h"
@@ -13,11 +12,12 @@
 #include "GameFramework/Actor.h"
 #include "IAssetTools.h"
 #include "IContentBrowserSingleton.h"
+#include "Material/DynamicMaterialInstance.h"
 #include "Model/DynamicMaterialModel.h"
 #include "SlateOptMacros.h"
-#include "Material/DynamicMaterialInstance.h"
 #include "Styling/AppStyle.h"
 #include "Styling/StyleColors.h"
+#include "Utils/DMBlueprintFunctionLibrary.h"
 #include "Widgets/Colors/SColorBlock.h"
 #include "Widgets/Images/SImage.h"
 #include "Widgets/Input/SButton.h"
diff --git a/Engine/Plugins/Experimental/DynamicMaterial/Source/DynamicMaterialEditor/Private/Utils/DMBlueprintFunctionLibrary.cpp b/Engine/Plugins/Experimental/DynamicMaterial/Source/DynamicMaterialEditor/Private/Utils/DMBlueprintFunctionLibrary.cpp
new file mode 100644
--- /dev/null
+++ b/Engine/Plugins/Experimental/DynamicMaterial/Source/DynamicMaterialEditor/Private/Utils/DMBlueprintFunctionLibrary.cpp
@@ -0,0 +1,412 @@
+// Copyright Epic Games, Inc. All Rights Reserved.
+
+#include "Utils/DMBlueprintFunctionLibrary.h"
+#include "AssetRegistry/AssetRegistryModule.h"
+#include "Components/DMMaterialLayer.h"
+#include "Components/DMMaterialSlot.h"
+#include "Components/DMMaterialStage.h"
+#include "Components/DMMaterialStageBlend.h"
+#include "Components/DMMaterialStageFunction.h"
+#include "Components/DMMaterialStageGradient.h"
+#include "Components/DMMaterialStageSource.h"
+#include "Components/DMMaterialStageThroughput.h"
+#include "Components/DMMaterialStageThroughputLayerBlend.h"
+#include "Components/DMMaterialSubStage.h"
+#include "Components/DMRenderTargetRenderer.h"
+#include "Components/MaterialStageBlends/DMMSBNormal.h"
+#include "Components/MaterialStageExpressions/DMMSETextureSample.h"
+#include "Components/MaterialStageInputs/DMMSIExpression.h"
+#include "Components/MaterialStageInputs/DMMSITextureUV.h"
+#include "Components/MaterialStageInputs/DMMSIValue.h"
+#include "Components/MaterialValues/DMMaterialValueRenderTarget.h"
+#include "Components/MaterialValues/DMMaterialValueTexture.h"
+#include "Components/PrimitiveComponent.h"
+#include "DMEDefs.h"
+#include "DMObjectMaterialProperty.h"
+#include "DMWorldSubsystem.h"
+#include "DynamicMaterialEditorModule.h"
+#include "Engine/Texture.h"
+#include "Engine/World.h"
+#include "GameFramework/Actor.h"
+#include "Material/DynamicMaterialInstance.h"
+#include "Material/DynamicMaterialInstanceFactory.h"
+#include "Materials/Material.h"
+#include "Model/DynamicMaterialModel.h"
+#include "Model/DynamicMaterialModelEditorOnlyData.h"
+#include "Model/IDynamicMaterialModelEditorOnlyDataInterface.h"
+#include "UObject/Package.h"
+#include "UObject/UObjectGlobals.h"
+#include "Utils/DMPrivate.h"
+
+UDMMaterialStageInputValue* UDMBlueprintFunctionLibrary::FindDefaultStageOpacityInputValue(UDMMaterialStage* InStage)
+{
+	if (!IsValid(InStage))
+	{
+		return nullptr;
+	}
+
+	int32 OpacityInputIndex = -1;
+
+	if (UDMMaterialStageThroughput* const ThroughputSource = Cast<UDMMaterialStageThroughput>(InStage->GetSource()))
+	{
+		const TArray<FDMMaterialStageConnector>& Connectors = ThroughputSource->GetInputConnectors();
+		for (const FDMMaterialStageConnector& Connector : Connectors)
+		{
+			if (Connector.Name.ToString() == "Opacity")
+			{
+				OpacityInputIndex = Connector.Index;
+			}
+		}
+	}
+
+	const TArray<UDMMaterialStageInput*>& StageInputs = InStage->GetInputs();
+	const TArray<FDMMaterialStageConnection>& InputConnectionMap = InStage->GetInputConnectionMap();
+
+	if (InputConnectionMap.IsValidIndex(OpacityInputIndex) && InputConnectionMap[OpacityInputIndex].Channels.Num() > 0)
+	{
+		if (InputConnectionMap[OpacityInputIndex].Channels[0].SourceIndex >= FDMMaterialStageConnectorChannel::FIRST_STAGE_INPUT)
+		{
+			const int32 StageInputIndex = InputConnectionMap[OpacityInputIndex].Channels[0].SourceIndex - FDMMaterialStageConnectorChannel::FIRST_STAGE_INPUT;
+
+			return Cast<UDMMaterialStageInputValue>(StageInputs[StageInputIndex]);
+		}
+	}
+
+	return nullptr;
+}
+
+void UDMBlueprintFunctionLibrary::SetDefaultStageSourceTexture(UDMMaterialStage* InStage, UTexture* InTexture)
+{
+	if (!IsValid(InStage) || !IsValid(InTexture))
+	{
+		return;
+	}
+
+	UDMMaterialStageSource* StageSource = InStage->GetSource();
+	if (!StageSource)
+	{
+		return;
+	}
+
+	auto SetTextureValue = [InTexture](UDMMaterialStageInputValue* NewInputValue)
+	{
+		if (UDMMaterialValueTexture* TextureValue = Cast<UDMMaterialValueTexture>(NewInputValue->GetValue()))
+		{
+			TextureValue->SetValue(InTexture);
+		}
+	};
+
+	if (GUndo)
+	{
+		InStage->Modify();
+	}
+
+	if (UDMMaterialStageBlend* const Blend = Cast<UDMMaterialStageBlend>(StageSource))
+	{
+		UDMMaterialStageInputExpression* NewInput = UDMMaterialStageInputExpression::ChangeStageInput_Expression(
+			InStage,
+			UDMMaterialStageExpressionTextureSample::StaticClass(), 
+			UDMMaterialStageBlend::InputB,
+			FDMMaterialStageConnectorChannel::WHOLE_CHANNEL, 
+			0,
+			FDMMaterialStageConnectorChannel::WHOLE_CHANNEL
+		);
+
+		UDMMaterialSubStage* SubStage = NewInput->GetSubStage();
+		check(SubStage);
+
+		UDMMaterialStageInputValue* NewInputValue = UDMMaterialStageInputValue::ChangeStageInput_NewLocalValue(
+			SubStage,
+			0, 
+			FDMMaterialStageConnectorChannel::WHOLE_CHANNEL, 
+			EDMValueType::VT_Texture,
+			FDMMaterialStageConnectorChannel::WHOLE_CHANNEL
+		);
+
+		SetTextureValue(NewInputValue);
+		return;
+	}
+
+	if (UDMMaterialStageThroughputLayerBlend* const LayerBlend = Cast<UDMMaterialStageThroughputLayerBlend>(StageSource))
+	{
+		const bool bHasAlpha = UE::DynamicMaterial::Private::HasAlpha(InTexture);
+
+		UDMMaterialStageInputExpression* NewInput = UDMMaterialStageInputExpression::ChangeStageInput_Expression(
+			InStage,
+			UDMMaterialStageExpressionTextureSample::StaticClass(), 
+			2, 
+			FDMMaterialStageConnectorChannel::WHOLE_CHANNEL,
+			bHasAlpha ? 1 : 0, 
+			FDMMaterialStageConnectorChannel::WHOLE_CHANNEL
+		);
+
+		UDMMaterialSubStage* SubStage = NewInput->GetSubStage();
+		check(IsValid(SubStage));
+
+		UDMMaterialStageInputValue* NewInputValue = UDMMaterialStageInputValue::ChangeStageInput_NewLocalValue(
+			SubStage,
+			0, 
+			FDMMaterialStageConnectorChannel::WHOLE_CHANNEL, 
+			EDMValueType::VT_Texture,
+			FDMMaterialStageConnectorChannel::WHOLE_CHANNEL
+		);
+
+		SetTextureValue(NewInputValue);
+		return;
+	}
+
+	UDMMaterialStageExpression* NewExpression = InStage->ChangeSource<UDMMaterialStageExpressionTextureSample>();
+
+	UDMMaterialStageInputValue* NewInputValue = UDMMaterialStageInputValue::ChangeStageInput_NewLocalValue(
+		InStage,
+		0, 
+		FDMMaterialStageConnectorChannel::WHOLE_CHANNEL, 
+		EDMValueType::VT_Texture,
+		FDMMaterialStageConnectorChannel::WHOLE_CHANNEL
+	);
+
+	SetTextureValue(NewInputValue);
+}
+
+TArray<FDMObjectMaterialProperty> UDMBlueprintFunctionLibrary::GetActorMaterialProperties(AActor* InActor)
+{
+	TArray<FDMObjectMaterialProperty> ActorProperties;
+
+	if (!IsValid(InActor))
+	{
+		return ActorProperties;
+	}
+
+	FDMGetObjectMaterialPropertiesDelegate PropertyGenerator = FDynamicMaterialEditorModule::GetCustomMaterialPropertyGenerator(InActor->GetClass());
+
+	if (PropertyGenerator.IsBound())
+	{
+		ActorProperties = PropertyGenerator.Execute(InActor);
+
+		if (!ActorProperties.IsEmpty())
+		{
+			return ActorProperties;
+		}
+	}
+
+	InActor->ForEachComponent<UPrimitiveComponent>(false, [&ActorProperties](UPrimitiveComponent* InComp)
+		{
+			for (int32 MaterialIdx = 0; MaterialIdx < InComp->GetNumMaterials(); ++MaterialIdx)
+			{
+				ActorProperties.Add({InComp, MaterialIdx});
+			}
+		});
+
+	return ActorProperties;
+}
+
+UDynamicMaterialModel* UDMBlueprintFunctionLibrary::CreateDynamicMaterialInObject(FDMObjectMaterialProperty& InMaterialProperty)
+{
+	if (!InMaterialProperty.IsValid())
+	{
+		return nullptr;
+	}
+
+	UObject* const Outer = InMaterialProperty.OuterWeak.Get();
+
+	UDynamicMaterialInstanceFactory* const InstanceFactory = NewObject<UDynamicMaterialInstanceFactory>();
+	check(InstanceFactory);
+
+	UDynamicMaterialInstance* const NewInstance = Cast<UDynamicMaterialInstance>(InstanceFactory->FactoryCreateNew(UDynamicMaterialInstance::StaticClass(),
+		Outer, NAME_None, RF_Transactional, nullptr, GWarn));
+	check(NewInstance);
+
+	bool bSubsystemTakenOver = false;
+
+	if (const UWorld* const World = Outer->GetWorld())
+	{
+		if (IsValid(World))
+		{
+			UDMWorldSubsystem* const WorldSubsystem = World->GetSubsystem<UDMWorldSubsystem>();
+
+			if (IsValid(WorldSubsystem))
+			{
+				if (WorldSubsystem->GetMaterialValueSetterDelegate().IsBound())
+				{
+					bSubsystemTakenOver = WorldSubsystem->GetMaterialValueSetterDelegate().Execute(InMaterialProperty, NewInstance);
+				}
+			}
+		}
+	}
+
+	if (!bSubsystemTakenOver)
+	{
+		InMaterialProperty.SetMaterial(NewInstance);
+	}
+
+	if (UDynamicMaterialModel* MaterialModel = NewInstance->GetMaterialModel())
+	{
+		if (IDynamicMaterialModelEditorOnlyDataInterface* EditorOnlyData = MaterialModel->GetEditorOnlyData())
+		{
+			EditorOnlyData->RequestMaterialBuild();
+		}
+
+		return MaterialModel;
+	}
+
+	return nullptr;
+}
+
+bool UDMBlueprintFunctionLibrary::ExportMaterialInstance(UDynamicMaterialModel* InMaterialModel, const FString& InSavePath)
+{
+	if (!IsValid(InMaterialModel))
+	{
+		UE::DynamicMaterialEditor::Private::LogError(TEXT("Invalid material to export."));
+		return false;
+	}
+
+	if (InSavePath.IsEmpty())
+	{
+		UE::DynamicMaterialEditor::Private::LogError(TEXT("Invalid material save path to export."));
+		return false;
+	}
+
+	UDynamicMaterialInstance* Instance = InMaterialModel->GetDynamicMaterialInstance();
+	if (!Instance)
+	{
+		UE::DynamicMaterialEditor::Private::LogError(TEXT("Failed to find a Material Designer Instance to export."), true, InMaterialModel);
+		return false;
+	}
+
+	const FString PackagePath = FPaths::GetBaseFilename(*InSavePath, false);
+
+	UPackage* Package = CreatePackage(*PackagePath);
+	if (!Package)
+	{
+		UE::DynamicMaterialEditor::Private::LogError(FString::Printf(TEXT("Failed to create package for Material Designer Instance (%s)."), *PackagePath));
+		return false;
+	}
+
+	const FString AssetName = FPaths::GetBaseFilename(*InSavePath, true);
+
+	FObjectDuplicationParameters Params = InitStaticDuplicateObjectParams(Instance, Package, *AssetName,
+		RF_Public | RF_Standalone, nullptr, EDuplicateMode::Normal, EInternalObjectFlags::None);
+
+	UObject* NewAsset = StaticDuplicateObjectEx(Params);
+	if (!NewAsset)
+	{
+		UE::DynamicMaterialEditor::Private::LogError(TEXT("Failed to create new Material Designer Instance asset."));
+		return false;
+	}
+
+	// Not sure why these need to be set again!
+	NewAsset->SetFlags(RF_Public | RF_Standalone);
+
+	if (UDynamicMaterialInstance* NewInstance = Cast<UDynamicMaterialInstance>(NewAsset))
+	{
+		if (UDynamicMaterialModel* NewModel = NewInstance->GetMaterialModel())
+		{
+			if (IDynamicMaterialModelEditorOnlyDataInterface* ModelEditorOnlyData = NewModel->GetEditorOnlyData())
+			{
+				ModelEditorOnlyData->RequestMaterialBuild();
+			}
+		}
+	}
+
+	FAssetRegistryModule::AssetCreated(NewAsset);
+	Package->FullyLoad();
+
+	return true;
+}
+
+bool UDMBlueprintFunctionLibrary::ExportGeneratedMaterial(UDynamicMaterialModel* InMaterialModel, const FString& InSavePath)
+{
+	if (!IsValid(InMaterialModel))
+	{
+		UE::DynamicMaterialEditor::Private::LogError(TEXT("Invalid material to export."));
+		return false;
+	}
+
+	if (InSavePath.IsEmpty())
+	{
+		UE::DynamicMaterialEditor::Private::LogError(TEXT("Invalid material save path to export."));
+		return false;
+	}
+
+	UMaterial* GeneratedMaterial = InMaterialModel->GetGeneratedMaterial();
+	if (!GeneratedMaterial)
+	{
+		UE::DynamicMaterialEditor::Private::LogError(TEXT("Failed to find a generated material to export."));
+		return false;
+	}
+
+	const FString PackagePath = FPaths::GetBaseFilename(*InSavePath, false);
+
+	UPackage* Package = CreatePackage(*PackagePath);
+	if (!Package)
+	{
+		UE::DynamicMaterialEditor::Private::LogError(FString::Printf(TEXT("Failed to create package for exported material (%s)."), *PackagePath));
+		return false;
+	}
+
+	const FString AssetName = FPaths::GetBaseFilename(*InSavePath, true);
+
+	FObjectDuplicationParameters Params = InitStaticDuplicateObjectParams(GeneratedMaterial, Package, *AssetName,
+		RF_Public | RF_Standalone, nullptr, EDuplicateMode::Normal, EInternalObjectFlags::None);
+
+	UObject* NewAsset = StaticDuplicateObjectEx(Params);
+	if (!NewAsset)
+	{
+		UE::DynamicMaterialEditor::Private::LogError(TEXT("Failed to create new material asset."));
+		return false;
+	}
+
+	// Not sure why these need to be set again!
+	NewAsset->SetFlags(RF_Public | RF_Standalone);
+
+	FAssetRegistryModule::AssetCreated(NewAsset);
+	Package->FullyLoad();
+
+	return true;
+}
+
+void UDMBlueprintFunctionLibrary::SetStageInputToRenderer(UDMMaterialStage* InStage, TSubclassOf<UDMRenderTargetRenderer> InRendererClass, int32 InInputIndex)
+{
+	UDMMaterialStageInputExpression* InputExpression = UDMMaterialStageInputExpression::ChangeStageInput_Expression(
+		InStage,
+		UDMMaterialStageExpressionTextureSample::StaticClass(),
+		InInputIndex,
+		FDMMaterialStageConnectorChannel::WHOLE_CHANNEL,
+		0,
+		FDMMaterialStageConnectorChannel::THREE_CHANNELS
+	);
+
+	if (!ensure(InputExpression))
+	{
+		return;
+	}
+
+	UDMMaterialSubStage* SubStage = InputExpression->GetSubStage();
+
+	if (!ensure(SubStage))
+	{
+		return;
+	}
+
+	UDMMaterialStageInputValue* InputValue = UDMMaterialStageInputValue::ChangeStageInput_NewLocalValue(
+		SubStage,
+		0,
+		FDMMaterialStageConnectorChannel::WHOLE_CHANNEL,
+		UDMMaterialValueRenderTarget::StaticClass(),
+		FDMMaterialStageConnectorChannel::WHOLE_CHANNEL
+	);
+
+	if (!ensure(InputValue))
+	{
+		return;
+	}
+
+	UDMMaterialValueRenderTarget* RenderTargetValue = Cast<UDMMaterialValueRenderTarget>(InputValue->GetValue());
+
+	if (!RenderTargetValue)
+	{
+		return;
+	}
+
+	UDMRenderTargetRenderer::CreateRenderTargetRenderer(InRendererClass, RenderTargetValue);
+}
diff --git a/Engine/Plugins/Experimental/DynamicMaterial/Source/DynamicMaterialEditor/Private/Utils/DMBlueprintFunctionLibrary.h b/Engine/Plugins/Experimental/DynamicMaterial/Source/DynamicMaterialEditor/Private/Utils/DMBlueprintFunctionLibrary.h
new file mode 100644
--- /dev/null
+++ b/Engine/Plugins/Experimental/DynamicMaterial/Source/DynamicMaterialEditor/Private/Utils/DMBlueprintFunctionLibrary.h
@@ -0,0 +1,40 @@
+// Copyright Epic Games, Inc. All Rights Reserved.
+
+#pragma once
+
+#include "DMDefs.h"
+#include "DMEDefs.h"
+#include "Kismet/BlueprintFunctionLibrary.h"
+#include "DMBlueprintFunctionLibrary.generated.h"
+
+class AActor;
+class UDMMaterialStage;
+class UDMMaterialStageInputValue;
+class UDMRenderTargetRenderer;
+class UDynamicMaterialModel;
+class UTexture;
+struct FDMObjectMaterialProperty;
+
+/**
+ * Material Designer Blueprint Function Library
+ */
+UCLASS()
+class UDMBlueprintFunctionLibrary : public UBlueprintFunctionLibrary
+{
+	GENERATED_BODY()
+
+public:
+	static UDMMaterialStageInputValue* FindDefaultStageOpacityInputValue(UDMMaterialStage* InStage);
+
+	static void SetDefaultStageSourceTexture(UDMMaterialStage* InStage, UTexture* InTexture);
+
+	static TArray<FDMObjectMaterialProperty> GetActorMaterialProperties(AActor* InActor);
+
+	static UDynamicMaterialModel* CreateDynamicMaterialInObject(FDMObjectMaterialProperty& InMaterialProperty);
+
+	static bool ExportMaterialInstance(UDynamicMaterialModel* InMaterialModel, const FString& InSavePath);
+
+	static bool ExportGeneratedMaterial(UDynamicMaterialModel* InMaterialModel, const FString& InSavePath);
+
+	static void SetStageInputToRenderer(UDMMaterialStage* InStage, TSubclassOf<UDMRenderTargetRenderer> InRendererClass, int32 InInputIndex);
+};
diff --git a/Engine/Plugins/Experimental/DynamicMaterial/Source/DynamicMaterialEditor/Private/Utils/DMInputNodeBuilder.cpp b/Engine/Plugins/Experimental/DynamicMaterial/Source/DynamicMaterialEditor/Private/Utils/DMInputNodeBuilder.cpp
new file mode 100644
--- /dev/null
+++ b/Engine/Plugins/Experimental/DynamicMaterial/Source/DynamicMaterialEditor/Private/Utils/DMInputNodeBuilder.cpp
@@ -0,0 +1,183 @@
+// Copyright Epic Games, Inc. All Rights Reserved.
+
+#include "Utils/DMInputNodeBuilder.h"
+#include "DMEDefs.h"
+#include "Materials/MaterialExpressionComponentMask.h"
+#include "Materials/MaterialExpressionIf.h"
+#include "Materials/MaterialExpressionTextureCoordinate.h"
+#include "Model/DMMaterialBuildState.h"
+#include "Model/DMMaterialBuildUtils.h"
+#include "Utils/DMPrivate.h"
+
+namespace UE::DynamicMaterialEditor::Private
+{
+	void BuildExpressionInputs(const TSharedRef<FDMMaterialBuildState>& InBuildState, const TArray<FDMMaterialStageConnection>& InputConnectionMap, const TArray<FDMInputInputs>& Inputs)
+	{
+		switch (Inputs.Num())
+		{
+			case 0:
+				break;
+
+			case 1:
+				BuildExpressionOneInput(InBuildState, InputConnectionMap, Inputs[0]);
+				break;
+
+			case 2:
+				BuildExpressionTwoInputs(InBuildState, InputConnectionMap, Inputs[0], Inputs[1]);
+				break;
+
+			default:
+				BuildExpressionThreeInputs(InBuildState, InputConnectionMap, Inputs[0], Inputs[1], Inputs[2]);
+				break;
+		}
+	}
+
+	void BuildExpressionOneInput(const TSharedRef<FDMMaterialBuildState>& InBuildState, const TArray<FDMMaterialStageConnection>& InputConnectionMap,
+		const FDMInputInputs& InputOne)
+	{
+		int32 OutputIndex;
+		int32 OutputChannel;
+
+		const TArray<UMaterialExpression*> InputOneExpression = InBuildState->GetBuildUtils().CreateExpressionInputs(
+			InputConnectionMap,
+			InputOne.InputIndex,
+			InputOne.ChannelInputs,
+			OutputIndex,
+			OutputChannel
+		);
+
+		InBuildState->GetBuildUtils().UpdatePreviewMaterial(
+			InputOneExpression.Last(),
+			OutputIndex,
+			OutputChannel,
+			32
+		);
+	}
+
+	void BuildExpressionTwoInputs(const TSharedRef<FDMMaterialBuildState>& InBuildState, const TArray<FDMMaterialStageConnection>& InputConnectionMap,
+		const FDMInputInputs& InputOne, const FDMInputInputs& InputTwo)
+	{
+		UMaterialExpressionTextureCoordinate* TexCoord = InBuildState->GetBuildUtils().CreateExpression<UMaterialExpressionTextureCoordinate>(UE_DM_NodeComment_Default);
+		UMaterialExpressionComponentMask* Mask = InBuildState->GetBuildUtils().CreateExpression<UMaterialExpressionComponentMask>(UE_DM_NodeComment_Default);
+		UMaterialExpressionIf* IfExpr = InBuildState->GetBuildUtils().CreateExpression<UMaterialExpressionIf>(UE_DM_NodeComment_Default);
+
+		TexCoord->ConnectExpression(&Mask->Input, 0);
+
+		Mask->R = 1;
+		Mask->G = 0;
+		Mask->B = 0;
+		Mask->A = 0;
+		Mask->ConnectExpression(&IfExpr->A, 0);
+
+		IfExpr->ConstB = 0.5;
+
+		InBuildState->AddOtherExpressions({TexCoord, Mask, IfExpr});
+
+		int32 OutputIndexOne;
+		int32 OutputChannelOne;
+		int32 OutputIndexTwo;
+		int32 OutputChannelTwo;
+
+		const TArray<UMaterialExpression*> InputOneExpressions = InBuildState->GetBuildUtils().CreateExpressionInputs(
+			InputConnectionMap,
+			InputOne.InputIndex,
+			InputOne.ChannelInputs,
+			OutputIndexOne,
+			OutputChannelOne
+		);
+
+		const TArray<UMaterialExpression*> InputTwoExpressions = InBuildState->GetBuildUtils().CreateExpressionInputs(
+			InputConnectionMap,
+			InputTwo.InputIndex,
+			InputTwo.ChannelInputs,
+			OutputIndexTwo,
+			OutputChannelTwo
+		);
+
+		InputOneExpressions.Last()->ConnectExpression(&IfExpr->ALessThanB, OutputIndexOne);
+		InputOneExpressions.Last()->ConnectExpression(&IfExpr->AEqualsB, OutputIndexOne);
+		InputTwoExpressions.Last()->ConnectExpression(&IfExpr->AGreaterThanB, OutputIndexTwo);
+
+		const TArray<FExpressionOutput>& OutputsOne = InputOneExpressions.Last()->GetOutputs();
+		const TArray<FExpressionOutput>& OutputsTwo = InputTwoExpressions.Last()->GetOutputs();
+
+		UE::DynamicMaterialEditor::Private::SetMask(IfExpr->ALessThanB, OutputsOne[OutputIndexOne], OutputChannelOne);
+		UE::DynamicMaterialEditor::Private::SetMask(IfExpr->AEqualsB, OutputsOne[OutputIndexOne], OutputChannelOne);
+		UE::DynamicMaterialEditor::Private::SetMask(IfExpr->AGreaterThanB, OutputsTwo[OutputIndexTwo], OutputChannelTwo);
+
+		InBuildState->GetBuildUtils().UpdatePreviewMaterial(IfExpr, 0, 0, 32);
+	}
+
+	void BuildExpressionThreeInputs(const TSharedRef<FDMMaterialBuildState>& InBuildState, const TArray<FDMMaterialStageConnection>& InputConnectionMap,
+		const FDMInputInputs& InputOne, const FDMInputInputs& InputTwo, const FDMInputInputs& InputThree)
+	{
+		UMaterialExpressionTextureCoordinate* TexCoord = InBuildState->GetBuildUtils().CreateExpression<UMaterialExpressionTextureCoordinate>(UE_DM_NodeComment_Default);
+		UMaterialExpressionComponentMask* Mask = InBuildState->GetBuildUtils().CreateExpression<UMaterialExpressionComponentMask>(UE_DM_NodeComment_Default);
+		UMaterialExpressionIf* IfExpr1 = InBuildState->GetBuildUtils().CreateExpression<UMaterialExpressionIf>(UE_DM_NodeComment_Default);
+		UMaterialExpressionIf* IfExpr2 = InBuildState->GetBuildUtils().CreateExpression<UMaterialExpressionIf>(UE_DM_NodeComment_Default);
+
+		TexCoord->ConnectExpression(&Mask->Input, 0);
+
+		Mask->R = 1;
+		Mask->G = 0;
+		Mask->B = 0;
+		Mask->A = 0;
+		Mask->ConnectExpression(&IfExpr1->A, 0);
+		Mask->ConnectExpression(&IfExpr2->A, 0);
+
+		IfExpr1->ConstB = 0.33;
+		IfExpr1->ConnectExpression(&IfExpr2->ALessThanB, 0);
+		IfExpr1->ConnectExpression(&IfExpr2->AEqualsB, 0);
+
+		IfExpr2->ConstB = 0.66;
+
+		InBuildState->AddOtherExpressions({TexCoord, Mask, IfExpr1, IfExpr1});
+
+		int32 OutputIndexOne;
+		int32 OutputChannelOne;
+		int32 OutputIndexTwo;
+		int32 OutputChannelTwo;
+		int32 OutputIndexThree;
+		int32 OutputChannelThree;
+
+		const TArray<UMaterialExpression*> InputOneExpressions = InBuildState->GetBuildUtils().CreateExpressionInputs(
+			InputConnectionMap,
+			InputOne.InputIndex,
+			InputOne.ChannelInputs,
+			OutputIndexOne,
+			OutputChannelOne
+		);
+
+		const TArray<UMaterialExpression*> InputTwoExpressions = InBuildState->GetBuildUtils().CreateExpressionInputs(
+			InputConnectionMap,
+			InputTwo.InputIndex,
+			InputTwo.ChannelInputs,
+			OutputIndexTwo,
+			OutputChannelTwo
+		);
+
+		const TArray<UMaterialExpression*> InputThreeExpressions = InBuildState->GetBuildUtils().CreateExpressionInputs(
+			InputConnectionMap,
+			InputThree.InputIndex,
+			InputThree.ChannelInputs,
+			OutputIndexThree,
+			OutputChannelThree
+		);
+
+		InputOneExpressions.Last()->ConnectExpression(&IfExpr1->ALessThanB, OutputIndexOne);
+		InputOneExpressions.Last()->ConnectExpression(&IfExpr1->AEqualsB, OutputIndexOne);
+		InputTwoExpressions.Last()->ConnectExpression(&IfExpr1->AGreaterThanB, OutputIndexTwo);
+		InputThreeExpressions.Last()->ConnectExpression(&IfExpr2->AGreaterThanB, OutputIndexThree);
+
+		const TArray<FExpressionOutput>& OutputsOne = InputOneExpressions.Last()->GetOutputs();
+		const TArray<FExpressionOutput>& OutputsTwo = InputTwoExpressions.Last()->GetOutputs();
+		const TArray<FExpressionOutput>& OutputsThree = InputThreeExpressions.Last()->GetOutputs();
+
+		UE::DynamicMaterialEditor::Private::SetMask(IfExpr1->ALessThanB, OutputsOne[OutputIndexOne], OutputChannelOne);
+		UE::DynamicMaterialEditor::Private::SetMask(IfExpr1->AEqualsB, OutputsOne[OutputIndexOne], OutputChannelOne);
+		UE::DynamicMaterialEditor::Private::SetMask(IfExpr1->AGreaterThanB, OutputsTwo[OutputIndexTwo], OutputChannelTwo);
+		UE::DynamicMaterialEditor::Private::SetMask(IfExpr2->AGreaterThanB, OutputsThree[OutputIndexThree], OutputChannelThree);
+
+		InBuildState->GetBuildUtils().UpdatePreviewMaterial(IfExpr2, 0, FDMMaterialStageConnectorChannel::WHOLE_CHANNEL, 32);
+	}
+}
\ No newline at end of file
diff --git a/Engine/Plugins/Experimental/DynamicMaterial/Source/DynamicMaterialEditor/Private/Utils/DMInputNodeBuilder.h b/Engine/Plugins/Experimental/DynamicMaterial/Source/DynamicMaterialEditor/Private/Utils/DMInputNodeBuilder.h
new file mode 100644
--- /dev/null
+++ b/Engine/Plugins/Experimental/DynamicMaterial/Source/DynamicMaterialEditor/Private/Utils/DMInputNodeBuilder.h
@@ -0,0 +1,26 @@
+// Copyright Epic Games, Inc. All Rights Reserved.
+
+#pragma once
+
+#include "Templates/SharedPointer.h"
+
+class UDMMaterialStageInput;
+struct FDMMaterialBuildState;
+struct FDMMaterialStageConnection;
+
+namespace UE::DynamicMaterialEditor::Private
+{
+	struct FDMInputInputs;
+
+	void BuildExpressionInputs(const TSharedRef<FDMMaterialBuildState>& InBuildState, const TArray<FDMMaterialStageConnection>& InputConnectionMap,
+		const TArray<FDMInputInputs>& Inputs);
+
+	void BuildExpressionOneInput(const TSharedRef<FDMMaterialBuildState>& InBuildState, const TArray<FDMMaterialStageConnection>& InputConnectionMap,
+		const FDMInputInputs& InputOne);
+
+	void BuildExpressionTwoInputs(const TSharedRef<FDMMaterialBuildState>& InBuildState, const TArray<FDMMaterialStageConnection>& InputConnectionMap,
+		const FDMInputInputs& InputOne, const FDMInputInputs& InputTwo);
+
+	void BuildExpressionThreeInputs(const TSharedRef<FDMMaterialBuildState>& InBuildState, const TArray<FDMMaterialStageConnection>& InputConnectionMap,
+		const FDMInputInputs& InputOne, const FDMInputInputs& InputTwo, const FDMInputInputs& InputThree);
+}
diff --git a/Engine/Plugins/Experimental/DynamicMaterial/Source/DynamicMaterialEditor/Private/Utils/DMMaterialEffectStackPresetSubsystem.cpp b/Engine/Plugins/Experimental/DynamicMaterial/Source/DynamicMaterialEditor/Private/Utils/DMMaterialEffectStackPresetSubsystem.cpp
new file mode 100644
--- /dev/null
+++ b/Engine/Plugins/Experimental/DynamicMaterial/Source/DynamicMaterialEditor/Private/Utils/DMMaterialEffectStackPresetSubsystem.cpp
@@ -0,0 +1,264 @@
+// Copyright Epic Games, Inc. All Rights Reserved.
+
+#include "Utils/DMMaterialEffectStackPresetSubsystem.h"
+
+#include "Components/DMMaterialEffect.h"
+#include "Components/DMMaterialEffectStack.h"
+#include "Dom/JsonObject.h"
+#include "Dom/JsonValue.h"
+#include "DynamicMaterialEditorModule.h"
+#include "Editor.h"
+#include "HAL/FileManager.h"
+#include "Misc/FileHelper.h"
+#include "Serialization/JsonReader.h"
+#include "Serialization/JsonSerializer.h"
+#include "Serialization/JsonWriter.h"
+#include "Utils/DMJsonUtils.h"
+
+#define LOCTEXT_NAMESPACE "DMMaterialEffectStackPresetSubsystem"
+
+namespace UE::DynamicMaterialEditor::Private
+{
+	static const FString JsonPath = FString(TEXT("Config")) / FString(TEXT("MaterialEffectPresets"));
+	static const FString EffectListFieldName = TEXT("Effects");
+	static const FString EffectClassFieldName = TEXT("Class");
+	static const FString EffectDataFieldName = TEXT("Data");
+
+	FString GetEffectPresetBasePath(bool bInAllowCreate = true)
+	{
+		const FString ProjectPath = FPaths::ProjectDir() / JsonPath;
+
+		if (FPaths::DirectoryExists(ProjectPath))
+		{
+			return ProjectPath;
+		}
+
+		if (bInAllowCreate)
+		{
+			IFileManager::Get().MakeDirectory(*ProjectPath, /* Create tree */ true);
+
+			if (FPaths::DirectoryExists(ProjectPath))
+			{
+				return ProjectPath;
+			}
+		}
+
+		return "";
+	}
+
+	FString GetEffectPresetPath(const FString& InPresetName)
+	{
+		return GetEffectPresetBasePath() / InPresetName + TEXT(".json");
+	}
+
+	bool LoadPreset(const FString& InPath, FDMMaterialEffectStackJson& OutPreset)
+	{
+		FString FileJsonString;
+
+		if (!FFileHelper::LoadFileToString(FileJsonString, *InPath))
+		{
+			UE_LOG(LogDynamicMaterialEditor, Error, TEXT("Material Effect Preset Subsystem: LoadPreset() - Unable to load Json file: %s"), *InPath);
+			return false;
+		}
+
+		const TSharedPtr<TJsonReader<>> JsonReader = TJsonReaderFactory<>::Create(FileJsonString);
+		TSharedPtr<FJsonObject> RootObject;
+
+		if (!FJsonSerializer::Deserialize(JsonReader.ToSharedRef(), RootObject))
+		{
+			UE_LOG(LogDynamicMaterialEditor, Error, TEXT("Material Effect Preset Subsystem: LoadPreset() - Unable to parse file [%s]. Json=[%s]"), *InPath, *FileJsonString);
+			return false;
+		}
+
+		if (!RootObject->HasTypedField<EJson::Array>(EffectListFieldName))
+		{
+			UE_LOG(LogDynamicMaterialEditor, Error, TEXT("Material Effect Preset Subsystem: LoadPreset() - Malformed data in json [%s]. Json=[%s]"), *InPath, *FileJsonString);
+			return false;
+		}
+
+		const TArray<TSharedPtr<FJsonValue>> EffectList = RootObject->GetArrayField(EffectListFieldName);
+
+		FDMMaterialEffectStackJson MaterialEffectStackJson;
+		MaterialEffectStackJson.Effects.Reserve(EffectList.Num());
+
+		for (const TSharedPtr<FJsonValue>& JsonMaterialEffect : EffectList)
+		{
+			const TSharedPtr<FJsonObject>* JsonMaterialEffectObject = nullptr;
+
+			if (JsonMaterialEffect->TryGetObject(JsonMaterialEffectObject))
+			{
+				if (!(*JsonMaterialEffectObject)->HasTypedField<EJson::Object>(EffectDataFieldName))
+				{
+					UE_LOG(LogDynamicMaterialEditor, Error, TEXT("Material Effect Preset Subsystem: LoadPreset() - Malformed data in json [%s]. Json=[%s]"), *InPath, *FileJsonString);
+					return false;
+				}
+
+				FDMMaterialEffectJson MaterialEffectJson;
+
+				if (!FDMJsonUtils::Deserialize<UDMMaterialEffect>((*JsonMaterialEffectObject)->Values[EffectClassFieldName], MaterialEffectJson.Class))
+				{
+					UE_LOG(LogDynamicMaterialEditor, Error, TEXT("Material Effect Preset Subsystem: LoadPreset() - Invalid class in json [%s]. Json=[%s]"), *InPath, *FileJsonString);
+					return false;
+				}
+
+				MaterialEffectJson.Data = (*JsonMaterialEffectObject)->Values[EffectDataFieldName];
+
+				MaterialEffectStackJson.Effects.Add(MaterialEffectJson);
+			}
+			else
+			{
+				UE_LOG(LogDynamicMaterialEditor, Warning, TEXT("Material Effect Preset Subsystem: LoadPreset() - Malformed data in json [%s]. Json=[%s]"), *InPath, *FileJsonString);
+				return false;
+			}
+		}
+
+		OutPreset = MoveTemp(MaterialEffectStackJson);
+
+		return true;
+	}
+
+	bool SavePreset(const FString& InPath, const FDMMaterialEffectStackJson& InPreset)
+	{
+		TArray<TSharedPtr<FJsonValue>> EffectList;
+		EffectList.Reserve(InPreset.Effects.Num());
+
+		for (const FDMMaterialEffectJson& MaterialEffectJson : InPreset.Effects)
+		{
+			TSharedRef<FJsonObject> JsonMaterialEffectObject = MakeShared<FJsonObject>();
+			JsonMaterialEffectObject->SetField(EffectClassFieldName, FDMJsonUtils::Serialize(MaterialEffectJson.Class));
+			JsonMaterialEffectObject->SetField(EffectDataFieldName, MaterialEffectJson.Data);
+
+			EffectList.Add(MakeShared<FJsonValueObject>(JsonMaterialEffectObject));
+		}
+
+		TSharedRef<FJsonObject> RootObject = MakeShared<FJsonObject>();
+		RootObject->SetArrayField(EffectListFieldName, EffectList);
+
+		FString FileJsonString;
+		const TSharedRef<TJsonWriter<>> Writer = TJsonWriterFactory<>::Create(&FileJsonString);
+
+		if (!FJsonSerializer::Serialize(RootObject, Writer))
+		{
+			UE_LOG(LogDynamicMaterialEditor, Warning, TEXT("Material Effect Preset Subsystem: SavePreset() - Unable to serialize [%s]. Json=[%s]"), *InPath, *FileJsonString);
+			return false;
+		}
+
+		if (!FFileHelper::SaveStringToFile(FileJsonString, *InPath))
+		{
+			UE_LOG(LogDynamicMaterialEditor, Warning, TEXT("Material Effect Preset Subsystem: SavePreset() - Unable to save Json file: %s"), *InPath);
+			return false;
+		}
+
+		return true;
+	}
+}
+
+UDMMaterialEffectStackPresetSubsystem* UDMMaterialEffectStackPresetSubsystem::Get()
+{
+	if (!GEditor)
+	{
+		return nullptr;
+	}
+
+	return GEditor->GetEditorSubsystem<UDMMaterialEffectStackPresetSubsystem>();
+}
+
+bool UDMMaterialEffectStackPresetSubsystem::SavePreset(const FString& InPresetName, const FDMMaterialEffectStackJson& InPreset) const
+{
+	using namespace UE::DynamicMaterialEditor::Private;
+
+	if (InPreset.Effects.Num() == 0)
+	{
+		return false;
+	}
+
+	const FString Path = GetEffectPresetBasePath(/* Allow create */ true);
+
+	if (Path.IsEmpty() || !FPaths::DirectoryExists(Path))
+	{
+		return false;
+	}
+
+	const FString PresetFile = GetEffectPresetPath(InPresetName);
+
+	if (!UE::DynamicMaterialEditor::Private::SavePreset(PresetFile, InPreset))
+	{
+		UE_LOG(LogDynamicMaterialEditor, Error, TEXT("Material Effect Preset Subsystem: SavePreset() - Failed [%s]"), *InPresetName);
+		return false;
+	}
+
+	UE_LOG(LogDynamicMaterialEditor, Log, TEXT("Material Effect Preset Subsystem: SavePreset() - Success [%s]"), *InPresetName);
+	return true;
+}
+
+bool UDMMaterialEffectStackPresetSubsystem::LoadPreset(const FString& InPresetName, FDMMaterialEffectStackJson& OutPreset) const
+{
+	using namespace UE::DynamicMaterialEditor::Private;
+
+	const FString Path = GetEffectPresetBasePath(/* Allow create */ false);
+
+	if (Path.IsEmpty() || !FPaths::DirectoryExists(Path))
+	{
+		return false;
+	}
+
+	const FString PresetFile = GetEffectPresetPath(InPresetName);
+
+	if (!UE::DynamicMaterialEditor::Private::LoadPreset(PresetFile, OutPreset))
+	{
+		UE_LOG(LogDynamicMaterialEditor, Error, TEXT("Material Effect Preset Subsystem: LoadPreset() - Failed [%s]"), *InPresetName);
+		return false;
+	}
+
+	UE_LOG(LogDynamicMaterialEditor, Log, TEXT("Material Effect Preset Subsystem: LoadPreset() - Success [%s]"), *InPresetName);
+	return true;
+}
+
+bool UDMMaterialEffectStackPresetSubsystem::RemovePreset(const FString& InPresetName) const
+{
+	using namespace UE::DynamicMaterialEditor::Private;
+
+	const FString Path = GetEffectPresetBasePath(/* Allow create */ false);
+
+	if (Path.IsEmpty() || !FPaths::DirectoryExists(Path))
+	{
+		return false;
+	}
+
+	const FString PresetFile = GetEffectPresetPath(InPresetName);
+
+	if (!FPaths::FileExists(PresetFile))
+	{
+		return false;
+	}
+
+	return IFileManager::Get().Delete(*PresetFile);
+}
+
+TArray<FString> UDMMaterialEffectStackPresetSubsystem::GetPresetNames() const
+{
+	using namespace UE::DynamicMaterialEditor::Private;
+
+	const FString Path = GetEffectPresetBasePath(/* Allow create */ false);
+
+	if (Path.IsEmpty() || !FPaths::DirectoryExists(Path))
+	{
+		return {};
+	}
+
+	TArray<FString> JsonFiles;
+	IFileManager::Get().FindFiles(JsonFiles, *Path, TEXT(".json"));
+
+	TArray<FString> PresetNames;
+
+	for (const FString& JsonFile : JsonFiles)
+	{
+		const FString JsonFilePath = Path / JsonFile;
+		const FString PresetName = FPaths::GetBaseFilename(JsonFilePath, /* Remove Path */ true);
+		PresetNames.Add(PresetName);
+	}
+
+	return PresetNames;
+}
+
+#undef LOCTEXT_NAMESPACE
diff --git a/Engine/Plugins/Experimental/DynamicMaterial/Source/DynamicMaterialEditor/Private/Utils/DMMaterialFunctionLibrary.cpp b/Engine/Plugins/Experimental/DynamicMaterial/Source/DynamicMaterialEditor/Private/Utils/DMMaterialFunctionLibrary.cpp
new file mode 100644
--- /dev/null
+++ b/Engine/Plugins/Experimental/DynamicMaterial/Source/DynamicMaterialEditor/Private/Utils/DMMaterialFunctionLibrary.cpp
@@ -0,0 +1,154 @@
+// Copyright Epic Games, Inc. All Rights Reserved.
+
+#include "Utils/DMMaterialFunctionLibrary.h"
+#include "Materials/Material.h"
+#include "Materials/MaterialFunction.h"
+#include "Materials/MaterialExpressionMaterialFunctionCall.h"
+
+UMaterialFunctionInterface* FDMMaterialFunctionLibrary::LoadFunction(const FString& Path)
+{
+	return LoadObject<UMaterialFunctionInterface>(nullptr, *Path);
+}
+
+UMaterialFunctionInterface* FDMMaterialFunctionLibrary::GetFunction(const FName& Name, const FString& Path)
+{
+	TStrongObjectPtr<UMaterialFunctionInterface>* FunctionPtr = LoadedFunctions.Find(Name);
+
+	if (FunctionPtr)
+	{
+		return (*FunctionPtr).Get();
+	}
+
+	UMaterialFunctionInterface* LoadedFunction = LoadFunction(Path);
+
+	// Will emplace null to stop it uselessly looking it up in the future.
+	LoadedFunctions.Emplace(Name, TStrongObjectPtr<UMaterialFunctionInterface>(LoadedFunction));
+
+	return LoadedFunction;
+}
+
+UMaterialExpressionMaterialFunctionCall* FDMMaterialFunctionLibrary::MakeExpression(UMaterial* Parent, UMaterialFunctionInterface* Function, const FString& InComment)
+{
+	check(Parent);
+	check(Function);
+
+	UMaterialExpressionMaterialFunctionCall* MFC = NewObject<UMaterialExpressionMaterialFunctionCall>(Parent, NAME_None, RF_Transactional);
+	MFC->SetMaterialFunction(Function);
+	MFC->UpdateFromFunctionResource();
+	MFC->Desc = InComment;
+
+	Parent->GetEditorOnlyData()->ExpressionCollection.AddExpression(MFC);
+
+	return MFC;
+}
+
+UMaterialExpressionMaterialFunctionCall* FDMMaterialFunctionLibrary::MakeExpression(UMaterial* Parent, const FName& Name, const FString& Path, const FString& InComment)
+{
+	check(Parent);
+
+	UMaterialFunctionInterface* Function = GetFunction(Name, Path);
+
+	if (!Function)
+	{
+		return nullptr;
+	}
+
+	return MakeExpression(Parent, Function, InComment);
+}
+
+FDMMaterialFunctionLibrary& FDMMaterialFunctionLibrary::Get()
+{
+	static FDMMaterialFunctionLibrary FunctionLibrary;
+	return FunctionLibrary;
+}
+
+UMaterialExpressionMaterialFunctionCall* FDMMaterialFunctionLibrary::GetBreakOutFloat2Components(UMaterial* Parent, const FString& InComment)
+{
+	return MakeExpression(
+		Parent,
+		"BreakOutFloat2Components",
+		TEXT("/Engine/Functions/Engine_MaterialFunctions02/Utility/BreakOutFloat2Components.BreakOutFloat2Components"),
+		InComment
+	);
+}
+
+UMaterialExpressionMaterialFunctionCall* FDMMaterialFunctionLibrary::GetBreakOutFloat3Components(UMaterial* Parent, const FString& InComment)
+{
+	return MakeExpression(
+		Parent,
+		"BreakOutFloat3Components",
+		TEXT("/Engine/Functions/Engine_MaterialFunctions02/Utility/BreakOutFloat3Components.BreakOutFloat3Components"),
+		InComment
+	);
+}
+
+UMaterialExpressionMaterialFunctionCall* FDMMaterialFunctionLibrary::GetBreakOutFloat4Components(UMaterial* Parent, const FString& InComment)
+{
+	return MakeExpression(
+		Parent,
+		"BreakOutFloat4Components",
+		TEXT("/Engine/Functions/Engine_MaterialFunctions02/Utility/BreakOutFloat4Components.BreakOutFloat4Components"),
+		InComment
+	);
+}
+
+UMaterialExpressionMaterialFunctionCall* FDMMaterialFunctionLibrary::GetRGVtoHSV(UMaterial* Parent, const FString& InComment)
+{
+	return MakeExpression(
+		Parent,
+		"RGVtoHSV",
+		TEXT("/Engine/Functions/Engine_MaterialFunctions02/Utility/RGBtoHSV.RGBtoHSV"),
+		InComment
+	);
+}
+
+UMaterialExpressionMaterialFunctionCall* FDMMaterialFunctionLibrary::GetMakeFloat2(UMaterial* Parent, const FString& InComment)
+{
+	return MakeExpression(
+		Parent,
+		"MakeFloat2",
+		TEXT("/Engine/Functions/Engine_MaterialFunctions02/Utility/MakeFloat2.MakeFloat2"),
+		InComment
+	);
+}
+
+UMaterialExpressionMaterialFunctionCall* FDMMaterialFunctionLibrary::GetMakeFloat3(UMaterial* Parent, const FString& InComment)
+{
+	return MakeExpression(
+		Parent,
+		"MakeFloat3",
+		TEXT("/Engine/Functions/Engine_MaterialFunctions02/Utility/MakeFloat3.MakeFloat3"),
+		InComment
+	);
+}
+
+UMaterialExpressionMaterialFunctionCall* FDMMaterialFunctionLibrary::GetMakeFloat4(UMaterial* Parent, const FString& InComment)
+{
+	return MakeExpression(
+		Parent,
+		"MakeFloat4",
+		TEXT("/Engine/Functions/Engine_MaterialFunctions02/Utility/MakeFloat4.MakeFloat4"),
+		InComment
+	);
+}
+
+
+UMaterialExpressionMaterialFunctionCall* FDMMaterialFunctionLibrary::GetCustomRotator(UMaterial* Parent, const FString& InComment)
+{
+	return MakeExpression(
+		Parent,
+		"CustomRotator",
+		TEXT("/Engine/Functions/Engine_MaterialFunctions02/Texturing/CustomRotator.CustomRotator"),
+		InComment
+	);
+}
+
+UMaterialExpressionMaterialFunctionCall* FDMMaterialFunctionLibrary::GetRadialGradientExponential(UMaterial* Parent, const FString& InComment)
+{
+	return MakeExpression(
+		Parent,
+		"RadialGradientExponential",
+		TEXT("/Engine/Functions/Engine_MaterialFunctions01/Gradient/RadialGradientExponential.RadialGradientExponential"),
+		InComment
+	);
+}
diff --git a/Engine/Plugins/Experimental/DynamicMaterial/Source/DynamicMaterialEditor/Private/Utils/DMMaterialFunctionLibrary.h b/Engine/Plugins/Experimental/DynamicMaterial/Source/DynamicMaterialEditor/Private/Utils/DMMaterialFunctionLibrary.h
new file mode 100644
--- /dev/null
+++ b/Engine/Plugins/Experimental/DynamicMaterial/Source/DynamicMaterialEditor/Private/Utils/DMMaterialFunctionLibrary.h
@@ -0,0 +1,35 @@
+// Copyright Epic Games, Inc. All Rights Reserved.
+
+#pragma once
+
+#include "Materials/MaterialFunctionInterface.h"
+#include "UObject/StrongObjectPtr.h"
+
+class UMaterial;
+class UMaterialExpressionMaterialFunctionCall;
+
+class FDMMaterialFunctionLibrary
+{
+public:
+	static FDMMaterialFunctionLibrary& Get();
+
+	UMaterialFunctionInterface* GetFunction(const FName& Name, const FString& Path);
+
+	UMaterialFunctionInterface* LoadFunction(const FString& Path);
+
+	UMaterialExpressionMaterialFunctionCall* MakeExpression(UMaterial* Parent, UMaterialFunctionInterface* Function, const FString& InComment);
+	UMaterialExpressionMaterialFunctionCall* MakeExpression(UMaterial* Parent, const FName& Name, const FString& Path, const FString& InComment);
+
+	UMaterialExpressionMaterialFunctionCall* GetBreakOutFloat2Components(UMaterial* Parent, const FString& InComment);
+	UMaterialExpressionMaterialFunctionCall* GetBreakOutFloat3Components(UMaterial* Parent, const FString& InComment);
+	UMaterialExpressionMaterialFunctionCall* GetBreakOutFloat4Components(UMaterial* Parent, const FString& InComment);
+	UMaterialExpressionMaterialFunctionCall* GetRGVtoHSV(UMaterial* Parent, const FString& InComment);
+	UMaterialExpressionMaterialFunctionCall* GetMakeFloat2(UMaterial* Parent, const FString& InComment);
+	UMaterialExpressionMaterialFunctionCall* GetMakeFloat3(UMaterial* Parent, const FString& InComment);
+	UMaterialExpressionMaterialFunctionCall* GetMakeFloat4(UMaterial* Parent, const FString& InComment);
+	UMaterialExpressionMaterialFunctionCall* GetCustomRotator(UMaterial* Parent, const FString& InComment); /** Possibly promote this to a full expression? */
+	UMaterialExpressionMaterialFunctionCall* GetRadialGradientExponential(UMaterial* Parent, const FString& InComment); /** Possibly promote this to a full expression? */
+
+protected:
+	TMap<FName, TStrongObjectPtr<UMaterialFunctionInterface>> LoadedFunctions;
+};
diff --git a/Engine/Plugins/Experimental/DynamicMaterial/Source/DynamicMaterialEditor/Private/Utils/DMMaterialSnapshotLibrary.cpp b/Engine/Plugins/Experimental/DynamicMaterial/Source/DynamicMaterialEditor/Private/Utils/DMMaterialSnapshotLibrary.cpp
new file mode 100644
--- /dev/null
+++ b/Engine/Plugins/Experimental/DynamicMaterial/Source/DynamicMaterialEditor/Private/Utils/DMMaterialSnapshotLibrary.cpp
@@ -0,0 +1,177 @@
+// Copyright Epic Games, Inc. All Rights Reserved.
+
+#include "Utils/DMMaterialSnapshotLibrary.h"
+#include "DMAlphaOneMinusPS.h"
+#include "DynamicMaterialEditorModule.h"
+#include "Engine/Canvas.h"
+#include "Engine/Engine.h"
+#include "Framework/Notifications/NotificationManager.h"
+#include "ImageUtils.h"
+#include "Kismet/KismetRenderingLibrary.h"
+#include "Materials/MaterialInterface.h"
+#include "RenderGraphBuilder.h"
+#include "ScreenPass.h"
+#include "TextureResource.h"
+#include "UObject/Object.h"
+#include "UObject/Package.h"
+#include "UObject/UObjectGlobals.h"
+#include "Widgets/Notifications/SNotificationList.h"
+
+#define LOCTEXT_NAMESPACE "DMMaterialShapshotLibrary"
+
+namespace UE::DynamicMaterial::Private
+{
+	void RenderMaterialToRenderTarget(UMaterialInterface* InMaterial, UTextureRenderTarget2D* InRenderTarget)
+	{
+		InMaterial->EnsureIsComplete();
+
+		FTextureRenderTargetResource* RenderTargetResource = InRenderTarget->GameThread_GetRenderTargetResource();
+
+		UCanvas* Canvas = NewObject<UCanvas>(GetTransientPackage());
+
+		FCanvas RenderCanvas(
+			RenderTargetResource,
+			nullptr,
+			FGameTime::CreateUndilated(0, 0),
+			GEngine->GetDefaultWorldFeatureLevel()
+		);
+
+		Canvas->Init(InRenderTarget->SizeX, InRenderTarget->SizeY, nullptr, &RenderCanvas);
+
+		{
+			SCOPED_DRAW_EVENTF_GAMETHREAD(DrawMaterialToRenderTarget, TEXT("DrawMaterialToRenderTarget: %s"), InRenderTarget->GetFName());
+
+			ENQUEUE_RENDER_COMMAND(FlushDeferredResourceUpdateCommand)(
+				[RenderTargetResource](FRHICommandListImmediate& RHICmdList)
+				{
+					RenderTargetResource->FlushDeferredResourceUpdate(RHICmdList);
+				});
+
+			Canvas->K2_DrawMaterial(InMaterial, FVector2D(0, 0), FVector2D(InRenderTarget->SizeX, InRenderTarget->SizeY), FVector2D(0, 0));
+
+			RenderCanvas.Flush_GameThread();
+			Canvas->Canvas = nullptr;
+
+			//UpdateResourceImmediate must be called here to ensure mips are generated.
+			InRenderTarget->UpdateResourceImmediate(false);
+
+			ENQUEUE_RENDER_COMMAND(ResetSceneTextureExtentHistory)(
+				[RenderTargetResource](FRHICommandListImmediate& RHICmdList)
+				{
+					RenderTargetResource->ResetSceneTextureExtentsHistory();
+				});
+		}
+	}
+
+	UTextureRenderTarget2D* CreateSnapshotRenderTarget(const FIntPoint& InTextureSize)
+	{
+		UTextureRenderTarget2D* RenderTarget = NewObject<UTextureRenderTarget2D>(GetTransientPackage());
+		check(RenderTarget);
+		RenderTarget->RenderTargetFormat = RTF_RGBA32f;
+		RenderTarget->ClearColor = FLinearColor::Black;
+		RenderTarget->bAutoGenerateMips = false;
+		RenderTarget->bCanCreateUAV = false;
+		RenderTarget->InitAutoFormat(InTextureSize.X, InTextureSize.Y);
+		RenderTarget->UpdateResourceImmediate(true);
+
+		return RenderTarget;
+	}
+
+	void ApplyAlphaOneMinusShader(FRHICommandListImmediate& InRHICmdList, FTextureRenderTargetResource* InSourceTextureResource, FTextureRenderTargetResource* InDestTargetResource)
+	{
+		FRDGBuilder GraphBuilder(InRHICmdList);
+
+		const FTextureRHIRef& SourceTexture = InSourceTextureResource->GetRenderTargetTexture();
+		const FTextureRHIRef& DestTarget = InDestTargetResource->GetRenderTargetTexture();
+
+		FRDGTexture* InputTexture = GraphBuilder.RegisterExternalTexture(CreateRenderTarget(SourceTexture, TEXT("SourceTexture")));
+		FRDGTexture* OutputTexture = GraphBuilder.RegisterExternalTexture(CreateRenderTarget(DestTarget, TEXT("DestTarget")));
+		const FGlobalShaderMap* GlobalShaderMap = GetGlobalShaderMap(GMaxRHIFeatureLevel);
+
+		// The formats or sizes differ to pixel shader stuff
+		// Configure source/output viewport to get the right UV scaling from source texture to output texture
+		const FScreenPassTextureViewport InputViewport(InputTexture);
+		const FScreenPassTextureViewport OutputViewport(OutputTexture);
+
+		const TShaderMapRef<FScreenPassVS> VertexShader(GlobalShaderMap);
+
+		// Rectangle area to use from source
+		const FIntRect ViewRect(FIntPoint(0, 0), InputTexture->Desc.Extent);
+
+		//Dummy ViewFamily/ViewInfo created to use built in Draw Screen/Texture Pass
+		const FSceneViewFamily ViewFamily(FSceneViewFamily::ConstructionValues(nullptr, nullptr, FEngineShowFlags(ESFIM_Game))
+			.SetTime(FGameTime()));
+		FSceneViewInitOptions ViewInitOptions;
+		ViewInitOptions.ViewFamily = &ViewFamily;
+		ViewInitOptions.SetViewRectangle(ViewRect);
+		ViewInitOptions.ViewOrigin = FVector::ZeroVector;
+		ViewInitOptions.ViewRotationMatrix = FMatrix::Identity;
+		ViewInitOptions.ProjectionMatrix = FMatrix::Identity;
+		const FSceneView ViewInfo(ViewInitOptions);
+
+		const TShaderMapRef<FDMAlphaOneMinusPS> PixelShader(GlobalShaderMap);
+		FDMAlphaOneMinusPS::FParameters* Parameters = PixelShader->AllocateAndSetParameters(GraphBuilder, InputTexture, OutputTexture);
+		AddDrawScreenPass(GraphBuilder, RDG_EVENT_NAME("ApplyAlphaOneMinusShader"), ViewInfo, OutputViewport, InputViewport, VertexShader, PixelShader, Parameters);
+
+		GraphBuilder.Execute();
+	}
+
+	UTextureRenderTarget2D* ApplyAlphaOneMinusShader(UTextureRenderTarget2D* InRenderTarget)
+	{
+		FTextureRenderTargetResource* Source = InRenderTarget->GameThread_GetRenderTargetResource();
+
+		UTextureRenderTarget2D* OutRenderTarget = CreateSnapshotRenderTarget({InRenderTarget->SizeX, InRenderTarget->SizeY});
+		FTextureRenderTargetResource* Target = OutRenderTarget->GameThread_GetRenderTargetResource();
+
+		ENQUEUE_RENDER_COMMAND(ApplyAlphaOneMinusCommand)(
+			[Source, Target](FRHICommandListImmediate& InRHICmdList)
+			{
+				ApplyAlphaOneMinusShader(InRHICmdList, Source, Target);
+			});
+
+		return OutRenderTarget;
+	}
+}
+
+bool FDMMaterialShapshotLibrary::SnapshotMaterial(UMaterialInterface* InMaterial, const FIntPoint& InTextureSize, const FString& InSavePath)
+{
+	UTextureRenderTarget2D* RenderTarget = UE::DynamicMaterial::Private::CreateSnapshotRenderTarget(InTextureSize);
+
+	if (!RenderTarget)
+	{
+		UE_LOG(LogDynamicMaterialEditor, Warning, TEXT("Unable to create render target for material snapshot."));
+		return false;
+	}
+
+	UE::DynamicMaterial::Private::RenderMaterialToRenderTarget(InMaterial, RenderTarget);
+
+	// Render target has inverted alpha - fix that.
+	UTextureRenderTarget2D* FixedAlphaRenderTarget = UE::DynamicMaterial::Private::ApplyAlphaOneMinusShader(RenderTarget);
+
+	FImage Image;
+	FImageUtils::GetRenderTargetImage(FixedAlphaRenderTarget, Image);
+
+	FImageUtils::SaveImageByExtension(*InSavePath, Image);
+
+	RenderTarget->ReleaseResource();
+	FixedAlphaRenderTarget->ReleaseResource();
+
+	// Make a toast!
+	const FString HyperLinkText = FPaths::ConvertRelativePathToFull(InSavePath);
+
+	FNotificationInfo Info(LOCTEXT("SnapshotCreated", "Snapshot created."));
+
+	Info.Hyperlink = FSimpleDelegate::CreateStatic([](FString SourceFilePath)
+		{
+			FPlatformProcess::ExploreFolder(*(FPaths::GetPath(SourceFilePath)));
+		}, HyperLinkText);
+
+	Info.HyperlinkText = FText::FromString(HyperLinkText);
+	Info.ExpireDuration = 3.0f;
+
+	FSlateNotificationManager::Get().AddNotification(Info);
+
+	return true;
+}
+
+#undef LOCTEXT_NAMESPACE
diff --git a/Engine/Plugins/Experimental/DynamicMaterial/Source/DynamicMaterialEditor/Private/Utils/DMMaterialSnapshotLibrary.h b/Engine/Plugins/Experimental/DynamicMaterial/Source/DynamicMaterialEditor/Private/Utils/DMMaterialSnapshotLibrary.h
new file mode 100644
--- /dev/null
+++ b/Engine/Plugins/Experimental/DynamicMaterial/Source/DynamicMaterialEditor/Private/Utils/DMMaterialSnapshotLibrary.h
@@ -0,0 +1,14 @@
+// Copyright Epic Games, Inc. All Rights Reserved.
+
+#pragma once
+
+#include "Containers/UnrealString.h"
+#include "Math/MathFwd.h"
+
+class UMaterialInterface;
+
+class FDMMaterialShapshotLibrary
+{
+public:
+	static bool SnapshotMaterial(UMaterialInterface* InMaterial, const FIntPoint& InTextureSize, const FString& InSavePath);
+};
diff --git a/Engine/Plugins/Experimental/DynamicMaterial/Source/DynamicMaterialEditor/Private/Utils/DMPrivate.cpp b/Engine/Plugins/Experimental/DynamicMaterial/Source/DynamicMaterialEditor/Private/Utils/DMPrivate.cpp
new file mode 100644
--- /dev/null
+++ b/Engine/Plugins/Experimental/DynamicMaterial/Source/DynamicMaterialEditor/Private/Utils/DMPrivate.cpp
@@ -0,0 +1,159 @@
+// Copyright Epic Games, Inc. All Rights Reserved.
+
+#include "Utils/DMPrivate.h"
+#include "Components/DMMaterialLayer.h"
+#include "Components/DMMaterialStage.h"
+#include "DMDefs.h"
+#include "DynamicMaterialEditorModule.h"
+#include "Framework/Notifications/NotificationManager.h"
+#include "MaterialExpressionIO.h"
+#include "Widgets/Notifications/SNotificationList.h"
+
+#define LOCTEXT_NAMESPACE "FDMMaterialLayerReference"
+
+namespace UE::DynamicMaterialEditor::Private
+{
+	void SetMask(FExpressionInput& InInputConnector, const FExpressionOutput& InOutputConnector, int32 InChannelOverride)
+	{
+		const bool bUseOutputMask = (InChannelOverride != FDMMaterialStageConnectorChannel::WHOLE_CHANNEL);
+
+		if (!bUseOutputMask)
+		{
+			InInputConnector.SetMask(InOutputConnector.Mask, InOutputConnector.MaskR, InOutputConnector.MaskG, InOutputConnector.MaskB, InOutputConnector.MaskA);
+		}
+		else
+		{
+			int32 MaskR = (!!(InChannelOverride & FDMMaterialStageConnectorChannel::FIRST_CHANNEL)) * (!!(InOutputConnector.MaskR));
+			int32 MaskG = (!!(InChannelOverride & FDMMaterialStageConnectorChannel::SECOND_CHANNEL)) * (!!(InOutputConnector.MaskG));
+			int32 MaskB = (!!(InChannelOverride & FDMMaterialStageConnectorChannel::THIRD_CHANNEL)) * (!!(InOutputConnector.MaskB));
+			int32 MaskA = (!!(InChannelOverride & FDMMaterialStageConnectorChannel::FOURTH_CHANNEL)) * (!!(InOutputConnector.MaskA));
+
+			int32 MatchingMasks = MaskR + MaskG + MaskB + MaskA;
+
+			if (MatchingMasks == 0)
+			{
+				InInputConnector.SetMask(InOutputConnector.Mask, InOutputConnector.MaskR, InOutputConnector.MaskG, InOutputConnector.MaskB, InOutputConnector.MaskA);
+			}
+			else
+			{
+				InInputConnector.SetMask(1, MaskR, MaskG, MaskB, MaskA);
+			}
+		}
+	}
+
+	int32 ChannelIndexToChannelBit(int32 InChannelIndex)
+	{
+		switch (InChannelIndex)
+		{
+		case 0:
+			return FDMMaterialStageConnectorChannel::WHOLE_CHANNEL;
+
+		case 1:
+			return FDMMaterialStageConnectorChannel::FIRST_CHANNEL;
+
+		case 2:
+			return FDMMaterialStageConnectorChannel::SECOND_CHANNEL;
+
+		case 3:
+			return FDMMaterialStageConnectorChannel::THIRD_CHANNEL;
+
+		case 4:
+			return FDMMaterialStageConnectorChannel::FOURTH_CHANNEL;
+
+		default:
+			checkNoEntry();
+			return 0;
+		}
+	}
+
+	bool IsCustomMaterialProperty(EDMMaterialPropertyType InMaterialProperty)
+	{
+		return (InMaterialProperty >= EDMMaterialPropertyType::Custom1 && InMaterialProperty <= EDMMaterialPropertyType::Custom4);
+	}
+
+	static bool bAllowUIFeedback = false;
+	static FText LogErrorObjectFormat = LOCTEXT("LogErrorObjectFormat", "%s (Source: %s)");
+
+	void LogError(const FString& InMessage, bool bInToast, const UObject* InSource)
+	{
+		const FString Message = IsValid(InSource)
+			? FString::Printf(TEXT("%s (Source: %s)"), *InMessage, *InSource->GetPathName())
+			: InMessage;
+
+		UE_LOG(LogDynamicMaterialEditor, Error, TEXT("%s"), *Message);
+
+		if (bAllowUIFeedback && bInToast)
+		{
+			const FText MessageText = IsValid(InSource)
+				? FText::Format(LogErrorObjectFormat, FText::FromString(InMessage), FText::FromString(InSource->GetPathName()))
+				: FText::FromString(InMessage);
+
+			FNotificationInfo Info(MessageText);
+			Info.ExpireDuration = 5.0f;
+			FSlateNotificationManager::Get().AddNotification(Info);
+		}
+	}
+}
+
+FDMMaterialLayerReference::FDMMaterialLayerReference()
+	: FDMMaterialLayerReference(nullptr)
+{
+}
+
+FDMMaterialLayerReference::FDMMaterialLayerReference(UDMMaterialLayerObject* InLayer)
+{
+	LayerWeak = InLayer;
+}
+
+UDMMaterialLayerObject* FDMMaterialLayerReference::GetLayer() const
+{
+	return LayerWeak.Get();
+}
+
+bool FDMMaterialLayerReference::IsBaseEnabled() const
+{
+	if (const UDMMaterialLayerObject* Layer = GetLayer())
+	{
+		return Layer->IsStageEnabled(EDMMaterialLayerStage::Base);
+	}
+
+	return false;
+}
+
+bool FDMMaterialLayerReference::IsBaseBeingEdited() const
+{
+	if (const UDMMaterialLayerObject* Layer = GetLayer())
+	{
+		return Layer->IsStageBeingEdited(EDMMaterialLayerStage::Base);
+	}
+
+	return false;
+}
+
+bool FDMMaterialLayerReference::IsMaskEnabled() const
+{
+	if (const UDMMaterialLayerObject* Layer = GetLayer())
+	{
+		return Layer->IsStageEnabled(EDMMaterialLayerStage::Mask);
+	}
+
+	return false;
+}
+
+bool FDMMaterialLayerReference::IsMaskBeingEdited() const
+{
+	if (const UDMMaterialLayerObject* Layer = GetLayer())
+	{
+		return Layer->IsStageBeingEdited(EDMMaterialLayerStage::Mask);
+	}
+
+	return false;
+}
+
+FDMScopedUITransaction::FDMScopedUITransaction(const FText& InSessionName, bool bInShouldActuallyTransact)
+	: Transaction(FScopedTransaction(InSessionName, bInShouldActuallyTransact))
+	, UIFeedbackGuard(TGuardValue<bool>(UE::DynamicMaterialEditor::Private::bAllowUIFeedback, true))
+{
+}
+
+#undef LOCTEXT_NAMESPACE
diff --git a/Engine/Plugins/Experimental/DynamicMaterial/Source/DynamicMaterialEditor/Private/Utils/DMPrivate.h b/Engine/Plugins/Experimental/DynamicMaterial/Source/DynamicMaterialEditor/Private/Utils/DMPrivate.h
new file mode 100644
--- /dev/null
+++ b/Engine/Plugins/Experimental/DynamicMaterial/Source/DynamicMaterialEditor/Private/Utils/DMPrivate.h
@@ -0,0 +1,57 @@
+// Copyright Epic Games, Inc. All Rights Reserved.
+
+#pragma once
+
+#include "Containers/Array.h"
+#include "HAL/Platform.h"
+#include "ScopedTransaction.h"
+#include "Slate/SDMSlot.h"
+#include "UObject/WeakObjectPtr.h"
+
+class UDMMaterialLayerObject;
+class UDMMaterialSlot;
+class UDMMaterialStage;
+class UDMMaterialStageInput;
+struct FExpressionInput;
+struct FExpressionOutput;
+
+namespace UE::DynamicMaterialEditor::Private
+{
+	struct FDMInputInputs
+	{
+		int32 InputIndex;
+		TArray<UDMMaterialStageInput*> ChannelInputs;
+	};
+
+	void SetMask(FExpressionInput& InInputConnector, const FExpressionOutput& InOutputConnector, int32 InChannelOverride);
+
+	/** Converts 0,1,2,3,4 to 0,1,2,4,8 */
+	int32 ChannelIndexToChannelBit(int32 InChannelIndex);
+
+	bool IsCustomMaterialProperty(EDMMaterialPropertyType InMaterialProperty);
+
+	void LogError(const FString& InMessage, bool bInToast = true, const UObject* InSource = nullptr);
+}
+
+struct FDMMaterialLayerReference
+{
+	TWeakObjectPtr<UDMMaterialLayerObject> LayerWeak;
+
+	FDMMaterialLayerReference();
+	FDMMaterialLayerReference(UDMMaterialLayerObject* InLayer);
+
+	UDMMaterialLayerObject* GetLayer() const;
+
+	bool IsBaseEnabled() const;
+	bool IsBaseBeingEdited() const;
+	bool IsMaskEnabled() const;
+	bool IsMaskBeingEdited() const;
+};
+
+struct FDMScopedUITransaction
+{
+	FScopedTransaction Transaction;
+	TGuardValue<bool> UIFeedbackGuard;
+
+	FDMScopedUITransaction(const FText& InSessionName, bool bInShouldActuallyTransact = true);
+};
diff --git a/Engine/Plugins/Experimental/DynamicMaterial/Source/DynamicMaterialEditor/Public/Components/DMMaterialEffect.h b/Engine/Plugins/Experimental/DynamicMaterial/Source/DynamicMaterialEditor/Public/Components/DMMaterialEffect.h
--- a/Engine/Plugins/Experimental/DynamicMaterial/Source/DynamicMaterialEditor/Public/Components/DMMaterialEffect.h
+++ b/Engine/Plugins/Experimental/DynamicMaterial/Source/DynamicMaterialEditor/Public/Components/DMMaterialEffect.h
@@ -3,6 +3,7 @@
 #pragma once
 
 #include "Components/DMMaterialComponent.h"
+#include "Utils/DMJsonUtils.h"
 #include "Components/DMMaterialStage.h"
 #include "DMDefs.h"
 #include "DMMaterialEffect.generated.h"
@@ -25,7 +26,7 @@ enum class EDMMaterialEffectTarget : uint8
 };
 
 UCLASS(Abstract, BlueprintType, Blueprintable, ClassGroup = "Material Designer", meta = (DisplayName = "Material Designer Effect"))
-class DYNAMICMATERIALEDITOR_API UDMMaterialEffect : public UDMMaterialComponent
+class DYNAMICMATERIALEDITOR_API UDMMaterialEffect : public UDMMaterialComponent, public IDMJsonSerializable
 {
 	GENERATED_BODY()
 
@@ -65,6 +66,8 @@ public:
 	UFUNCTION(BlueprintPure, Category = "Material Designer")
 	virtual FText GetEffectDescription() const PURE_VIRTUAL(UDMMaterialEffect::ApplyTo, return FText::GetEmpty();)
 
+	virtual bool IsCompatibleWith(UDMMaterialEffect* InEffect) const { return true; }
+
 	virtual void ApplyTo(const TSharedRef<FDMMaterialBuildState>& InBuildState, TArray<UMaterialExpression*>& InOutExpressions, 
 		int32& InOutLastExpressionOutputChannel, int32& InLastExpressionOutputIndex) const PURE_VIRTUAL(UDMMaterialEffect::ApplyTo)
 
diff --git a/Engine/Plugins/Experimental/DynamicMaterial/Source/DynamicMaterialEditor/Public/Components/DMMaterialEffectFunction.h b/Engine/Plugins/Experimental/DynamicMaterial/Source/DynamicMaterialEditor/Public/Components/DMMaterialEffectFunction.h
--- a/Engine/Plugins/Experimental/DynamicMaterial/Source/DynamicMaterialEditor/Public/Components/DMMaterialEffectFunction.h
+++ b/Engine/Plugins/Experimental/DynamicMaterial/Source/DynamicMaterialEditor/Public/Components/DMMaterialEffectFunction.h
@@ -31,9 +31,17 @@ public:
 
 	const TArray<TObjectPtr<UDMMaterialValue>>& GetInputValues() const;
 
+#if WITH_EDITOR
+	//~ Begin IDMJsonSerializable
+	virtual TSharedPtr<FJsonValue> JsonSerialize() const override;
+	virtual bool JsonDeserialize(const TSharedPtr<FJsonValue>& InJsonValue) override;
+	//~ End IDMJsonSerializable
+#endif
+
 	//~ Begin UDMMaterialEffect
 	virtual FText GetEffectName() const override;
 	virtual FText GetEffectDescription() const override;
+	virtual bool IsCompatibleWith(UDMMaterialEffect* InEffect) const override;
 	virtual void ApplyTo(const TSharedRef<FDMMaterialBuildState>& InBuildState, TArray<UMaterialExpression*>& InOutExpressions,
 		int32& InOutLastExpressionOutputChannel, int32& InOutLastExpressionOutputIndex) const override;
 	//~ End UDMMaterialEffect
diff --git a/Engine/Plugins/Experimental/DynamicMaterial/Source/DynamicMaterialEditor/Public/Components/DMMaterialEffectStack.h b/Engine/Plugins/Experimental/DynamicMaterial/Source/DynamicMaterialEditor/Public/Components/DMMaterialEffectStack.h
--- a/Engine/Plugins/Experimental/DynamicMaterial/Source/DynamicMaterialEditor/Public/Components/DMMaterialEffectStack.h
+++ b/Engine/Plugins/Experimental/DynamicMaterial/Source/DynamicMaterialEditor/Public/Components/DMMaterialEffectStack.h
@@ -4,8 +4,10 @@
 
 #include "Components/DMMaterialComponent.h"
 #include "DMDefs.h"
+#include "Templates/SharedPointer.h"
 #include "DMMaterialEffectStack.generated.h"
 
+class FJsonValue;
 class UDMMaterialEffect;
 class UDMMaterialLayerObject;
 class UDMMaterialSlot;
@@ -14,6 +16,21 @@ class UMaterialExpression;
 enum class EDMMaterialEffectTarget : uint8;
 struct FDMMaterialBuildState;
 
+struct FDMMaterialEffectJson
+{
+	TSubclassOf<UDMMaterialEffect> Class = nullptr;
+	TSharedPtr<FJsonValue> Data = nullptr;
+};
+
+USTRUCT(BlueprintType)
+struct FDMMaterialEffectStackJson
+{
+	GENERATED_BODY()
+
+	bool bEnabled = false;
+	TArray<FDMMaterialEffectJson> Effects = {};
+};
+
 /**
  * Container for effects. Effects can be applied to either layers (on a per stage basis) or to slots.
  */
@@ -71,7 +88,7 @@ public:
 	bool AddEffect(UDMMaterialEffect* InEffect);
 
 	UFUNCTION(BlueprintCallable, Category = "Material Designer")
-	bool SetEffect(int32 InIndex, UDMMaterialEffect* InEffect);
+	UDMMaterialEffect* SetEffect(int32 InIndex, UDMMaterialEffect* InEffect);
 
 	UFUNCTION(BlueprintCallable, Category = "Material Designer", Meta = (DisplayName = "Move Effect (By Index)"))
 	bool BP_MoveEffectByIndex(int32 InIndex, int32 InNewIndex)
@@ -90,12 +107,12 @@ public:
 	bool MoveEffect(UDMMaterialEffect* InEffect, int32 InNewIndex);
 
 	UFUNCTION(BlueprintCallable, Category = "Material Designer", Meta = (DisplayName = "Remove Effect (By Index)"))
-	bool BP_RemoveEffectByIndex(int32 InIndex)
+	UDMMaterialEffect* BP_RemoveEffectByIndex(int32 InIndex)
 	{
 		return RemoveEffect(InIndex);
 	}
 
-	bool RemoveEffect(int32 InIndex);
+	UDMMaterialEffect* RemoveEffect(int32 InIndex);
 
 	UFUNCTION(BlueprintCallable, Category = "Material Designer", Meta = (DisplayName = "Remove Effect (By Value)"))
 	bool BP_RemoveEffectByValue(UDMMaterialEffect* InEffect)
@@ -108,6 +125,12 @@ public:
 	bool ApplyEffects(const TSharedRef<FDMMaterialBuildState>& InBuildState, EDMMaterialEffectTarget InEffectTarget,
 		TArray<UMaterialExpression*>& InOutStageExpressions, int32& InOutLastExpressionOutputChannel, int32& InOutLastExpressionOutputIndex) const;
 
+	UFUNCTION(BlueprintCallable, Category = "Material Designer")
+	FDMMaterialEffectStackJson CreatePreset();
+
+	UFUNCTION(BlueprintCallable, Category = "Material Designer")
+	void ApplyPreset(const FDMMaterialEffectStackJson& InPreset);
+
 	//~ Begin UDMMaterialComponent
 	virtual UDMMaterialComponent* GetParentComponent() const override;
 	virtual FString GetComponentPathComponent() const override;
@@ -128,6 +151,10 @@ protected:
 	UPROPERTY(VisibleInstanceOnly, BlueprintReadOnly, Instanced, Category = "Material Designer")
 	TArray<TObjectPtr<UDMMaterialEffect>> Effects;
 
+	TArray<UDMMaterialEffect*> GetIncompatibleEffects(UDMMaterialEffect* InEffect);
+
+	TArray<UDMMaterialEffect*> RemoveIncompatibleEffects(UDMMaterialEffect* InEffect);
+
 	//~ Begin UDMMaterialComponent
 	virtual UDMMaterialComponent* GetSubComponentByPath(FDMComponentPath& InPath, const FDMComponentPathSegment& InPathSegment) const override;
 	virtual void OnComponentAdded() override;
diff --git a/Engine/Plugins/Experimental/DynamicMaterial/Source/DynamicMaterialEditor/Public/Utils/DMMaterialEffectStackPresetSubsystem.h b/Engine/Plugins/Experimental/DynamicMaterial/Source/DynamicMaterialEditor/Public/Utils/DMMaterialEffectStackPresetSubsystem.h
new file mode 100644
--- /dev/null
+++ b/Engine/Plugins/Experimental/DynamicMaterial/Source/DynamicMaterialEditor/Public/Utils/DMMaterialEffectStackPresetSubsystem.h
@@ -0,0 +1,33 @@
+// Copyright Epic Games, Inc. All Rights Reserved.
+
+#pragma once
+
+#include "EditorSubsystem.h"
+#include "Containers/ContainersFwd.h"
+#include "Containers/UnrealString.h"
+#include "DMMaterialEffectStackPresetSubsystem.generated.h"
+
+class UDMMaterialEffect;
+class UDMMaterialEffectStack;
+struct FDMMaterialEffectStackJson;
+
+UCLASS(BlueprintType)
+class UDMMaterialEffectStackPresetSubsystem : public UEditorSubsystem
+{
+	GENERATED_BODY()
+
+public:
+	static UDMMaterialEffectStackPresetSubsystem* Get();
+
+	UFUNCTION(BlueprintCallable, Category = "Material Designer")
+	bool SavePreset(const FString& InPresetName, const FDMMaterialEffectStackJson& InPreset) const;
+
+	UFUNCTION(BlueprintCallable, Category = "Material Designer")
+	bool LoadPreset(const FString& InPresetName, FDMMaterialEffectStackJson& OutPreset) const;
+
+	UFUNCTION(BlueprintCallable, Category = "Material Designer")
+	bool RemovePreset(const FString& InPresetName) const;
+
+	UFUNCTION(BlueprintCallable, Category = "Material Designer")
+	TArray<FString> GetPresetNames() const;
+};
