Material Designer
- Horde build issue

diff --git a/Engine/Plugins/Experimental/DynamicMaterial/Source/DynamicMaterial/Private/Utils/DMJsonUtils.cpp b/Engine/Plugins/Experimental/DynamicMaterial/Source/DynamicMaterial/Private/Utils/DMJsonUtils.cpp
--- a/Engine/Plugins/Experimental/DynamicMaterial/Source/DynamicMaterial/Private/Utils/DMJsonUtils.cpp
+++ b/Engine/Plugins/Experimental/DynamicMaterial/Source/DynamicMaterial/Private/Utils/DMJsonUtils.cpp
@@ -126,7 +126,10 @@ TSharedPtr<FJsonValue> FDMJsonUtils::Serialize(const FObjectPtr& InObject)
 
 TSharedPtr<FJsonValue> FDMJsonUtils::Serialize(const TMap<FString, TSharedPtr<FJsonValue>>& InMap)
 {
-	return MakeShared<FJsonValueObject>(MakeShared<FJsonObject>(InMap));
+	TSharedPtr<FJsonObject> JsonObject = MakeShared<FJsonObject>();
+	JsonObject->Values = InMap;
+
+	return MakeShared<FJsonValueObject>(JsonObject);
 }
 
 bool FDMJsonUtils::DeserializeNumber(const TSharedPtr<FJsonValue>& InJsonValue, double& OutNumber)
