Material Designer
- Horde 459483

diff --git a/Engine/Plugins/Experimental/DynamicMaterial/Source/DynamicMaterial/Private/Components/DMMaterialValue.cpp b/Engine/Plugins/Experimental/DynamicMaterial/Source/DynamicMaterial/Private/Components/DMMaterialValue.cpp
--- a/Engine/Plugins/Experimental/DynamicMaterial/Source/DynamicMaterial/Private/Components/DMMaterialValue.cpp
+++ b/Engine/Plugins/Experimental/DynamicMaterial/Source/DynamicMaterial/Private/Components/DMMaterialValue.cpp
@@ -307,7 +307,17 @@ void UDMMaterialValue::SetMIDParameter(UMaterialInstanceDynamic* InMID) const
 {
 	checkNoEntry();
 }
- 
+
+TSharedPtr<FJsonValue> UDMMaterialValue::JsonSerialize() const
+{
+	return nullptr;
+}
+
+bool UDMMaterialValue::JsonDeserialize(const TSharedPtr<FJsonValue>& InJsonValue)
+{
+	return false;
+}
+
 #if WITH_EDITOR
 void UDMMaterialValue::DoClean()
 {
diff --git a/Engine/Plugins/Experimental/DynamicMaterial/Source/DynamicMaterial/Private/Components/DMRenderTargetRenderer.cpp b/Engine/Plugins/Experimental/DynamicMaterial/Source/DynamicMaterial/Private/Components/DMRenderTargetRenderer.cpp
--- a/Engine/Plugins/Experimental/DynamicMaterial/Source/DynamicMaterial/Private/Components/DMRenderTargetRenderer.cpp
+++ b/Engine/Plugins/Experimental/DynamicMaterial/Source/DynamicMaterial/Private/Components/DMRenderTargetRenderer.cpp
@@ -66,6 +66,16 @@ void UDMRenderTargetRenderer::FlushUpdateRenderTarget()
 	}
 }
 
+TSharedPtr<FJsonValue> UDMRenderTargetRenderer::JsonSerialize() const
+{
+	return nullptr;
+}
+
+bool UDMRenderTargetRenderer::JsonDeserialize(const TSharedPtr<FJsonValue>& InJsonValue)
+{
+	return false;
+}
+
 void UDMRenderTargetRenderer::PostLoad()
 {
 	Super::PostLoad();
diff --git a/Engine/Plugins/Experimental/DynamicMaterial/Source/DynamicMaterial/Public/Components/DMMaterialValue.h b/Engine/Plugins/Experimental/DynamicMaterial/Source/DynamicMaterial/Public/Components/DMMaterialValue.h
--- a/Engine/Plugins/Experimental/DynamicMaterial/Source/DynamicMaterial/Public/Components/DMMaterialValue.h
+++ b/Engine/Plugins/Experimental/DynamicMaterial/Source/DynamicMaterial/Public/Components/DMMaterialValue.h
@@ -130,6 +130,11 @@ public:
 #endif
  
 	virtual void SetMIDParameter(UMaterialInstanceDynamic* InMID) const;
+
+	//~ Begin IDMJsonSerializable
+	virtual TSharedPtr<FJsonValue> JsonSerialize() const override;
+	virtual bool JsonDeserialize(const TSharedPtr<FJsonValue>& InJsonValue) override;
+	//~ End IDMJsonSerializable
  
 	//~ Begin UDMMaterialComponent
 	virtual void Update(EDMUpdateType InUpdateType = EDMUpdateType::Value) override;
diff --git a/Engine/Plugins/Experimental/DynamicMaterial/Source/DynamicMaterial/Public/Components/DMRenderTargetRenderer.h b/Engine/Plugins/Experimental/DynamicMaterial/Source/DynamicMaterial/Public/Components/DMRenderTargetRenderer.h
--- a/Engine/Plugins/Experimental/DynamicMaterial/Source/DynamicMaterial/Public/Components/DMRenderTargetRenderer.h
+++ b/Engine/Plugins/Experimental/DynamicMaterial/Source/DynamicMaterial/Public/Components/DMRenderTargetRenderer.h
@@ -43,6 +43,11 @@ public:
 	UFUNCTION(BlueprintPure, Category = "Material Designer")
 	bool IsUpdating() const { return bUpdating; }
 
+	//~ Begin IDMJsonSerializable
+	virtual TSharedPtr<FJsonValue> JsonSerialize() const override;
+	virtual bool JsonDeserialize(const TSharedPtr<FJsonValue>& InJsonValue) override;
+	//~ End IDMJsonSerializable
+
 	//~ Begin UObject
 	virtual void PostLoad() override;
 	//~ End UObject
