Material Designer
- Added global alpha sliders for properties with which it can work.
- Added global offset and rotation sliders.

#virtualized

diff --git a/Engine/Plugins/Experimental/Avalanche/Source/AvalancheModifiers/Private/Modifiers/AvaGlobalOpacityModifier.cpp b/Engine/Plugins/Experimental/Avalanche/Source/AvalancheModifiers/Private/Modifiers/AvaGlobalOpacityModifier.cpp
--- a/Engine/Plugins/Experimental/Avalanche/Source/AvalancheModifiers/Private/Modifiers/AvaGlobalOpacityModifier.cpp
+++ b/Engine/Plugins/Experimental/Avalanche/Source/AvalancheModifiers/Private/Modifiers/AvaGlobalOpacityModifier.cpp
@@ -63,7 +63,7 @@ void UAvaGlobalOpacityModifier::OnActorMaterialAdded(UMaterialInstanceDynamic* I
 	{
 		if (const UDynamicMaterialModel* Model = MDI->GetMaterialModel())
 		{
-			if (UDMMaterialValue* GlobalOpacityValue = Model->GetGlobalOpacityValue())
+			if (UDMMaterialValue* GlobalOpacityValue = Model->GetGlobalParameterValue(UDynamicMaterialModel::GlobalOpacityValueName))
 			{
 				GlobalOpacityValue->GetOnUpdate().RemoveAll(this);
 				GlobalOpacityValue->GetOnUpdate().AddUObject(this, &UAvaGlobalOpacityModifier::OnDynamicMaterialValueChanged);
@@ -82,7 +82,7 @@ void UAvaGlobalOpacityModifier::OnActorMaterialRemoved(UMaterialInstanceDynamic*
 	{
 		if (const UDynamicMaterialModel* Model = MDI->GetMaterialModel())
 		{
-			if (UDMMaterialValue* GlobalOpacityValue = Model->GetGlobalOpacityValue())
+			if (UDMMaterialValue* GlobalOpacityValue = Model->GetGlobalParameterValue(UDynamicMaterialModel::GlobalOpacityValueName))
 			{
 				GlobalOpacityValue->GetOnUpdate().RemoveAll(this);
 			}
diff --git a/Engine/Plugins/Experimental/DynamicMaterial/Source/DynamicMaterial/Private/Model/DynamicMaterialModel.cpp b/Engine/Plugins/Experimental/DynamicMaterial/Source/DynamicMaterial/Private/Model/DynamicMaterialModel.cpp
--- a/Engine/Plugins/Experimental/DynamicMaterial/Source/DynamicMaterial/Private/Model/DynamicMaterialModel.cpp
+++ b/Engine/Plugins/Experimental/DynamicMaterial/Source/DynamicMaterial/Private/Model/DynamicMaterialModel.cpp
@@ -17,12 +17,32 @@
 #include "Materials/Material.h"
 #endif
 
-const FString UDynamicMaterialModel::ValuesPathToken           = FString(TEXT("Values"));
-const FString UDynamicMaterialModel::ParametersPathToken       = FString(TEXT("Parameters"));
-const FName UDynamicMaterialModel::GlobalOpacityPropertyName   = FName("GlobalOpacityValue");
-const FName UDynamicMaterialModel::GlobalOpacityParameterName  = FName("GlobalOpacity");
-const FName UDynamicMaterialModel::GlobalTilingPropertyName    = FName("GlobalTilingValue");
-const FName UDynamicMaterialModel::GlobalTilingParameterName   = FName("GlobalTiling");
+const FString UDynamicMaterialModel::ValuesPathToken = FString(TEXT("Values"));
+const FString UDynamicMaterialModel::ParametersPathToken = FString(TEXT("Parameters"));
+const FName UDynamicMaterialModel::GlobalOpacityValueName = FName(TEXT("GlobalOpacityValue"));
+const FName UDynamicMaterialModel::GlobalOpacityParameterName = FName(TEXT("GlobalOpacity"));
+const FName UDynamicMaterialModel::GlobalMetallicValueName = FName(TEXT("GlobalMetallicValue"));
+const FName UDynamicMaterialModel::GlobalMetallicParameterName = FName(TEXT("GlobalMetallic"));
+const FName UDynamicMaterialModel::GlobalSpecularValueName = FName(TEXT("GlobalSpecularValue"));
+const FName UDynamicMaterialModel::GlobalSpecularParameterName = FName(TEXT("GlobalSpecular"));
+const FName UDynamicMaterialModel::GlobalRoughnessValueName = FName(TEXT("GlobalRoughnessValue"));
+const FName UDynamicMaterialModel::GlobalRoughnessParameterName = FName(TEXT("GlobalRoughness"));
+const FName UDynamicMaterialModel::GlobalAnisotropyValueName = FName(TEXT("GlobalAnisotropyValue"));
+const FName UDynamicMaterialModel::GlobalAnisotropyParameterName = FName(TEXT("GlobalAnisotropy"));
+const FName UDynamicMaterialModel::GlobalWorldPositionOffsetValueName = FName(TEXT("GlobalWorldPositionOffsetValue"));
+const FName UDynamicMaterialModel::GlobalWorldPositionOffsetParameterName = FName(TEXT("GlobalWorldPositionOffset"));
+const FName UDynamicMaterialModel::GlobalAmbientOcclusionValueName = FName(TEXT("GlobalAmbientOcclusionValue"));
+const FName UDynamicMaterialModel::GlobalAmbientOcclusionParameterName = FName(TEXT("GlobalAmbientOcclusion"));
+const FName UDynamicMaterialModel::GlobalRefractionValueName = FName(TEXT("GlobalRefractionValue"));
+const FName UDynamicMaterialModel::GlobalRefractionParameterName = FName(TEXT("GlobalRefraction"));
+const FName UDynamicMaterialModel::GlobalPixelDepthOffsetValueName = FName(TEXT("GlobalPixelDepthOffsetValue"));
+const FName UDynamicMaterialModel::GlobalPixelDepthOffsetParameterName = FName(TEXT("GlobalPixelDepthOffset"));
+const FName UDynamicMaterialModel::GlobalOffsetValueName = FName(TEXT("GlobalOffsetValue"));
+const FName UDynamicMaterialModel::GlobalOffsetParameterName = FName(TEXT("GlobalOffset"));
+const FName UDynamicMaterialModel::GlobalTilingValueName = FName(TEXT("GlobalTilingValue"));
+const FName UDynamicMaterialModel::GlobalTilingParameterName = FName(TEXT("GlobalTiling"));
+const FName UDynamicMaterialModel::GlobalRotationValueName = FName(TEXT("GlobalRotationValue"));
+const FName UDynamicMaterialModel::GlobalRotationParameterName = FName(TEXT("GlobalRotation"));
 
 UDynamicMaterialModel::UDynamicMaterialModel()
 {
@@ -32,33 +52,59 @@ UDynamicMaterialModel::UDynamicMaterialModel()
 	EditorOnlyDataSI.SetObject(nullptr);
 #endif
 
-	FDMUpdateGuard Guard;
-	GlobalOpacityValue = CreateDefaultSubobject<UDMMaterialValueFloat1>(GlobalOpacityParameterName);
+	auto AddFloatParameter = [this](FName InPropertyName, FName InParameterName, float InDefaultValue = 1.f, bool bSetValueRange = true)
+		{
+			UDMMaterialValueFloat1* Property = CreateDefaultSubobject<UDMMaterialValueFloat1>(InPropertyName);
+			GlobalParameterValues.Add(InPropertyName, Property);
 
 #if WITH_EDITOR
-	GlobalOpacityValue->SetValueRange({0.f, 1.f});
-	GlobalOpacityValue->SetDefaultValue(1.f);
-	GlobalOpacityValue->ApplyDefaultValue();
+			if (bSetValueRange)
+			{
+				Property->SetValueRange({0.f, 1.f});
+			}
+
+			Property->SetDefaultValue(InDefaultValue);
+			Property->ApplyDefaultValue();
 #endif
 
-	GlobalOpacityParameter = CreateDefaultSubobject<UDMMaterialParameter>("GlobalOpacityParameter");
-	GlobalOpacityParameter->ParameterName = GlobalOpacityParameterName;
-	GlobalOpacityValue->Parameter = GlobalOpacityParameter;
+			UDMMaterialParameter* Parameter = CreateDefaultSubobject<UDMMaterialParameter>(InParameterName);
+			Parameter->ParameterName = InPropertyName;
+			Property->Parameter = Parameter;
 
-	ParameterMap.Add(GlobalOpacityParameterName, GlobalOpacityParameter);
+			ParameterMap.Add(InPropertyName, Parameter);
+		};
 
-	GlobalTilingValue = CreateDefaultSubobject<UDMMaterialValueFloat2>(GlobalTilingParameterName);
+	auto AddVector2Parameter = [this](FName InPropertyName, FName InParameterName, const FVector2D& InDefaultValue)
+		{
+			UDMMaterialValueFloat2* Property = CreateDefaultSubobject<UDMMaterialValueFloat2>(InPropertyName);
+			GlobalParameterValues.Add(InPropertyName, Property);
 
 #if WITH_EDITOR
-	GlobalTilingValue->SetDefaultValue(FVector2D(1.0, 1.0));
-	GlobalTilingValue->ApplyDefaultValue();
+			Property->SetDefaultValue(InDefaultValue);
+			Property->ApplyDefaultValue();
 #endif
 
-	GlobalTilingParameter = CreateDefaultSubobject<UDMMaterialParameter>("GlobalTilingParameter");
-	GlobalTilingParameter->ParameterName = GlobalTilingParameterName;
-	GlobalTilingValue->Parameter = GlobalTilingParameter;
+			UDMMaterialParameter* Parameter = CreateDefaultSubobject<UDMMaterialParameter>(InParameterName);
+			Parameter->ParameterName = InPropertyName;
+			Property->Parameter = Parameter;
+
+			ParameterMap.Add(InPropertyName, Parameter);
+		};
 
-	ParameterMap.Add(GlobalTilingParameterName, GlobalTilingParameter);
+	FDMUpdateGuard Guard;
+	AddFloatParameter(GlobalOpacityValueName, GlobalOpacityParameterName);
+	AddFloatParameter(GlobalMetallicValueName, GlobalMetallicParameterName);
+	AddFloatParameter(GlobalSpecularValueName, GlobalSpecularParameterName);
+	AddFloatParameter(GlobalRoughnessValueName, GlobalRoughnessParameterName);
+	AddFloatParameter(GlobalAnisotropyValueName, GlobalAnisotropyParameterName);
+	AddFloatParameter(GlobalWorldPositionOffsetValueName, GlobalWorldPositionOffsetParameterName);
+	AddFloatParameter(GlobalAmbientOcclusionValueName, GlobalAmbientOcclusionParameterName);
+	AddFloatParameter(GlobalRefractionValueName, GlobalRefractionParameterName);
+	AddFloatParameter(GlobalPixelDepthOffsetValueName, GlobalPixelDepthOffsetParameterName);
+
+	AddVector2Parameter(GlobalOffsetValueName, GlobalOffsetParameterName, FVector2D::ZeroVector);
+	AddVector2Parameter(GlobalTilingValueName, GlobalTilingParameterName, FVector2D::UnitVector);
+	AddFloatParameter(GlobalRotationValueName, GlobalRotationParameterName, /* Default Value */ 0.f, /* Set Value Range */ false);
 }
 
 void UDynamicMaterialModel::SetDynamicMaterialInstance(UDynamicMaterialInstance* InDynamicMaterialInstance)
@@ -73,6 +119,21 @@ void UDynamicMaterialModel::SetDynamicMaterialInstance(UDynamicMaterialInstance*
 #endif
 }
 
+UDMMaterialValueFloat1* UDynamicMaterialModel::GetGlobalOpacityValue() const
+{
+	return GetTypedGlobalParameterValue<UDMMaterialValueFloat1>(GlobalOpacityValueName);
+}
+
+UDMMaterialValue* UDynamicMaterialModel::GetGlobalParameterValue(FName InName) const
+{
+	if (const TObjectPtr<UDMMaterialValue>* ValuePtr = GlobalParameterValues.Find(InName))
+	{
+		return Cast<UDMMaterialValue>(*ValuePtr);
+	}
+	
+	return nullptr;
+}
+
 bool UDynamicMaterialModel::IsModelValid() const
 {
 	return (!HasAnyFlags(RF_BeginDestroyed | RF_FinishDestroyed) && IsValid(this));
@@ -368,6 +429,20 @@ void UDynamicMaterialModel::PostLoad()
 
 	SetFlags(RF_Transactional);
 
+	PRAGMA_DISABLE_DEPRECATION_WARNINGS
+
+	if (GlobalOpacityValue)
+	{
+		if (UDMMaterialValueFloat1* NewOpacityValue = Cast<UDMMaterialValueFloat1>(GetGlobalParameterValue(GlobalOpacityValueName)))
+		{
+			NewOpacityValue->SetValue(GlobalOpacityValue->GetValue());
+		}
+
+		GlobalOpacityValue = nullptr;
+	}
+
+	PRAGMA_ENABLE_DEPRECATION_WARNINGS
+
 	FixGlobalVars();
 
 	IDynamicMaterialModelEditorOnlyDataInterface* ModelEditorOnlyData = GetEditorOnlyData();
@@ -431,20 +506,23 @@ void UDynamicMaterialModel::PostEditorDuplicate()
 		Value->PostEditorDuplicate(this, nullptr);
 	}
 
-	if (GUndo)
+	for (const TPair<FName, TObjectPtr<UDMMaterialValue>>& ValuePair : GlobalParameterValues)
 	{
-		GlobalOpacityValue->Modify();
-	}
+		if (GUndo)
+		{
+			ValuePair.Value->Modify();
+		}
 
-	GlobalOpacityValue->PostEditorDuplicate(this, nullptr);
+		ValuePair.Value->PostEditorDuplicate(this, nullptr);
+	}
 
-	for (const TPair<FName, TWeakObjectPtr<UDMMaterialParameter>>& Pair : ParameterMap)
+	for (const TPair<FName, TWeakObjectPtr<UDMMaterialParameter>>& ParameterPair : ParameterMap)
 	{
-		if (UDMMaterialParameter* Parameter = Pair.Value.Get())
+		if (UDMMaterialParameter* Parameter = ParameterPair.Value.Get())
 		{
 			if (GUndo)
 			{
-				GlobalOpacityValue->Modify();
+				ParameterPair.Value->Modify();
 			}
 
 			Parameter->PostEditorDuplicate(this, nullptr);
@@ -469,16 +547,6 @@ void UDynamicMaterialModel::ReinitComponents()
 		DynamicMaterial->AtomicallySetFlags(RF_DuplicateTransient);
 	}
 
-	if (IsValid(GlobalOpacityValue) && GlobalOpacityValue->GetMaterialParameterName() != GlobalOpacityParameterName)
-	{
-		if (GUndo)
-		{
-			GlobalOpacityValue->Modify();
-		}
-
-		GlobalOpacityValue->SetParameterName(GlobalOpacityParameterName);
-	}
-
 	// Clean up old parameters
 	ParameterMap.Empty();
 
@@ -513,65 +581,65 @@ void UDynamicMaterialModel::ReinitComponents()
 
 void UDynamicMaterialModel::FixGlobalVars()
 {
-	if (const TWeakObjectPtr<UDMMaterialParameter>* ParameterPtr = ParameterMap.Find(GlobalOpacityParameterName))
-	{
-		UDMMaterialParameter* Parameter = ParameterPtr->Get();
-
-		if (!Parameter || Parameter->HasAnyFlags(RF_ArchetypeObject)
-			|| (GlobalOpacityParameter && Parameter != GlobalOpacityParameter))
+	auto FixGlobalVar = [this](FName InValueName, FName InParameterName)
 		{
-			if (GlobalOpacityParameter)
-			{
-				ParameterMap[GlobalOpacityParameterName] = GlobalOpacityParameter;
-			}
-			else
+			UDMMaterialValue* Property = GetGlobalParameterValue(InValueName);
+
+			if (!Property)
 			{
-				ParameterMap.Remove(GlobalOpacityParameterName);
+				return;
 			}
-		}
-	}
-
-	if (GlobalOpacityValue && GlobalOpacityParameter != GlobalOpacityValue->Parameter)
-	{
-		if (GUndo)
-		{
-			GlobalOpacityParameter->Modify();
-			GlobalOpacityValue->Modify();
-		}
 
-		GlobalOpacityParameter->ParameterName = GlobalOpacityParameterName;
-		GlobalOpacityValue->Parameter = GlobalOpacityParameter;
-	}
-
-	if (const TWeakObjectPtr<UDMMaterialParameter>* ParameterPtr = ParameterMap.Find(GlobalTilingParameterName))
-	{
-		UDMMaterialParameter* Parameter = ParameterPtr->Get();
+			UDMMaterialParameter* Parameter = Property->GetParameter();
 
-		if (!Parameter || Parameter->HasAnyFlags(RF_ArchetypeObject)
-			|| (GlobalTilingParameter && Parameter != GlobalTilingParameter))
-		{
-			if (GlobalTilingParameter)
+			if (!Parameter)
 			{
-				ParameterMap[GlobalTilingParameterName] = GlobalTilingParameter;
+				return;
 			}
-			else
+	
+			if (const TWeakObjectPtr<UDMMaterialParameter>* FoundParameterPtr = ParameterMap.Find(InParameterName))
 			{
-				ParameterMap.Remove(GlobalTilingParameterName);
+				UDMMaterialParameter* FoundParameter = FoundParameterPtr->Get();
+
+				if (!FoundParameter || FoundParameter->HasAnyFlags(RF_ArchetypeObject)
+					|| (Parameter && FoundParameter != Parameter))
+				{
+					if (Parameter)
+					{
+						ParameterMap[InParameterName] = Parameter;
+					}
+					else
+					{
+						ParameterMap.Remove(InParameterName);
+					}
+				}
 			}
-		}
-	}
 
-	if (GlobalTilingValue && GlobalTilingParameter != GlobalTilingValue->Parameter)
-	{
-		if (GUndo)
-		{
-			GlobalTilingParameter->Modify();
-			GlobalTilingValue->Modify();
-		}
+			if (Property && Parameter != Property->Parameter)
+			{
+				if (GUndo)
+				{
+					Parameter->Modify();
+					Property->Modify();
+				}
+
+				Parameter->ParameterName = InParameterName;
+				Property->Parameter = Parameter;
+			}
+		};
 
-		GlobalTilingParameter->ParameterName = GlobalTilingParameterName;
-		GlobalTilingValue->Parameter = GlobalTilingParameter;
-	}
+	FixGlobalVar(GlobalOpacityValueName, GlobalOpacityParameterName);
+	FixGlobalVar(GlobalRoughnessValueName, GlobalRoughnessParameterName);
+	FixGlobalVar(GlobalSpecularValueName, GlobalSpecularParameterName);
+	FixGlobalVar(GlobalMetallicValueName, GlobalMetallicParameterName);
+	FixGlobalVar(GlobalAnisotropyValueName, GlobalAnisotropyParameterName);
+	FixGlobalVar(GlobalWorldPositionOffsetValueName, GlobalWorldPositionOffsetParameterName);
+	FixGlobalVar(GlobalAmbientOcclusionValueName, GlobalAmbientOcclusionParameterName);
+	FixGlobalVar(GlobalRefractionValueName, GlobalRefractionParameterName);
+	FixGlobalVar(GlobalPixelDepthOffsetValueName, GlobalPixelDepthOffsetParameterName);
+	FixGlobalVar(GlobalOffsetValueName, GlobalOffsetParameterName);
+	FixGlobalVar(GlobalTilingValueName, GlobalTilingParameterName);
+	FixGlobalVar(GlobalRotationValueName, GlobalRotationParameterName);
 }
 
 FName UDynamicMaterialModel::CreateUniqueParameterName(FName InBaseName)
diff --git a/Engine/Plugins/Experimental/DynamicMaterial/Source/DynamicMaterial/Public/Model/DynamicMaterialModel.h b/Engine/Plugins/Experimental/DynamicMaterial/Source/DynamicMaterial/Public/Model/DynamicMaterialModel.h
--- a/Engine/Plugins/Experimental/DynamicMaterial/Source/DynamicMaterial/Public/Model/DynamicMaterialModel.h
+++ b/Engine/Plugins/Experimental/DynamicMaterial/Source/DynamicMaterial/Public/Model/DynamicMaterialModel.h
@@ -43,10 +43,30 @@ class DYNAMICMATERIAL_API UDynamicMaterialModel : public UObject
 public:
 	static const FString ValuesPathToken;
 	static const FString ParametersPathToken;
-	static const FName GlobalOpacityPropertyName;
+	static const FName GlobalOpacityValueName;
 	static const FName GlobalOpacityParameterName;
-	static const FName GlobalTilingPropertyName;
+	static const FName GlobalMetallicValueName;
+	static const FName GlobalMetallicParameterName;
+	static const FName GlobalRoughnessValueName;
+	static const FName GlobalRoughnessParameterName;
+	static const FName GlobalSpecularValueName;
+	static const FName GlobalSpecularParameterName;
+	static const FName GlobalAnisotropyValueName;
+	static const FName GlobalAnisotropyParameterName;
+	static const FName GlobalWorldPositionOffsetValueName;
+	static const FName GlobalWorldPositionOffsetParameterName;
+	static const FName GlobalAmbientOcclusionValueName;
+	static const FName GlobalAmbientOcclusionParameterName;
+	static const FName GlobalRefractionValueName;
+	static const FName GlobalRefractionParameterName;
+	static const FName GlobalPixelDepthOffsetValueName;
+	static const FName GlobalPixelDepthOffsetParameterName;
+	static const FName GlobalOffsetValueName;
+	static const FName GlobalOffsetParameterName;
+	static const FName GlobalTilingValueName;
 	static const FName GlobalTilingParameterName;
+	static const FName GlobalRotationValueName;
+	static const FName GlobalRotationParameterName;
 
 	UDynamicMaterialModel();
 
@@ -58,11 +78,18 @@ public:
 	UFUNCTION(BlueprintPure, Category = "Material Designer")
 	UMaterial* GetGeneratedMaterial() const { return DynamicMaterial; }
 
+	UE_DEPRECATED(5.5, "Added to GlobalParameterValues map.")
 	UFUNCTION(BlueprintPure, Category = "Material Designer")
-	UDMMaterialValueFloat1* GetGlobalOpacityValue() const { return GlobalOpacityValue; }
+	UDMMaterialValueFloat1* GetGlobalOpacityValue() const;
 
 	UFUNCTION(BlueprintPure, Category = "Material Designer")
-	UDMMaterialValueFloat2* GetGlobalTilingValue() const { return GlobalTilingValue; }
+	UDMMaterialValue* GetGlobalParameterValue(FName InName) const;
+
+	template<typename InValueClass>
+	InValueClass* GetTypedGlobalParameterValue(FName InName) const
+	{
+		return Cast<InValueClass>(GetGlobalParameterValue(InName));
+	}
 
 	UFUNCTION(BlueprintCallable, Category = "Material Designer")
 	UDMMaterialComponent* GetComponentByPath(const FString& InPath) const;
@@ -155,16 +182,7 @@ protected:
 	TSet<TObjectPtr<UDMMaterialComponent>> RuntimeComponents;
 
 	UPROPERTY(VisibleInstanceOnly, BlueprintReadOnly, Category = "Material Designer")
-	TObjectPtr<UDMMaterialValueFloat1> GlobalOpacityValue;
-
-	UPROPERTY(VisibleInstanceOnly, BlueprintReadOnly, Category = "Material Designer")
-	TObjectPtr<UDMMaterialParameter> GlobalOpacityParameter;
-
-	UPROPERTY(VisibleInstanceOnly, BlueprintReadOnly, Category = "Material Designer")
-	TObjectPtr<UDMMaterialValueFloat2> GlobalTilingValue;
-
-	UPROPERTY(VisibleInstanceOnly, BlueprintReadOnly, Category = "Material Designer")
-	TObjectPtr<UDMMaterialParameter> GlobalTilingParameter;
+	TMap<FName, TObjectPtr<UDMMaterialValue>> GlobalParameterValues;
 
 	UPROPERTY(VisibleInstanceOnly, DuplicateTransient, TextExportTransient, Category = "Material Designer")
 	TMap<FName, TWeakObjectPtr<UDMMaterialParameter>> ParameterMap;
@@ -221,5 +239,9 @@ protected:
 	UE_DEPRECATED(5.3, "Moved to editor-only subobject.")
 	UPROPERTY()
 	bool bCreateMaterialPackage;
+
+	UE_DEPRECATED(5.5, "Moved to GlobalParameterValues map.")
+	UPROPERTY()
+	TObjectPtr<UDMMaterialValueFloat1> GlobalOpacityValue;
 #endif
 };
diff --git a/Engine/Plugins/Experimental/DynamicMaterial/Source/DynamicMaterial/Public/Model/IDynamicMaterialModelEditorOnlyDataInterface.h b/Engine/Plugins/Experimental/DynamicMaterial/Source/DynamicMaterial/Public/Model/IDynamicMaterialModelEditorOnlyDataInterface.h
--- a/Engine/Plugins/Experimental/DynamicMaterial/Source/DynamicMaterial/Public/Model/IDynamicMaterialModelEditorOnlyDataInterface.h
+++ b/Engine/Plugins/Experimental/DynamicMaterial/Source/DynamicMaterial/Public/Model/IDynamicMaterialModelEditorOnlyDataInterface.h
@@ -51,6 +51,7 @@ public:
 	virtual TSharedRef<IDMMaterialBuildStateInterface> CreateBuildStateInterface(UMaterial* InMaterialToBuild) const = 0;
 #endif
 
+	virtual UDMMaterialComponent* GetSubComponentByPath(FDMComponentPath& InPath) const = 0;
 	virtual UDMMaterialComponent* GetSubComponentByPath(FDMComponentPath& InPath, const FDMComponentPathSegment& InPathSegment) const = 0;
 
 protected:
diff --git a/Engine/Plugins/Experimental/DynamicMaterial/Source/DynamicMaterialEditor/Private/Components/DMMaterialProperty.cpp b/Engine/Plugins/Experimental/DynamicMaterial/Source/DynamicMaterialEditor/Private/Components/DMMaterialProperty.cpp
--- a/Engine/Plugins/Experimental/DynamicMaterial/Source/DynamicMaterialEditor/Private/Components/DMMaterialProperty.cpp
+++ b/Engine/Plugins/Experimental/DynamicMaterial/Source/DynamicMaterialEditor/Private/Components/DMMaterialProperty.cpp
@@ -1,6 +1,9 @@
 // Copyright Epic Games, Inc. All Rights Reserved.
 
 #include "Components/DMMaterialProperty.h"
+
+#include "DMComponentPath.h"
+#include "Components/DMMaterialComponent.h"
 #include "Components/DMMaterialSlot.h"
 #include "Components/DMMaterialStageThroughputLayerBlend.h"
 #include "Components/DMMaterialSubStage.h"
@@ -27,8 +30,10 @@
 
 #define LOCTEXT_NAMESPACE "DMMaterialProperty"
 
+const FString UDMMaterialProperty::ComponentsPathToken = TEXT("Components");
+
 UDMMaterialProperty::UDMMaterialProperty()
-	: UDMMaterialProperty(EDMMaterialPropertyType(0), EDMValueType::VT_Float1)
+	: UDMMaterialProperty(EDMMaterialPropertyType::None, EDMValueType::VT_Float1)
 {
 }
 
@@ -43,7 +48,8 @@ FString UDMMaterialProperty::GetComponentPathComponent() const
 	return StaticEnum<EDMMaterialPropertyType>()->GetNameStringByValue(static_cast<int64>(MaterialProperty));
 }
 
-UDMMaterialProperty* UDMMaterialProperty::CreateCustomMaterialPropertyDefaultSubobject(UDynamicMaterialModelEditorOnlyData* InModelEditorOnlyData, EDMMaterialPropertyType InMaterialProperty, const FName& InSubObjName)
+UDMMaterialProperty* UDMMaterialProperty::CreateCustomMaterialPropertyDefaultSubobject(UDynamicMaterialModelEditorOnlyData* InModelEditorOnlyData, 
+	EDMMaterialPropertyType InMaterialProperty, const FName& InSubObjName)
 {
 	check(InModelEditorOnlyData);
 	check(UE::DynamicMaterialEditor::Private::IsCustomMaterialProperty(InMaterialProperty));
@@ -197,6 +203,59 @@ void UDMMaterialProperty::OnSlotAdded(UDMMaterialSlot* InSlot)
 	}
 }
 
+UDMMaterialComponent* UDMMaterialProperty::AddComponent(FName InName, UDMMaterialComponent* InComponent)
+{
+	const TObjectPtr<UDMMaterialComponent>* CurrentComponentPtr = Components.Find(InName);
+
+	if (CurrentComponentPtr && IsValid(*CurrentComponentPtr))
+	{
+		(*CurrentComponentPtr)->SetComponentState(EDMComponentLifetimeState::Removed);
+	}
+
+	if (InComponent)
+	{
+		Components.FindOrAdd(InName) = InComponent;
+		InComponent->SetComponentState(EDMComponentLifetimeState::Added);
+	}
+	else if (CurrentComponentPtr)
+	{
+		Components.Remove(InName);
+	}
+
+	if (CurrentComponentPtr)
+	{
+		return *CurrentComponentPtr;
+	}
+
+	return nullptr;
+}
+
+bool UDMMaterialProperty::HasComponent(FName InName) const
+{
+	return Components.Contains(InName);
+}
+
+UDMMaterialComponent* UDMMaterialProperty::GetComponent(FName InName) const
+{
+	if (const TObjectPtr<UDMMaterialComponent>* CurrentComponentPtr = Components.Find(InName))
+	{
+		return *CurrentComponentPtr;
+	}
+
+	return nullptr;
+}
+
+UDMMaterialComponent* UDMMaterialProperty::RemoveComponent(FName InName)
+{
+	if (const TObjectPtr<UDMMaterialComponent>* CurrentComponentPtr = Components.Find(InName))
+	{
+		Components.Remove(InName);
+		return *CurrentComponentPtr;
+	}
+
+	return nullptr;
+}
+
 void UDMMaterialProperty::Update(EDMUpdateType InUpdateType)
 {
 	if (!IsComponentValid())
@@ -328,6 +387,52 @@ void UDMMaterialProperty::OnOutputProcessorUpdated()
 	Update(EDMUpdateType::Structure);
 }
 
+UDMMaterialComponent* UDMMaterialProperty::GetSubComponentByPath(FDMComponentPath& InPath, const FDMComponentPathSegment& InPathSegment) const
+{
+	if (InPathSegment.GetToken() == ComponentsPathToken)
+	{
+		FString ComponentString;
+
+		if (InPathSegment.GetParameter(ComponentString))
+		{
+			const FName ComponentName = *ComponentString;
+
+			if (const TObjectPtr<UDMMaterialComponent>* CurrentComponentPtr = Components.Find(ComponentName))
+			{
+				return *CurrentComponentPtr;
+			}
+		}
+	}
+
+	return Super::GetSubComponentByPath(InPath, InPathSegment);
+}
+
+void UDMMaterialProperty::OnComponentAdded()
+{
+	Super::OnComponentAdded();
+
+	for (const TPair<FName, TObjectPtr<UDMMaterialComponent>>& Pair : Components)
+	{
+		if (IsValid(Pair.Value))
+		{
+			Pair.Value->SetComponentState(EDMComponentLifetimeState::Added);
+		}
+	}
+}
+
+void UDMMaterialProperty::OnComponentRemoved()
+{
+	Super::OnComponentRemoved();
+
+	for (const TPair<FName, TObjectPtr<UDMMaterialComponent>>& Pair : Components)
+	{
+		if (IsValid(Pair.Value))
+		{
+			Pair.Value->SetComponentState(EDMComponentLifetimeState::Removed);
+		}
+	}
+}
+
 UMaterialExpression* UDMMaterialProperty::CreateConstant(const TSharedRef<FDMMaterialBuildState>& InBuildState,
 	float InDefaultValue)
 {
diff --git a/Engine/Plugins/Experimental/DynamicMaterial/Source/DynamicMaterialEditor/Private/Components/MaterialStageInputs/DMMSITextureUV.cpp b/Engine/Plugins/Experimental/DynamicMaterial/Source/DynamicMaterialEditor/Private/Components/MaterialStageInputs/DMMSITextureUV.cpp
--- a/Engine/Plugins/Experimental/DynamicMaterial/Source/DynamicMaterialEditor/Private/Components/MaterialStageInputs/DMMSITextureUV.cpp
+++ b/Engine/Plugins/Experimental/DynamicMaterial/Source/DynamicMaterialEditor/Private/Components/MaterialStageInputs/DMMSITextureUV.cpp
@@ -448,12 +448,24 @@ TArray<UMaterialExpression*> UDMMaterialStageInputTextureUV::CreateTextureUVExpr
 	CreateParameterExpression(ParamID::Scale,    ParamID::ScaleX,   TEXT("ScaleX"),   InTextureUV->GetScale().X);
 	CreateParameterExpression(ParamID::Scale,    ParamID::ScaleY,   TEXT("ScaleY"),   InTextureUV->GetScale().Y);
 
-	if (UMaterialExpression* GlobalTiling = InBuildState->GetGlobalTilingExpression())
+	if (UMaterialExpression* GlobalOffset = InBuildState->GetGlobalExpression(UDynamicMaterialModel::GlobalOffsetValueName))
+	{
+		check(NameToInputIndex.Contains(UDynamicMaterialModel::GlobalOffsetParameterName));
+		GlobalOffset->ConnectExpression(TextureUVFunc->GetInput(NameToInputIndex[UDynamicMaterialModel::GlobalOffsetParameterName]), 0);
+	}
+
+	if (UMaterialExpression* GlobalTiling = InBuildState->GetGlobalExpression(UDynamicMaterialModel::GlobalTilingValueName))
 	{
 		check(NameToInputIndex.Contains(UDynamicMaterialModel::GlobalTilingParameterName));
 		GlobalTiling->ConnectExpression(TextureUVFunc->GetInput(NameToInputIndex[UDynamicMaterialModel::GlobalTilingParameterName]), 0);
 	}
 
+	if (UMaterialExpression* GlobalRotation = InBuildState->GetGlobalExpression(UDynamicMaterialModel::GlobalRotationValueName))
+	{
+		check(NameToInputIndex.Contains(UDynamicMaterialModel::GlobalRotationParameterName));
+		GlobalRotation->ConnectExpression(TextureUVFunc->GetInput(NameToInputIndex[UDynamicMaterialModel::GlobalRotationParameterName]), 0);
+	}
+
 	// Output
 	Nodes.Add(TextureUVFunc);
 
diff --git a/Engine/Plugins/Experimental/DynamicMaterial/Source/DynamicMaterialEditor/Private/Menus/DMToolBarMenus.cpp b/Engine/Plugins/Experimental/DynamicMaterial/Source/DynamicMaterialEditor/Private/Menus/DMToolBarMenus.cpp
--- a/Engine/Plugins/Experimental/DynamicMaterial/Source/DynamicMaterialEditor/Private/Menus/DMToolBarMenus.cpp
+++ b/Engine/Plugins/Experimental/DynamicMaterial/Source/DynamicMaterialEditor/Private/Menus/DMToolBarMenus.cpp
@@ -50,6 +50,7 @@ namespace UE::DynamicMaterialEditor::Private
 					{
 						EditorWidget->RefreshSlotPickerList();
 						EditorWidget->RefreshSlotWidget();
+						EditorWidget->RefreshComponentEditWidget();
 					}
 				}
 			}
@@ -109,6 +110,7 @@ namespace UE::DynamicMaterialEditor::Private
 					{
 						EditorWidget->RefreshSlotPickerList();
 						EditorWidget->RefreshSlotWidget();
+						EditorWidget->RefreshComponentEditWidget();
 					}
 				}
 			}
@@ -184,6 +186,7 @@ namespace UE::DynamicMaterialEditor::Private
 					{
 						EditorWidget->RefreshSlotPickerList();
 						EditorWidget->RefreshSlotWidget();
+						EditorWidget->RefreshComponentEditWidget();
 					}
 				}
 			}
@@ -267,6 +270,8 @@ namespace UE::DynamicMaterialEditor::Private
 					if (TSharedPtr<SDMEditor> EditorWidget = Context->GetEditorWidget())
 					{
 						EditorWidget->RefreshSlotPickerList();
+						EditorWidget->RefreshSlotWidget();
+						EditorWidget->RefreshComponentEditWidget();
 					}
 				}
 			}
@@ -320,6 +325,8 @@ namespace UE::DynamicMaterialEditor::Private
 					if (TSharedPtr<SDMEditor> EditorWidget = Context->GetEditorWidget())
 					{
 						EditorWidget->RefreshSlotPickerList();
+						EditorWidget->RefreshSlotWidget();
+						EditorWidget->RefreshComponentEditWidget();
 					}					
 				}
 			}
@@ -368,6 +375,8 @@ namespace UE::DynamicMaterialEditor::Private
 					if (TSharedPtr<SDMEditor> EditorWidget = Context->GetEditorWidget())
 					{
 						EditorWidget->RefreshSlotPickerList();
+						EditorWidget->RefreshSlotWidget();
+						EditorWidget->RefreshComponentEditWidget();
 					}					
 				}
 			}
diff --git a/Engine/Plugins/Experimental/DynamicMaterial/Source/DynamicMaterialEditor/Private/Model/DMMaterialBuildState.cpp b/Engine/Plugins/Experimental/DynamicMaterial/Source/DynamicMaterialEditor/Private/Model/DMMaterialBuildState.cpp
--- a/Engine/Plugins/Experimental/DynamicMaterial/Source/DynamicMaterialEditor/Private/Model/DMMaterialBuildState.cpp
+++ b/Engine/Plugins/Experimental/DynamicMaterial/Source/DynamicMaterialEditor/Private/Model/DMMaterialBuildState.cpp
@@ -28,7 +28,6 @@ FDMMaterialBuildState::FDMMaterialBuildState(UMaterial* InDynamicMaterial, UDyna
 	, bIgnoreUVs(false)
 	, bIsPreviewMaterial(false)
 	, Utils(MakeShared<FDMMaterialBuildUtils>(*this))
-	, GlobalTilingExpression(nullptr)
 {
 	check(InDynamicMaterial);
 	check(InMaterialModel);
@@ -526,3 +525,21 @@ const TSet<UMaterialExpression*>& FDMMaterialBuildState::GetOtherExpressions()
 {
 	return OtherExpressions;
 }
+
+///////////////////////////////////////
+/// Global expression
+
+UMaterialExpression* FDMMaterialBuildState::GetGlobalExpression(FName InName) const
+{
+	if (UMaterialExpression* const* GlobalExpressionPtr = GlobalExpressions.Find(InName))
+	{
+		return *GlobalExpressionPtr;
+	}
+
+	return nullptr;
+}
+
+void FDMMaterialBuildState::SetGlobalExpression(FName InName, UMaterialExpression* InExpression)
+{
+	GlobalExpressions.FindOrAdd(InName) = InExpression;
+}
diff --git a/Engine/Plugins/Experimental/DynamicMaterial/Source/DynamicMaterialEditor/Private/Model/DynamicMaterialModelEditorOnlyData.cpp b/Engine/Plugins/Experimental/DynamicMaterial/Source/DynamicMaterialEditor/Private/Model/DynamicMaterialModelEditorOnlyData.cpp
--- a/Engine/Plugins/Experimental/DynamicMaterial/Source/DynamicMaterialEditor/Private/Model/DynamicMaterialModelEditorOnlyData.cpp
+++ b/Engine/Plugins/Experimental/DynamicMaterial/Source/DynamicMaterialEditor/Private/Model/DynamicMaterialModelEditorOnlyData.cpp
@@ -52,25 +52,6 @@
 
 #define LOCTEXT_NAMESPACE "MaterialDesignerModel"
 
-const FString UDynamicMaterialModelEditorOnlyData::SlotsPathToken               = FString(TEXT("Slots"));
-const FString UDynamicMaterialModelEditorOnlyData::RGBSlotPathToken             = FString(TEXT("RGBSlot"));
-const FString UDynamicMaterialModelEditorOnlyData::OpacitySlotPathToken         = FString(TEXT("OpacitySlot"));
-const FString UDynamicMaterialModelEditorOnlyData::RoughnessPathToken           = FString(TEXT("RoughnessSlot"));
-const FString UDynamicMaterialModelEditorOnlyData::SpecularPathToken            = FString(TEXT("SpecularSlot"));
-const FString UDynamicMaterialModelEditorOnlyData::MetallicPathToken            = FString(TEXT("MetallicSlot"));
-const FString UDynamicMaterialModelEditorOnlyData::NormalPathToken              = FString(TEXT("NormalSlot"));
-const FString UDynamicMaterialModelEditorOnlyData::PixelDepthOffsetPathToken    = FString(TEXT("PixelDepthOffsetSlot"));
-const FString UDynamicMaterialModelEditorOnlyData::WorldPositionOffsetPathToken = FString(TEXT("WorldPositionOffsetSlot"));
-const FString UDynamicMaterialModelEditorOnlyData::AmbientOcclusionPathToken    = FString(TEXT("AmbientOcclusionSlot"));
-const FString UDynamicMaterialModelEditorOnlyData::AnisotropyPathToken          = FString(TEXT("AnisotropySlot"));
-const FString UDynamicMaterialModelEditorOnlyData::RefractionPathToken          = FString(TEXT("RefractionSlot"));
-const FString UDynamicMaterialModelEditorOnlyData::TangentPathToken             = FString(TEXT("TangentSlot"));
-const FString UDynamicMaterialModelEditorOnlyData::Custom1PathToken             = FString(TEXT("Custom1Slot"));
-const FString UDynamicMaterialModelEditorOnlyData::Custom2PathToken             = FString(TEXT("Custom2Slot"));
-const FString UDynamicMaterialModelEditorOnlyData::Custom3PathToken             = FString(TEXT("Custom3Slot"));
-const FString UDynamicMaterialModelEditorOnlyData::Custom4PathToken             = FString(TEXT("Custom4Slot"));
-const FString UDynamicMaterialModelEditorOnlyData::PropertiesPathToken          = FString(TEXT("Properties"));
-
 namespace UE::DynamicMaterialEditor::Private
 {
 	const TMap<FString, EDMMaterialPropertyType> TokenToPropertyMap = {
@@ -93,6 +74,27 @@
 	};
 }
 
+const FString UDynamicMaterialModelEditorOnlyData::SlotsPathToken               = FString(TEXT("Slots"));
+const FString UDynamicMaterialModelEditorOnlyData::RGBSlotPathToken             = FString(TEXT("RGBSlot"));
+const FString UDynamicMaterialModelEditorOnlyData::OpacitySlotPathToken         = FString(TEXT("OpacitySlot"));
+const FString UDynamicMaterialModelEditorOnlyData::RoughnessPathToken           = FString(TEXT("RoughnessSlot"));
+const FString UDynamicMaterialModelEditorOnlyData::SpecularPathToken            = FString(TEXT("SpecularSlot"));
+const FString UDynamicMaterialModelEditorOnlyData::MetallicPathToken            = FString(TEXT("MetallicSlot"));
+const FString UDynamicMaterialModelEditorOnlyData::NormalPathToken              = FString(TEXT("NormalSlot"));
+const FString UDynamicMaterialModelEditorOnlyData::PixelDepthOffsetPathToken    = FString(TEXT("PixelDepthOffsetSlot"));
+const FString UDynamicMaterialModelEditorOnlyData::WorldPositionOffsetPathToken = FString(TEXT("WorldPositionOffsetSlot"));
+const FString UDynamicMaterialModelEditorOnlyData::AmbientOcclusionPathToken    = FString(TEXT("AmbientOcclusionSlot"));
+const FString UDynamicMaterialModelEditorOnlyData::AnisotropyPathToken          = FString(TEXT("AnisotropySlot"));
+const FString UDynamicMaterialModelEditorOnlyData::RefractionPathToken          = FString(TEXT("RefractionSlot"));
+const FString UDynamicMaterialModelEditorOnlyData::TangentPathToken             = FString(TEXT("TangentSlot"));
+const FString UDynamicMaterialModelEditorOnlyData::Custom1PathToken             = FString(TEXT("Custom1Slot"));
+const FString UDynamicMaterialModelEditorOnlyData::Custom2PathToken             = FString(TEXT("Custom2Slot"));
+const FString UDynamicMaterialModelEditorOnlyData::Custom3PathToken             = FString(TEXT("Custom3Slot"));
+const FString UDynamicMaterialModelEditorOnlyData::Custom4PathToken             = FString(TEXT("Custom4Slot"));
+const FString UDynamicMaterialModelEditorOnlyData::PropertiesPathToken          = FString(TEXT("Properties"));
+
+const FName UDynamicMaterialModelEditorOnlyData::AlphaValueName = TEXT("AlphaValue");
+
 const TArray<EMaterialDomain> UDynamicMaterialModelEditorOnlyData::SupportedDomains =
 {
 	EMaterialDomain::MD_Surface,
@@ -170,6 +172,20 @@
 	}
 }
 
+void UDynamicMaterialModelEditorOnlyData::AssignPropertyAlphaValues()
+{
+	Properties[EDMMaterialPropertyType::Opacity            ]->AddComponent(AlphaValueName, MaterialModel->GetGlobalParameterValue(UDynamicMaterialModel::GlobalOpacityValueName));
+	Properties[EDMMaterialPropertyType::OpacityMask        ]->AddComponent(AlphaValueName, MaterialModel->GetGlobalParameterValue(UDynamicMaterialModel::GlobalOpacityValueName));
+	Properties[EDMMaterialPropertyType::Metallic           ]->AddComponent(AlphaValueName, MaterialModel->GetGlobalParameterValue(UDynamicMaterialModel::GlobalMetallicValueName));
+	Properties[EDMMaterialPropertyType::Specular           ]->AddComponent(AlphaValueName, MaterialModel->GetGlobalParameterValue(UDynamicMaterialModel::GlobalSpecularValueName));
+	Properties[EDMMaterialPropertyType::Roughness          ]->AddComponent(AlphaValueName, MaterialModel->GetGlobalParameterValue(UDynamicMaterialModel::GlobalRoughnessValueName));
+	Properties[EDMMaterialPropertyType::Anisotropy         ]->AddComponent(AlphaValueName, MaterialModel->GetGlobalParameterValue(UDynamicMaterialModel::GlobalAnisotropyValueName));
+	Properties[EDMMaterialPropertyType::WorldPositionOffset]->AddComponent(AlphaValueName, MaterialModel->GetGlobalParameterValue(UDynamicMaterialModel::GlobalWorldPositionOffsetValueName));
+	Properties[EDMMaterialPropertyType::AmbientOcclusion   ]->AddComponent(AlphaValueName, MaterialModel->GetGlobalParameterValue(UDynamicMaterialModel::GlobalAmbientOcclusionValueName));
+	Properties[EDMMaterialPropertyType::Refraction         ]->AddComponent(AlphaValueName, MaterialModel->GetGlobalParameterValue(UDynamicMaterialModel::GlobalRefractionValueName));
+	Properties[EDMMaterialPropertyType::PixelDepthOffset   ]->AddComponent(AlphaValueName, MaterialModel->GetGlobalParameterValue(UDynamicMaterialModel::GlobalPixelDepthOffsetValueName));
+}
+
 void UDynamicMaterialModelEditorOnlyData::Initialize()
 {
 	if (!Slots.IsEmpty())
@@ -177,6 +193,8 @@
 		return;
 	}
 
+	AssignPropertyAlphaValues();
+
 	// Will choose appropriate type between BaseColor and EmissiveColor
 	AddSlotForMaterialProperty(EDMMaterialPropertyType::BaseColor);
 }
@@ -298,14 +316,5 @@
 	const bool bIsPostProcessMaterial = Domain == EMaterialDomain::MD_PostProcess;
 
 	/**
-	 * Add global tiling parameter
-	 */
-	if (MaterialModel->GlobalTilingValue)
-	{
-		MaterialModel->GlobalTilingValue->GenerateExpression(BuildState);
-		BuildState->SetGlobalTilingExpression(BuildState->GetLastValueExpression(MaterialModel->GlobalTilingValue));
-	}
-
-	/**
 	 * Process slots to build base material inputs.
 	 */
@@ -366,6 +375,40 @@
 		{
 			MaterialPropertyPtr->OutputIndex = 0;
 		}
+
+		// Global opacity is handled later
+		if (Pair.Key == EDMMaterialPropertyType::Opacity || Pair.Key == EDMMaterialPropertyType::OpacityMask)
+		{
+			continue;
+		}
+
+		UDMMaterialValueFloat1* AlphaValue = Cast<UDMMaterialValueFloat1>(Pair.Value->GetComponent(AlphaValueName));
+
+		if (!AlphaValue)
+		{
+			continue;
+		}
+
+		AlphaValue->GenerateExpression(BuildState);
+		UMaterialExpression* AlphaValueExpression = BuildState->GetLastValueExpression(AlphaValue);
+
+		if (MaterialPropertyPtr->Expression)
+		{
+			UMaterialExpressionMultiply* OpacityMultiply = BuildState->GetBuildUtils().CreateExpression<UMaterialExpressionMultiply>(UE_DM_NodeComment_Default);
+			OpacityMultiply->A.Expression = MaterialPropertyPtr->Expression;
+			OpacityMultiply->A.Mask = MaterialPropertyPtr->Mask;
+			OpacityMultiply->A.MaskR = MaterialPropertyPtr->MaskR;
+			OpacityMultiply->A.MaskG = MaterialPropertyPtr->MaskG;
+			OpacityMultiply->A.MaskB = MaterialPropertyPtr->MaskB;
+			OpacityMultiply->A.MaskA = MaterialPropertyPtr->MaskA;
+			OpacityMultiply->A.OutputIndex = MaterialPropertyPtr->OutputIndex;
+
+			OpacityMultiply->B.Expression = AlphaValueExpression;
+			OpacityMultiply->B.SetMask(1, 1, 0, 0, 0);
+			OpacityMultiply->B.OutputIndex = 0;
+
+			MaterialPropertyPtr->Expression = OpacityMultiply;
+		}
 	}
 
 	if (!bIsPostProcessMaterial)
@@ -375,18 +418,18 @@
 		 */
 		EDMMaterialPropertyType GenerateOpacityInput = EDMMaterialPropertyType::None;
 
-		if (BlendMode == EBlendMode::BLEND_Translucent)
+		if (BlendMode == EBlendMode::BLEND_Masked)
 		{
-			if (GetSlotForMaterialProperty(EDMMaterialPropertyType::Opacity) == nullptr)
-			{
-				GenerateOpacityInput = EDMMaterialPropertyType::Opacity;
-			}
-		}
-		else if (BlendMode == EBlendMode::BLEND_Masked)
-		{
 			if (GetSlotForMaterialProperty(EDMMaterialPropertyType::OpacityMask) == nullptr)
 			{
 				GenerateOpacityInput = EDMMaterialPropertyType::OpacityMask;
 			}
 		}
+		else if (BlendMode != EBlendMode::BLEND_Opaque)
+		{
+			if (GetSlotForMaterialProperty(EDMMaterialPropertyType::Opacity) == nullptr)
+			{
+				GenerateOpacityInput = EDMMaterialPropertyType::Opacity;
+			}
+		}
 
@@ -435,7 +478,7 @@
 		}
 
 		/**
-		 * Apply global opacity slider
+		 * Apply global opacity slider after automatic opacity generation
 		 */
 		FExpressionInput* OpacityPropertyPtr = nullptr;
 
@@ -448,29 +491,31 @@
 			OpacityPropertyPtr = BuildState->GetMaterialProperty(EDMMaterialPropertyType::OpacityMask);
 		}
 
-		if (OpacityPropertyPtr != nullptr && IsValid(MaterialModel) && IsValid(MaterialModel->GlobalOpacityValue))
+		if (OpacityPropertyPtr != nullptr)
 		{
-			MaterialModel->GlobalOpacityValue->GenerateExpression(BuildState);
-
-			UMaterialExpression* GlobalOpacityExpression = BuildState->GetLastValueExpression(MaterialModel->GlobalOpacityValue);
-
-			if (OpacityPropertyPtr->Expression)
+			if (UDMMaterialValueFloat1* GlobalOpacityValue = MaterialModel->GetTypedGlobalParameterValue<UDMMaterialValueFloat1>(UDynamicMaterialModel::GlobalOpacityValueName))
 			{
-				UMaterialExpressionMultiply* OpacityMultiply = BuildState->GetBuildUtils().CreateExpression<UMaterialExpressionMultiply>(UE_DM_NodeComment_Default);
-				OpacityMultiply->A.Expression = OpacityPropertyPtr->Expression;
-				OpacityMultiply->A.Mask = OpacityPropertyPtr->Mask;
-				OpacityMultiply->A.MaskR = OpacityPropertyPtr->MaskR;
-				OpacityMultiply->A.MaskG = OpacityPropertyPtr->MaskG;
-				OpacityMultiply->A.MaskB = OpacityPropertyPtr->MaskB;
-				OpacityMultiply->A.MaskA = OpacityPropertyPtr->MaskA;
-				OpacityMultiply->A.OutputIndex = OpacityPropertyPtr->OutputIndex;
+				GlobalOpacityValue->GenerateExpression(BuildState);
+				UMaterialExpression* GlobalOpacityExpression = BuildState->GetLastValueExpression(GlobalOpacityValue);
 
-				OpacityMultiply->B.Expression = GlobalOpacityExpression;
-				OpacityMultiply->B.SetMask(1, 1, 0, 0, 0);
-				OpacityMultiply->B.OutputIndex = 0;
+				if (GlobalOpacityExpression && OpacityPropertyPtr->Expression)
+				{
+					UMaterialExpressionMultiply* OpacityMultiply = BuildState->GetBuildUtils().CreateExpression<UMaterialExpressionMultiply>(UE_DM_NodeComment_Default);
+					OpacityMultiply->A.Expression = OpacityPropertyPtr->Expression;
+					OpacityMultiply->A.Mask = OpacityPropertyPtr->Mask;
+					OpacityMultiply->A.MaskR = OpacityPropertyPtr->MaskR;
+					OpacityMultiply->A.MaskG = OpacityPropertyPtr->MaskG;
+					OpacityMultiply->A.MaskB = OpacityPropertyPtr->MaskB;
+					OpacityMultiply->A.MaskA = OpacityPropertyPtr->MaskA;
+					OpacityMultiply->A.OutputIndex = OpacityPropertyPtr->OutputIndex;
 
-				OpacityPropertyPtr->Expression = OpacityMultiply;
-				OpacityPropertyPtr->SetMask(1, 1, 0, 0, 0);
+					OpacityMultiply->B.Expression = GlobalOpacityExpression;
+					OpacityMultiply->B.SetMask(1, 1, 0, 0, 0);
+					OpacityMultiply->B.OutputIndex = 0;
+
+					OpacityPropertyPtr->Expression = OpacityMultiply;
+					OpacityPropertyPtr->SetMask(1, 1, 0, 0, 0);
+				}
 			}
 		}
 	}
@@ -675,7 +720,30 @@
 	check(MaterialModel);
 	check(InMaterialToBuild);
 
-	return MakeShared<FDMMaterialBuildState>(InMaterialToBuild, MaterialModel, bInDirtyAssets);
+	TSharedRef<FDMMaterialBuildState> BuildState = MakeShared<FDMMaterialBuildState>(InMaterialToBuild, MaterialModel, bInDirtyAssets);
+
+	/**
+	 * Add global UV parameters
+	 */
+	if (UDMMaterialValue* GlobalOffsetValue = MaterialModel->GetGlobalParameterValue(UDynamicMaterialModel::GlobalOffsetValueName))
+	{
+		GlobalOffsetValue->GenerateExpression(BuildState);
+		BuildState->SetGlobalExpression(UDynamicMaterialModel::GlobalOffsetValueName, BuildState->GetLastValueExpression(GlobalOffsetValue));
+	}
+
+	if (UDMMaterialValue* GlobalTilingValue = MaterialModel->GetGlobalParameterValue(UDynamicMaterialModel::GlobalTilingValueName))
+	{
+		GlobalTilingValue->GenerateExpression(BuildState);
+		BuildState->SetGlobalExpression(UDynamicMaterialModel::GlobalTilingValueName, BuildState->GetLastValueExpression(GlobalTilingValue));
+	}
+
+	if (UDMMaterialValue* GlobalRotationValue = MaterialModel->GetGlobalParameterValue(UDynamicMaterialModel::GlobalRotationValueName))
+	{
+		GlobalRotationValue->GenerateExpression(BuildState);
+		BuildState->SetGlobalExpression(UDynamicMaterialModel::GlobalRotationValueName, BuildState->GetLastValueExpression(GlobalRotationValue));
+	}
+
+	return BuildState;
 }
 
 bool UDynamicMaterialModelEditorOnlyData::NeedsWizard() const
@@ -787,6 +855,25 @@
 	return CreateBuildState(InMaterialToBuild);
 }
 
+UDMMaterialComponent* UDynamicMaterialModelEditorOnlyData::GetSubComponentByPath(FDMComponentPath& InPath) const
+{
+	if (InPath.IsLeaf())
+	{
+		// This is not a component
+		return nullptr;
+	}
+
+	// Fetches the first component of the path and removes it from the path
+	const FDMComponentPathSegment FirstComponent = InPath.GetFirstSegment();
+
+	if (UDMMaterialComponent* SubComponent = GetSubComponentByPath(InPath, FirstComponent))
+	{
+		return SubComponent->GetComponentByPath(InPath);
+	}
+
+	return nullptr;
+}
+
 void UDynamicMaterialModelEditorOnlyData::DoBuild_Implementation(bool bInDirtyAssets)
 {
 	BuildMaterial(bInDirtyAssets);
@@ -1371,6 +1458,8 @@
 
 	SetFlags(RF_Transactional);
 
+	AssignPropertyAlphaValues();
+
 	ReinitComponents();
 }
 
diff --git a/Engine/Plugins/Experimental/DynamicMaterial/Source/DynamicMaterialEditor/Private/Slate/SDMComponentEdit.cpp b/Engine/Plugins/Experimental/DynamicMaterial/Source/DynamicMaterialEditor/Private/Slate/SDMComponentEdit.cpp
--- a/Engine/Plugins/Experimental/DynamicMaterial/Source/DynamicMaterialEditor/Private/Slate/SDMComponentEdit.cpp
+++ b/Engine/Plugins/Experimental/DynamicMaterial/Source/DynamicMaterialEditor/Private/Slate/SDMComponentEdit.cpp
@@ -783,35 +783,78 @@ void SDMComponentEdit::OnUndo()
 void SDMComponentEdit::GenerateMaterialModelPropertyRows(const TSharedRef<SDMEditor> InEditorWidget, UDynamicMaterialModel* InMaterialModel,
 	TArray<FDMPropertyHandle>& InOutPropertyRows, TSet<UDMMaterialComponent*>& InOutProcessedObjects)
 {
-	if (UDynamicMaterialModelEditorOnlyData* EditorOnlyData = UDynamicMaterialModelEditorOnlyData::Get(InMaterialModel))
-	{
-		if (EditorOnlyData->GetBlendMode() != BLEND_Opaque)
+	const FName UIMin = FName("UIMin");
+	const FName UIMax = FName("UIMax");
+	const FName ClampMin = FName("ClampMin");
+	const FName ClampMax = FName("ClampMax");
+
+	auto AddGlobalVar = [InEditorWidget, InMaterialModel, &InOutPropertyRows, &InOutProcessedObjects, &UIMin, &UIMax, &ClampMin, &ClampMax]
+		(UDMMaterialValue* InValue, const FText& InNameOverride)
 		{
-			if (UDMMaterialValueFloat1* GlobalOpacityValue = InMaterialModel->GetGlobalOpacityValue())
+			if (IsValid(InValue))
 			{
-				if (!InOutProcessedObjects.Contains(GlobalOpacityValue))
+				if (!InOutProcessedObjects.Contains(InValue))
 				{
-					FDMPropertyHandle& GlobalOpacityHandle = InOutPropertyRows.Add_GetRef(InEditorWidget->GetPropertyHandle(&*InEditorWidget,
-						GlobalOpacityValue, UDMMaterialValue::ValueName));
+					FDMPropertyHandle& ValueHandle = InOutPropertyRows.Add_GetRef(InEditorWidget->GetPropertyHandle(&*InEditorWidget,
+						InValue, UDMMaterialValue::ValueName));
 
-					GlobalOpacityHandle.NameOverride = LOCTEXT("GlobalOpacity", "Global Opacity");
+					ValueHandle.NameOverride = InNameOverride;
 
-					InOutProcessedObjects.Add(GlobalOpacityValue);
+					if (UDMMaterialValueFloat* FloatValue = Cast<UDMMaterialValueFloat>(InValue))
+					{
+						if (FloatValue->HasValueRange() && ValueHandle.PropertyHandle)
+						{
+							ValueHandle.PropertyHandle->SetInstanceMetaData(UIMin, FString::SanitizeFloat(FloatValue->GetValueRange().Min));
+							ValueHandle.PropertyHandle->SetInstanceMetaData(ClampMin, FString::SanitizeFloat(FloatValue->GetValueRange().Min));
+							ValueHandle.PropertyHandle->SetInstanceMetaData(UIMax, FString::SanitizeFloat(FloatValue->GetValueRange().Max));
+							ValueHandle.PropertyHandle->SetInstanceMetaData(ClampMax, FString::SanitizeFloat(FloatValue->GetValueRange().Max));
+						}
+					}
+
+					InOutProcessedObjects.Add(InValue);
 				}
 			}
-		}
-	}
+		};
+
+	AddGlobalVar(InMaterialModel->GetGlobalParameterValue(UDynamicMaterialModel::GlobalOffsetValueName), LOCTEXT("GlobalOffset", "Global Offset"));
+	AddGlobalVar(InMaterialModel->GetGlobalParameterValue(UDynamicMaterialModel::GlobalTilingValueName), LOCTEXT("GlobalTiling", "Global Tiling"));
+	AddGlobalVar(InMaterialModel->GetGlobalParameterValue(UDynamicMaterialModel::GlobalRotationValueName), LOCTEXT("GlobalRotation", "Global Rotation"));
 
-	if (UDMMaterialValueFloat2* GlobalTilingValue = InMaterialModel->GetGlobalTilingValue())
+	if (UDynamicMaterialModelEditorOnlyData* EditorOnlyData = UDynamicMaterialModelEditorOnlyData::Get(InMaterialModel))
 	{
-		if (!InOutProcessedObjects.Contains(GlobalTilingValue))
+		if (EditorOnlyData->GetBlendMode() != BLEND_Opaque)
 		{
-			FDMPropertyHandle& GlobalTilingHandle = InOutPropertyRows.Add_GetRef(InEditorWidget->GetPropertyHandle(&*InEditorWidget, 
-				GlobalTilingValue, UDMMaterialValue::ValueName));
+			AddGlobalVar(InMaterialModel->GetGlobalParameterValue(UDynamicMaterialModel::GlobalOpacityValueName), LOCTEXT("GlobalOpacity", "Global Opacity"));
+		}
 
-			GlobalTilingHandle.NameOverride = LOCTEXT("GlobalTiling", "Global Tiling");
+		const FText PropertyFormat = LOCTEXT("PropertyFormat", "Global {0}");
+		UEnum* MaterialPropertyEnum = StaticEnum<EDMMaterialPropertyType>();
+
+		for (uint8 PropertyIndex = static_cast<uint8>(EDMMaterialPropertyType::OpacityMask) + 1;
+			PropertyIndex < static_cast<uint8>(EDMMaterialPropertyType::Any);
+			++PropertyIndex)
+		{
+			const EDMMaterialPropertyType Property = static_cast<EDMMaterialPropertyType>(PropertyIndex);
 
-			InOutProcessedObjects.Add(GlobalTilingValue);
+			if (!!EditorOnlyData->GetSlotForMaterialProperty(Property))
+			{
+				if (UDMMaterialProperty* MaterialProperty = EditorOnlyData->GetMaterialProperty(Property))
+				{
+					if (MaterialProperty->IsValidForModel(*EditorOnlyData))
+					{
+						if (UDMMaterialValueFloat1* AlphaValue = Cast<UDMMaterialValueFloat1>(MaterialProperty->GetComponent(UDynamicMaterialModelEditorOnlyData::AlphaValueName)))
+						{
+							AddGlobalVar(
+								AlphaValue,
+								FText::Format(
+									PropertyFormat,
+									MaterialPropertyEnum->GetDisplayNameTextByValue(static_cast<int64>(Property))
+								)
+							);
+						}
+					}
+				}
+			}
 		}
 	}
 }
diff --git a/Engine/Plugins/Experimental/DynamicMaterial/Source/DynamicMaterialEditor/Private/Slate/SDMEditor.cpp b/Engine/Plugins/Experimental/DynamicMaterial/Source/DynamicMaterialEditor/Private/Slate/SDMEditor.cpp
--- a/Engine/Plugins/Experimental/DynamicMaterial/Source/DynamicMaterialEditor/Private/Slate/SDMEditor.cpp
+++ b/Engine/Plugins/Experimental/DynamicMaterial/Source/DynamicMaterialEditor/Private/Slate/SDMEditor.cpp
@@ -300,6 +300,7 @@ void SDMEditor::SetMaterialModel(UDynamicMaterialModel* InMaterialModel)
 	SplitterContainer.Reset();
 
 	Toolbar->SetMaterialModel(InMaterialModel);
+	SetEditedComponent(nullptr);
 
 	if (!InMaterialModel)
 	{
diff --git a/Engine/Plugins/Experimental/DynamicMaterial/Source/DynamicMaterialEditor/Public/Components/DMMaterialProperty.h b/Engine/Plugins/Experimental/DynamicMaterial/Source/DynamicMaterialEditor/Public/Components/DMMaterialProperty.h
--- a/Engine/Plugins/Experimental/DynamicMaterial/Source/DynamicMaterialEditor/Public/Components/DMMaterialProperty.h
+++ b/Engine/Plugins/Experimental/DynamicMaterial/Source/DynamicMaterialEditor/Public/Components/DMMaterialProperty.h
@@ -26,6 +26,15 @@ class DYNAMICMATERIALEDITOR_API UDMMaterialProperty : public UDMMaterialComponen
 	GENERATED_BODY()
 
 public:
+	static const FString ComponentsPathToken;
+
+	/**
+	 * Craetes a material property as a default subobject of the EditorOnlyData object.
+	 * @param InModelEditorOnlyData The parent EditorOnlyData object.
+	 * @param InMaterialProperty The property type of the material property.
+	 * @param InSubObjName The subobject name.
+	 * @return The newly created default subobject.
+	 */
 	static UDMMaterialProperty* CreateCustomMaterialPropertyDefaultSubobject(UDynamicMaterialModelEditorOnlyData* InModelEditorOnlyData, 
 		EDMMaterialPropertyType InMaterialProperty, const FName& InSubObjName);
 
@@ -35,34 +44,103 @@ public:
 	UFUNCTION(BlueprintPure, Category = "Material Designer")
 	EDMMaterialPropertyType GetMaterialProperty() const { return MaterialProperty; }
 
+	/**
+	 * @return The description of this property based on the EDMMaterialPropertyType enum.
+	 */
 	UFUNCTION(BlueprintPure, Category = "Material Designer")
 	FText GetDescription() const;
 
+	/**
+	 * @return The value type of this property. Will be either VT_Float1, VT_Float3_RGB or VT_Float3_XYZ.
+	 */
 	UFUNCTION(BlueprintPure, Category = "Material Designer")
 	EDMValueType GetInputConnectorType() const { return InputConnectorType; }
 
+	/**
+	 * @return Get the expressions which connect to this property, mapped per channel.
+	 */
 	UFUNCTION(BlueprintPure, Category = "Material Designer")
 	const FDMMaterialStageConnection& GetInputConnectionMap() const { return InputConnectionMap; }
 
+	/**
+	 * @return Get the expressions which connect to this property, mapped per channel.
+	 */
 	FDMMaterialStageConnection& GetInputConnectionMap() { return InputConnectionMap; }
 
+	/**
+	 * The output process is an optional material function which is applied in between the property and its inputs.
+	 * @return The current output processor.
+	 */
 	UFUNCTION(BlueprintPure, Category = "Material Designer")
 	UMaterialFunctionInterface* GetOutputProcessor() const { return OutputProcessor; }
 
+	/**
+	 * The output process is an optional material function which is applied in between the property and its inputs.
+	 * @param InFunction The new output processor.
+	 */
 	UFUNCTION(BlueprintCallable, Category = "Material Designer")
 	void SetOutputProcessor(UMaterialFunctionInterface* InFunction);
 
+	/**
+	 * Check whether the current model settings are valid for this property.
+	 * @param InMaterialModel The model with the settings.
+	 * @return True if it is valid.
+	 */
 	virtual bool IsValidForModel(UDynamicMaterialModelEditorOnlyData& InMaterialModel) const { return true; }
 
+	/**
+	 * Empty the connection map and recalculate the input types.
+	 */
 	UFUNCTION(BlueprintCallable, Category = "Material Designer")
 	virtual void ResetInputConnectionMap();
 
+	/**
+	 * Generates a default input node for this property if the layer stack is empty.
+	 * @param InBuildState The current build state.
+	 * @return The last in the chain of any expressions created (for direct connection to the property).
+	 */
 	virtual UMaterialExpression* GetDefaultInput(const TSharedRef<FDMMaterialBuildState>& InBuildState) const;
 
+	/**
+	 * @return The texture sample type to use for any texture samplers connected to this property.
+	 */
 	virtual TEnumAsByte<EMaterialSamplerType> GetTextureSamplerType() const;
 
+	/**
+	 * Triggered when a slot is added to this property. Should create useful, default layer stack.
+	 * @param InSlot The slot that was created.
+	 */
 	virtual void OnSlotAdded(UDMMaterialSlot* InSlot);
 
+	/**
+	 * Adds a component to the component list. Sets component state to added.
+	 * @param InName The name of the component in the map.
+	 * @param InComponent The component to add.
+	 * @return The replaced component from the map.
+	 */
+	virtual UDMMaterialComponent* AddComponent(FName InName, UDMMaterialComponent* InComponent);
+
+	/**
+	 * Checks to see if a component with this name exists.
+	 * @param InName The name of the component.
+	 * @return True if it exists, even if it isn't valid.
+	 */
+	virtual bool HasComponent(FName InName) const;
+
+	/**
+	 * Fetches a component from the component list.
+	 * @param InName The name of the component.
+	 * @return The component in the map with that name.
+	 */
+	virtual UDMMaterialComponent* GetComponent(FName InName) const;
+
+	/**
+	 * Removes the component from the component list. Sets component state to removed.
+	 * @param InName The name of the component.
+	 * @return The removed component.
+	 */
+	virtual UDMMaterialComponent* RemoveComponent(FName InName);
+
 	//~ Begin UDMMaterialComponent
 	virtual void Update(EDMUpdateType InUpdateType) override;
 	virtual bool IsRootComponent() const override { return true; }
@@ -80,23 +158,51 @@ protected:
 	static UMaterialExpression* CreateConstant(const TSharedRef<FDMMaterialBuildState>& InBuildState, const FVector3d& InDefaultValue);
 	static UMaterialExpression* CreateConstant(const TSharedRef<FDMMaterialBuildState>& InBuildState, const FVector4d& InDefaultValue);
 
+	/**
+	 * The property type of this property.
+	 */
 	UPROPERTY(VisibleInstanceOnly, BlueprintReadOnly, Category = "Material Designer")
 	EDMMaterialPropertyType MaterialProperty;
 
+	/**
+	 * The value type used to connect to this property. Will be either VT_Float1, VT_Float3_RGB or VT_Float3_XYZ.
+	 */
 	UPROPERTY(VisibleInstanceOnly, BlueprintReadOnly, Category = "Material Designer")
 	EDMValueType InputConnectorType;
 
+	/**
+	 * The map of expressions connected to this property's input node.
+	 */
 	UPROPERTY(VisibleInstanceOnly, BlueprintReadOnly, Category = "Material Designer")
 	FDMMaterialStageConnection InputConnectionMap;
 
+	/**
+	 * An optional material function which is applied in between the property and its inputs.
+	 */
 	UPROPERTY(VisibleInstanceOnly, BlueprintReadOnly, Category = "Material Designer")
 	TObjectPtr<UMaterialFunctionInterface> OutputProcessor;
 
 	UPROPERTY(Transient, DuplicateTransient, TextExportTransient)
 	TObjectPtr<UMaterialFunctionInterface> OutputProcessor_PreUpdate;
 
+	/**
+	 * Components of this property. Not necessarily owned or controlled by this property.
+	 */
+	UPROPERTY(VisibleInstanceOnly, BlueprintReadOnly, Category = "Material Designer")
+	TMap<FName, TObjectPtr<UDMMaterialComponent>> Components;
+
 	UDMMaterialProperty();
 	UDMMaterialProperty(EDMMaterialPropertyType InMaterialProperty, EDMValueType InInputConnectorType);
 
+	/**
+	 * Triggered when the output process is updated either through PostEditChangeProperty or via the setter.
+	 * Will potentially reset the value to null if it is invalid.
+	 */
 	void OnOutputProcessorUpdated();
+
+	//~ Begin UDMMaterialComponent
+	virtual UDMMaterialComponent* GetSubComponentByPath(FDMComponentPath& InPath, const FDMComponentPathSegment& InPathSegment) const override;
+	virtual void OnComponentAdded() override;
+	virtual void OnComponentRemoved() override;
+	//~ End UDMMaterialComponent
 };
diff --git a/Engine/Plugins/Experimental/DynamicMaterial/Source/DynamicMaterialEditor/Public/Model/DMMaterialBuildState.h b/Engine/Plugins/Experimental/DynamicMaterial/Source/DynamicMaterialEditor/Public/Model/DMMaterialBuildState.h
--- a/Engine/Plugins/Experimental/DynamicMaterial/Source/DynamicMaterialEditor/Public/Model/DMMaterialBuildState.h
+++ b/Engine/Plugins/Experimental/DynamicMaterial/Source/DynamicMaterialEditor/Public/Model/DMMaterialBuildState.h
@@ -54,10 +54,6 @@ struct DYNAMICMATERIALEDITOR_API FDMMaterialBuildState : public IDMMaterialBuild
 
 	bool IsPreviewMaterial() const { return bIsPreviewMaterial; }
 
-	UMaterialExpression* GetGlobalTilingExpression() const { return GlobalTilingExpression; }
-
-	void SetGlobalTilingExpression(UMaterialExpression* InExpression) { GlobalTilingExpression = InExpression; }
-
 	virtual IDMMaterialBuildUtilsInterface& GetBuildUtils() const override;
 
 	FExpressionInput* GetMaterialProperty(EDMMaterialPropertyType InProperty) const;
@@ -153,6 +149,11 @@ struct DYNAMICMATERIALEDITOR_API FDMMaterialBuildState : public IDMMaterialBuild
 
 	virtual const TSet<UMaterialExpression*>& GetOtherExpressions() override;
 
+	/** Global Expressions */
+	UMaterialExpression* GetGlobalExpression(FName InName) const;
+
+	void SetGlobalExpression(FName InName, UMaterialExpression* InExpression);
+
 private:
 	UMaterial* DynamicMaterial = nullptr;
 	UDynamicMaterialModel* MaterialModel = nullptr;
@@ -170,5 +171,5 @@ private:
 	TMap<const UDMMaterialStageSource*, TArray<UMaterialExpression*>> StageSources;
 	TMap<const UDMMaterialStageSource*, TArray<UDMMaterialStageSource*>> Callbacks;
 	TSet<UMaterialExpression*> OtherExpressions;
-	UMaterialExpression* GlobalTilingExpression;
+	TMap<FName, UMaterialExpression*> GlobalExpressions;
 };
diff --git a/Engine/Plugins/Experimental/DynamicMaterial/Source/DynamicMaterialEditor/Public/Model/DynamicMaterialModelEditorOnlyData.h b/Engine/Plugins/Experimental/DynamicMaterial/Source/DynamicMaterialEditor/Public/Model/DynamicMaterialModelEditorOnlyData.h
--- a/Engine/Plugins/Experimental/DynamicMaterial/Source/DynamicMaterialEditor/Public/Model/DynamicMaterialModelEditorOnlyData.h
+++ b/Engine/Plugins/Experimental/DynamicMaterial/Source/DynamicMaterialEditor/Public/Model/DynamicMaterialModelEditorOnlyData.h
@@ -69,6 +69,8 @@ public:
 	static const TArray<EMaterialDomain> SupportedDomains;
 	static const TArray<EBlendMode> SupportedBlendModes;
 
+	static const FName AlphaValueName;
+
 	static UDynamicMaterialModelEditorOnlyData* Get(UDynamicMaterialModel* InModel);
 	static UDynamicMaterialModelEditorOnlyData* Get(TWeakObjectPtr<UDynamicMaterialModel> InModelWeak);
 	static UDynamicMaterialModelEditorOnlyData* Get(const TScriptInterface<IDynamicMaterialModelEditorOnlyDataInterface>& InInterface);
@@ -206,6 +208,7 @@ public:
 	virtual void OnTextureUVUpdated(UDMTextureUV* InTextureUV) override;
 	virtual void LoadDeprecatedModelData(UDynamicMaterialModel* InMaterialModel) override;
 	virtual TSharedRef<IDMMaterialBuildStateInterface> CreateBuildStateInterface(UMaterial* InMaterialToBuild) const override;
+	virtual UDMMaterialComponent* GetSubComponentByPath(FDMComponentPath& InPath) const override;
 	virtual UDMMaterialComponent* GetSubComponentByPath(FDMComponentPath& InPath, const FDMComponentPathSegment& InPathSegment) const override;
 	//~ End IDynamicMaterialModelEditorOnlyDataInterface
 
@@ -280,6 +283,8 @@ protected:
 	virtual void ResetData() override;
 	//~ End IDynamicMaterialModelEditorOnlyDataInterface
 
+	void AssignPropertyAlphaValues();
+
 	void LoadDeprecatedModelData_Base(bool bInCreateMaterialPackage, EBlendMode InBlendMode, EDMMaterialShadingModel InShadingModel);
 	void LoadDeprecatedModelData_Expressions(TArray<TObjectPtr<UMaterialExpression>>& InExpressions);
 	void LoadDeprecatedModelData_Properties(TMap<EDMMaterialPropertyType, TObjectPtr<UObject>>& InProperties);
