Material Designer
- Integrated the new custom detail view category items.

diff --git a/Engine/Plugins/Experimental/DynamicMaterial/Source/DynamicMaterial/Private/Components/DMMaterialValue.cpp b/Engine/Plugins/Experimental/DynamicMaterial/Source/DynamicMaterial/Private/Components/DMMaterialValue.cpp
--- a/Engine/Plugins/Experimental/DynamicMaterial/Source/DynamicMaterial/Private/Components/DMMaterialValue.cpp
+++ b/Engine/Plugins/Experimental/DynamicMaterial/Source/DynamicMaterial/Private/Components/DMMaterialValue.cpp
@@ -6,6 +6,7 @@
 #include "Model/DynamicMaterialModel.h"
 #include "Materials/Material.h"
 #include "UObject/Package.h"
+#include "Utils/DMDetailsViewUtils.h"
 
 #if WITH_EDITOR
 #include "AssetRegistry/AssetRegistryModule.h"
@@ -439,34 +440,13 @@ void UDMMaterialValue::EnsureDetailObjects()
 	FPropertyRowGeneratorArgs RowGeneratorArgs;
 	PropertyRowGenerator = PropertyEditor.CreatePropertyRowGenerator(RowGeneratorArgs);
 	PropertyRowGenerator->SetObjects({this});
- 
-	for (const TSharedRef<IDetailTreeNode>& CategoryNode : PropertyRowGenerator->GetRootTreeNodes())
+
+	DetailTreeNode = FDMDetailsViewUtils::SearchNodesForProperty(PropertyRowGenerator->GetRootTreeNodes(), ValueName);
+
+	if (DetailTreeNode.IsValid())
 	{
-		if (CategoryNode->GetNodeName() != TEXT("Material Designer"))
-		{
-			continue;
-		}
- 
-		TArray<TSharedRef<IDetailTreeNode>> ChildNodes;
-		CategoryNode->GetChildren(ChildNodes);
- 
-		for (const TSharedRef<IDetailTreeNode>& ChildNode : ChildNodes)
-		{
-			if (ChildNode->GetNodeType() != EDetailNodeType::Item)
-			{
-				continue;
-			}
- 
-			if (ChildNode->GetNodeName() != ValueName)
-			{
-				continue;
-			}
- 
-			DetailTreeNode = ChildNode;
-			PropertyHandle = ChildNode->CreatePropertyHandle();
-			return;
-		}
-	}
+		PropertyHandle = DetailTreeNode->CreatePropertyHandle();
+	}	 
 }
 
 TSharedPtr<IDetailTreeNode> UDMMaterialValue::GetDetailTreeNode()
diff --git a/Engine/Plugins/Experimental/DynamicMaterial/Source/DynamicMaterial/Private/Components/DMTextureUV.cpp b/Engine/Plugins/Experimental/DynamicMaterial/Source/DynamicMaterial/Private/Components/DMTextureUV.cpp
--- a/Engine/Plugins/Experimental/DynamicMaterial/Source/DynamicMaterial/Private/Components/DMTextureUV.cpp
+++ b/Engine/Plugins/Experimental/DynamicMaterial/Source/DynamicMaterial/Private/Components/DMTextureUV.cpp
@@ -6,6 +6,7 @@
 #include "Materials/MaterialInstanceDynamic.h"
 #include "Model/DynamicMaterialModel.h"
 #include "Serialization/CustomVersion.h"
+#include "Utils/DMDetailsViewUtils.h"
 
 #if WITH_EDITOR
 #include "IDetailTreeNode.h"
@@ -508,33 +509,15 @@ void UDMTextureUV::EnsureDetailObjects()
 	PropertyRowGenerator = PropertyEditor.CreatePropertyRowGenerator(RowGeneratorArgs);
 	PropertyRowGenerator->SetObjects({this});
 
-	for (const TSharedRef<IDetailTreeNode>& CategoryNode : PropertyRowGenerator->GetRootTreeNodes())
-	{
-		if (CategoryNode->GetNodeName() != TEXT("Material Designer"))
-		{
-			continue;
-		}
+	const TArray<TSharedRef<IDetailTreeNode>>& CategoryNodes = PropertyRowGenerator->GetRootTreeNodes();
 
-		TArray<TSharedRef<IDetailTreeNode>> ChildNodes;
-		CategoryNode->GetChildren(ChildNodes);
-
-		for (const TSharedRef<IDetailTreeNode>& ChildNode : ChildNodes)
+	for (const TPair<FName, bool>& TextureProperty : TextureProperties)
+	{
+		if (TSharedPtr<IDetailTreeNode> DetailTreeNode = FDMDetailsViewUtils::SearchNodesForProperty(CategoryNodes, TextureProperty.Key))
 		{
-			if (ChildNode->GetNodeType() != EDetailNodeType::Item)
-			{
-				continue;
-			}
-			
-			if (TextureProperties.Contains(ChildNode->GetNodeName()) == false)
-			{
-				continue;
-			}
-
-			DetailTreeNodes.Emplace(ChildNode->GetNodeName(), ChildNode);
-			PropertyHandles.Emplace(ChildNode->GetNodeName(), ChildNode->CreatePropertyHandle());
+			DetailTreeNodes.Emplace(DetailTreeNode->GetNodeName(), DetailTreeNode);
+			PropertyHandles.Emplace(DetailTreeNode->GetNodeName(), DetailTreeNode->CreatePropertyHandle());
 		}
-
-		return;
 	}
 }
 #endif
diff --git a/Engine/Plugins/Experimental/DynamicMaterial/Source/DynamicMaterial/Private/Utils/DMDetailsViewUtils.cpp b/Engine/Plugins/Experimental/DynamicMaterial/Source/DynamicMaterial/Private/Utils/DMDetailsViewUtils.cpp
new file mode 100644
--- /dev/null
+++ b/Engine/Plugins/Experimental/DynamicMaterial/Source/DynamicMaterial/Private/Utils/DMDetailsViewUtils.cpp
@@ -0,0 +1,41 @@
+// Copyright Epic Games, Inc. All Rights Reserved.
+
+#if WITH_EDITOR
+#include "Utils/DMDetailsViewUtils.h"
+#include "IDetailTreeNode.h"
+
+TSharedPtr<IDetailTreeNode> FDMDetailsViewUtils::SearchNodesForProperty(const TArray<TSharedRef<IDetailTreeNode>>& InNodes, FName InPropertyName)
+{
+	for (const TSharedRef<IDetailTreeNode>& ChildNode : InNodes)
+	{
+		switch (ChildNode->GetNodeType())
+		{
+			case EDetailNodeType::Category:
+			{
+				TArray<TSharedRef<IDetailTreeNode>> CategoryChildNodes;
+				ChildNode->GetChildren(CategoryChildNodes);
+
+				if (TSharedPtr<IDetailTreeNode> FoundNode = SearchNodesForProperty(CategoryChildNodes, InPropertyName))
+				{
+					return FoundNode;
+				}
+
+				break;
+			}
+
+			case EDetailNodeType::Item:
+				if (ChildNode->GetNodeName() == InPropertyName)
+				{
+					return ChildNode;
+				}
+				break;
+
+			default:
+			// Do nothing
+			break;
+		}
+	}
+
+	return nullptr;
+}
+#endif
diff --git a/Engine/Plugins/Experimental/DynamicMaterial/Source/DynamicMaterial/Public/Components/DMTextureUV.h b/Engine/Plugins/Experimental/DynamicMaterial/Source/DynamicMaterial/Public/Components/DMTextureUV.h
--- a/Engine/Plugins/Experimental/DynamicMaterial/Source/DynamicMaterial/Public/Components/DMTextureUV.h
+++ b/Engine/Plugins/Experimental/DynamicMaterial/Source/DynamicMaterial/Public/Components/DMTextureUV.h
@@ -180,35 +180,35 @@ public:
 
 protected:
 #if WITH_EDITORONLY_DATA
-	UPROPERTY(EditInstanceOnly, BlueprintReadWrite, BlueprintSetter = SetUVSource, Category = "Material Designer",
+	UPROPERTY(EditInstanceOnly, BlueprintReadWrite, BlueprintSetter = SetUVSource, Category = "Material Designer|Texture UV",
 		meta = (AllowPrivateAccess = "true", ToolTip = "Forces a material rebuild.", NotKeyframeable))
 	EDMUVSource UVSource = EDMUVSource::Texture;
 
 	EDMUVSource UVSource_PreUndo = UVSource;
 #endif
 
-	UPROPERTY(EditInstanceOnly, BlueprintReadWrite, Getter = GetOffset, Setter = SetOffset, BlueprintSetter = SetOffset, Category = "Material Designer",
+	UPROPERTY(EditInstanceOnly, BlueprintReadWrite, Getter = GetOffset, Setter = SetOffset, BlueprintSetter = SetOffset, Category = "Material Designer|Texture UV",
 		meta = (AllowPrivateAccess = "true", Delta = 0.001))
 	FVector2D Offset = FVector2D(0.f, 0.f);
 
-	UPROPERTY(EditInstanceOnly, BlueprintReadWrite, Getter = GetPivot, Setter = SetPivot, BlueprintSetter = SetPivot, Category = "Material Designer",
+	UPROPERTY(EditInstanceOnly, BlueprintReadWrite, Getter = GetPivot, Setter = SetPivot, BlueprintSetter = SetPivot, Category = "Material Designer|Texture UV",
 		meta = (AllowPrivateAccess = "true", ToolTip="Pivot for rotation and scale.", Delta = 0.001))
 	FVector2D Pivot = FVector2D(0.5, 0.5);
 
-	UPROPERTY(EditInstanceOnly, BlueprintReadWrite, Getter = GetRotation, Setter = SetRotation, BlueprintSetter = SetRotation, Category = "Material Designer",
+	UPROPERTY(EditInstanceOnly, BlueprintReadWrite, Getter = GetRotation, Setter = SetRotation, BlueprintSetter = SetRotation, Category = "Material Designer|Texture UV",
 		meta = (AllowPrivateAccess = "true", Delta = 1.0))
 	float Rotation = 0.f;
 
-	UPROPERTY(EditInstanceOnly, BlueprintReadWrite, Getter = GetScale, Setter = SetScale, BlueprintSetter = SetScale, Category = "Material Designer",
+	UPROPERTY(EditInstanceOnly, BlueprintReadWrite, Getter = GetScale, Setter = SetScale, BlueprintSetter = SetScale, Category = "Material Designer|Texture UV",
 		meta = (AllowPrivateAccess = "true", AllowPreserveRatio, Delta = 0.001))
 	FVector2D Scale = FVector2D(1.f, 1.f);
 
 #if WITH_EDITORONLY_DATA
-	UPROPERTY(EditInstanceOnly, BlueprintReadWrite, BlueprintSetter = SetMirrorOnX, Category = "Material Designer",
+	UPROPERTY(EditInstanceOnly, BlueprintReadWrite, BlueprintSetter = SetMirrorOnX, Category = "Material Designer|Texture UV",
 		meta = (AllowPrivateAccess = "true", ToolTip = "Forces a material rebuild.", NotKeyframeable))
 	bool bMirrorOnX = false;
 
-	UPROPERTY(EditInstanceOnly, BlueprintReadWrite, BlueprintSetter = SetMirrorOnY, Category = "Material Designer",
+	UPROPERTY(EditInstanceOnly, BlueprintReadWrite, BlueprintSetter = SetMirrorOnY, Category = "Material Designer|Texture UV",
 		meta = (AllowPrivateAccess = "true", ToolTip = "Forces a material rebuild.", NotKeyframeable))
 	bool bMirrorOnY = false;
 
diff --git a/Engine/Plugins/Experimental/DynamicMaterial/Source/DynamicMaterial/Public/Components/MaterialValues/DMMaterialValueRenderTarget.h b/Engine/Plugins/Experimental/DynamicMaterial/Source/DynamicMaterial/Public/Components/MaterialValues/DMMaterialValueRenderTarget.h
--- a/Engine/Plugins/Experimental/DynamicMaterial/Source/DynamicMaterial/Public/Components/MaterialValues/DMMaterialValueRenderTarget.h
+++ b/Engine/Plugins/Experimental/DynamicMaterial/Source/DynamicMaterial/Public/Components/MaterialValues/DMMaterialValueRenderTarget.h
@@ -88,15 +88,15 @@ public:
 protected:
 	FDelegateHandle EndOfFrameDelegateHandle;
 
-	UPROPERTY(EditAnywhere, BlueprintReadWrite, Setter = SetTextureSize, BlueprintSetter = SetTextureSize, Category = "Material Designer",
+	UPROPERTY(EditAnywhere, BlueprintReadWrite, Setter = SetTextureSize, BlueprintSetter = SetTextureSize, Category = "Material Designer|Render Target",
 		meta = (AllowPrivateAccess = "true"))
 	FIntPoint TextureSize;
 
-	UPROPERTY(EditAnywhere, BlueprintReadWrite, Setter = SetTextureFormat, BlueprintSetter = SetTextureFormat, Category = "Material Designer",
+	UPROPERTY(EditAnywhere, BlueprintReadWrite, Setter = SetTextureFormat, BlueprintSetter = SetTextureFormat, Category = "Material Designer|Render Target",
 		meta = (AllowPrivateAccess = "true"))
 	TEnumAsByte<ETextureRenderTargetFormat> TextureFormat;
 
-	UPROPERTY(EditAnywhere, BlueprintReadWrite, Setter = SetClearColor, BlueprintSetter = SetClearColor, Category = "Material Designer",
+	UPROPERTY(EditAnywhere, BlueprintReadWrite, Setter = SetClearColor, BlueprintSetter = SetClearColor, Category = "Material Designer|Render Target",
 		meta = (AllowPrivateAccess = "true"))
 	FLinearColor ClearColor;
 
diff --git a/Engine/Plugins/Experimental/DynamicMaterial/Source/DynamicMaterial/Public/Components/RenderTargetRenderers/DMRenderTargetTextRenderer.h b/Engine/Plugins/Experimental/DynamicMaterial/Source/DynamicMaterial/Public/Components/RenderTargetRenderers/DMRenderTargetTextRenderer.h
--- a/Engine/Plugins/Experimental/DynamicMaterial/Source/DynamicMaterial/Public/Components/RenderTargetRenderers/DMRenderTargetTextRenderer.h
+++ b/Engine/Plugins/Experimental/DynamicMaterial/Source/DynamicMaterial/Public/Components/RenderTargetRenderers/DMRenderTargetTextRenderer.h
@@ -180,96 +180,96 @@ public:
 #endif
 
 protected:
-	UPROPERTY(EditInstanceOnly, BlueprintReadWrite, Setter = SetFont, BlueprintSetter = SetFont, Category = "Material Designer",
+	UPROPERTY(EditInstanceOnly, BlueprintReadWrite, Setter = SetFont, BlueprintSetter = SetFont, Category = "Material Designer|Text",
 		meta = (NotKeyframeable, AllowPrivateAccess = "true"))
 	UFont* Font = nullptr;
 
-	UPROPERTY(EditInstanceOnly, BlueprintReadWrite, Setter = SetText, BlueprintSetter = SetText, Category = "Material Designer",
+	UPROPERTY(EditInstanceOnly, BlueprintReadWrite, Setter = SetText, BlueprintSetter = SetText, Category = "Material Designer|Text",
 		meta = (MultiLine, NotKeyframeable, AllowPrivateAccess = "true"))
 	FText Text = FText::GetEmpty();
 
 	UPROPERTY()
 	TArray<FDMTextLine> Lines = {};
 
-	UPROPERTY(EditInstanceOnly, BlueprintReadWrite, Setter = SetTextColor, BlueprintSetter = SetTextColor, Category = "Material Designer",
+	UPROPERTY(EditInstanceOnly, BlueprintReadWrite, Setter = SetTextColor, BlueprintSetter = SetTextColor, Category = "Material Designer|Text",
 		meta = (AllowPrivateAccess = "true"))
 	FLinearColor TextColor = FLinearColor::White;
 
-	UPROPERTY(EditInstanceOnly, BlueprintReadWrite, Setter = SetJustify, BlueprintSetter = SetJustify, Category = "Material Designer",
+	UPROPERTY(EditInstanceOnly, BlueprintReadWrite, Setter = SetJustify, BlueprintSetter = SetJustify, Category = "Material Designer|Text",
 		meta = (NotKeyframeable, AllowPrivateAccess = "true"))
 	TEnumAsByte<ETextJustify::Type> Justify = ETextJustify::Left;
 
-	UPROPERTY(EditInstanceOnly, BlueprintReadWrite, Setter = SetKerning, BlueprintSetter = SetKerning, Category = "Material Designer",
+	UPROPERTY(EditInstanceOnly, BlueprintReadWrite, Setter = SetKerning, BlueprintSetter = SetKerning, Category = "Material Designer|Text",
 		meta = (AllowPrivateAccess = "true"))
 	float Kerning = 0.f;
 
 	/** Multiplier on the base font height. */
-	UPROPERTY(EditInstanceOnly, BlueprintReadWrite, Setter = SetLineHeight, BlueprintSetter = SetLineHeight, Category = "Material Designer",
+	UPROPERTY(EditInstanceOnly, BlueprintReadWrite, Setter = SetLineHeight, BlueprintSetter = SetLineHeight, Category = "Material Designer|Text",
 		meta = (NotKeyframeable, AllowPrivateAccess = "true", UIMin = 0.1, UIMax = 2.0))
 	float LineHeight = 1.f;
 
 	/** Extra space adding beyond the edge of the glyphs. Useful for shadows, glows, etc. */
-	UPROPERTY(EditInstanceOnly, BlueprintReadWrite, Setter = SetPaddingLeft, BlueprintSetter = SetPaddingLeft, Category = "Material Designer",
+	UPROPERTY(EditInstanceOnly, BlueprintReadWrite, Setter = SetPaddingLeft, BlueprintSetter = SetPaddingLeft, Category = "Material Designer|Text",
 		meta = (NotKeyframeable, AllowPrivateAccess = "true"))
 	float PaddingLeft = 0.f;
 
 	/** Extra space adding beyond the edge of the glyphs. Useful for shadows, glows, etc. */
-	UPROPERTY(EditInstanceOnly, BlueprintReadWrite, Setter = SetPaddingRight, BlueprintSetter = SetPaddingRight, Category = "Material Designer",
+	UPROPERTY(EditInstanceOnly, BlueprintReadWrite, Setter = SetPaddingRight, BlueprintSetter = SetPaddingRight, Category = "Material Designer|Text",
 		meta = (NotKeyframeable, AllowPrivateAccess = "true"))
 	float PaddingRight = 0.f;
 
 	/** Extra space adding beyond the edge of the glyphs. Useful for shadows, glows, etc. */
-	UPROPERTY(EditInstanceOnly, BlueprintReadWrite, Setter = SetPaddingTop, BlueprintSetter = SetPaddingTop, Category = "Material Designer",
+	UPROPERTY(EditInstanceOnly, BlueprintReadWrite, Setter = SetPaddingTop, BlueprintSetter = SetPaddingTop, Category = "Material Designer|Text",
 		meta = (NotKeyframeable, AllowPrivateAccess = "true"))
 	float PaddingTop = 0.f;
 
 	/** Extra space adding beyond the edge of the glyphs. Useful for shadows, glows, etc. */
-	UPROPERTY(EditInstanceOnly, BlueprintReadWrite, Setter = SetPaddingBottom, BlueprintSetter = SetPaddingBottom, Category = "Material Designer",
+	UPROPERTY(EditInstanceOnly, BlueprintReadWrite, Setter = SetPaddingBottom, BlueprintSetter = SetPaddingBottom, Category = "Material Designer|Text",
 		meta = (NotKeyframeable, AllowPrivateAccess = "true"))
 	float PaddingBottom = 0.f;
 
-	UPROPERTY(EditInstanceOnly, BlueprintReadWrite, Setter = SetTextScale, BlueprintSetter = SetTextScale, Category = "Material Designer",
+	UPROPERTY(EditInstanceOnly, BlueprintReadWrite, Setter = SetTextScale, BlueprintSetter = SetTextScale, Category = "Material Designer|Text",
 		meta = (NotKeyframeable, AllowPrivateAccess = "true", ClampMin = 1, UIMin = 1, ClampMax = 64, UIMax = 64,
 			ToolTip = "Increasing this will increase the resolution of the text."))
 	FIntPoint TextScale = FIntPoint(10, 10);
 
-	UPROPERTY(EditInstanceOnly, BlueprintReadWrite, Setter = SetTextureSizeOverride, BlueprintSetter = SetTextureSizeOverride, Category = "Material Designer",
+	UPROPERTY(EditInstanceOnly, BlueprintReadWrite, Setter = SetTextureSizeOverride, BlueprintSetter = SetTextureSizeOverride, Category = "Material Designer|Text",
 		meta = (NotKeyframeable, AllowPrivateAccess = "true", ToolTip = "Override the size of the render target. Set to 0 to disable."))
 	FIntPoint TextureSizeOverride = FIntPoint(0, 0);
 
-	UPROPERTY(EditInstanceOnly, BlueprintReadWrite, Setter = SetHasOutline, BlueprintSetter = SetHasOutline, Category = "Material Designer",
+	UPROPERTY(EditInstanceOnly, BlueprintReadWrite, Setter = SetHasOutline, BlueprintSetter = SetHasOutline, Category = "Material Designer|Text",
 		meta = (AllowPrivateAccess = "true"))
 	bool bOutline = false;
 
-	UPROPERTY(EditInstanceOnly, BlueprintReadWrite, Setter = SetOutlineColor, BlueprintSetter = SetOutlineColor, Category = "Material Designer",
+	UPROPERTY(EditInstanceOnly, BlueprintReadWrite, Setter = SetOutlineColor, BlueprintSetter = SetOutlineColor, Category = "Material Designer|Text",
 		meta = (AllowPrivateAccess = "true", EditCondition = "bOutline", EditConditionHides = true))
 	FLinearColor OutlineColor = FLinearColor::Black;
 
-	UPROPERTY(EditInstanceOnly, BlueprintReadWrite, Setter = SetHasShadow, BlueprintSetter = SetHasShadow, Category = "Material Designer",
+	UPROPERTY(EditInstanceOnly, BlueprintReadWrite, Setter = SetHasShadow, BlueprintSetter = SetHasShadow, Category = "Material Designer|Text",
 		meta = (AllowPrivateAccess = "true"))
 	bool bShadow = false;
 
-	UPROPERTY(EditInstanceOnly, BlueprintReadWrite, Setter = SetShadowColor, BlueprintSetter = SetShadowColor, Category = "Material Designer",
+	UPROPERTY(EditInstanceOnly, BlueprintReadWrite, Setter = SetShadowColor, BlueprintSetter = SetShadowColor, Category = "Material Designer|Text",
 		meta = (AllowPrivateAccess = "true", EditCondition = "bShadow", EditConditionHides = true))
 	FLinearColor ShadowColor = FLinearColor::Black;
 
-	UPROPERTY(EditInstanceOnly, BlueprintReadWrite, Setter = SetShadowOffset, BlueprintSetter = SetShadowOffset, Category = "Material Designer",
+	UPROPERTY(EditInstanceOnly, BlueprintReadWrite, Setter = SetShadowOffset, BlueprintSetter = SetShadowOffset, Category = "Material Designer|Text",
 		meta = (AllowPrivateAccess = "true", EditCondition = "bShadow", EditConditionHides = true))
 	FVector2D ShadowOffset = FVector2D(1.0, 1.0);
 
-	UPROPERTY(EditInstanceOnly, BlueprintReadWrite, Setter = SetHasGlow, BlueprintSetter = SetHasGlow, Category = "Material Designer",
+	UPROPERTY(EditInstanceOnly, BlueprintReadWrite, Setter = SetHasGlow, BlueprintSetter = SetHasGlow, Category = "Material Designer|Text",
 		meta = (AllowPrivateAccess = "true"))
 	bool bGlow = false;
 
-	UPROPERTY(EditInstanceOnly, BlueprintReadWrite, Setter = SetGlowColor, BlueprintSetter = SetGlowColor, Category = "Material Designer",
+	UPROPERTY(EditInstanceOnly, BlueprintReadWrite, Setter = SetGlowColor, BlueprintSetter = SetGlowColor, Category = "Material Designer|Text",
 		meta = (AllowPrivateAccess = "true", EditCondition = "bGlow", EditConditionHides = true))
 	FLinearColor GlowColor = FLinearColor::White;
 
-	UPROPERTY(EditInstanceOnly, BlueprintReadWrite, Setter = SetGlowInnerRadius, BlueprintSetter = SetGlowInnerRadius, Category = "Material Designer",
+	UPROPERTY(EditInstanceOnly, BlueprintReadWrite, Setter = SetGlowInnerRadius, BlueprintSetter = SetGlowInnerRadius, Category = "Material Designer|Text",
 		meta = (AllowPrivateAccess = "true", EditCondition = "bGlow", EditConditionHides = true, ClampMin = 0, UIMin = 0, ClampMax = 1, UIMax = 1))
 	FVector2D GlowInnerRadius = FVector2D(0.5, 0.5);
 
-	UPROPERTY(EditInstanceOnly, BlueprintReadWrite, Setter = SetGlowOuterRadius, BlueprintSetter = SetGlowOuterRadius, Category = "Material Designer",
+	UPROPERTY(EditInstanceOnly, BlueprintReadWrite, Setter = SetGlowOuterRadius, BlueprintSetter = SetGlowOuterRadius, Category = "Material Designer|Text",
 		meta = (AllowPrivateAccess = "true", EditCondition = "bGlow", EditConditionHides = true, ClampMin = 0, UIMin = 0, ClampMax = 1, UIMax = 1))
 	FVector2D GlowOuterRadius = FVector2D(0.5, 0.5);
 
diff --git a/Engine/Plugins/Experimental/DynamicMaterial/Source/DynamicMaterial/Public/Utils/DMDetailsViewUtils.h b/Engine/Plugins/Experimental/DynamicMaterial/Source/DynamicMaterial/Public/Utils/DMDetailsViewUtils.h
new file mode 100644
--- /dev/null
+++ b/Engine/Plugins/Experimental/DynamicMaterial/Source/DynamicMaterial/Public/Utils/DMDetailsViewUtils.h
@@ -0,0 +1,16 @@
+// Copyright Epic Games, Inc. All Rights Reserved.
+
+#pragma once
+
+#if WITH_EDITOR
+#include "Containers/Array.h"
+#include "Templates/SharedPointerFwd.h"
+#include "UObject/UnrealNames.h"
+
+class IDetailTreeNode;
+
+struct FDMDetailsViewUtils
+{
+	static DYNAMICMATERIAL_API TSharedPtr<IDetailTreeNode> SearchNodesForProperty(const TArray<TSharedRef<IDetailTreeNode>>& InNodes, FName InPropertyName);
+};
+#endif
diff --git a/Engine/Plugins/Experimental/DynamicMaterial/Source/DynamicMaterialEditor/Private/Slate/Properties/Generators/DMMaterialEffectFunctionPropertyRowGenerator.cpp b/Engine/Plugins/Experimental/DynamicMaterial/Source/DynamicMaterialEditor/Private/Slate/Properties/Generators/DMMaterialEffectFunctionPropertyRowGenerator.cpp
--- a/Engine/Plugins/Experimental/DynamicMaterial/Source/DynamicMaterialEditor/Private/Slate/Properties/Generators/DMMaterialEffectFunctionPropertyRowGenerator.cpp
+++ b/Engine/Plugins/Experimental/DynamicMaterial/Source/DynamicMaterialEditor/Private/Slate/Properties/Generators/DMMaterialEffectFunctionPropertyRowGenerator.cpp
@@ -96,10 +96,12 @@ void FDMMaterialEffectFunctionPropertyRowGenerator::AddComponentProperties(const
 		}
 
 		const FText Description = FText::FromString(Inputs[InputIndex].ExpressionInput->Description);
+		const FName MaterialFunctionCategory = TEXT("Function Inputs");
 
 		for (FDMPropertyHandle& ValuePropertyRow : ValuePropertyRows)
 		{
 			ValuePropertyRow.NameToolTipOverride = Description;
+			ValuePropertyRow.CategoryOverrideName = MaterialFunctionCategory;
 		}
 
 		InOutPropertyRows.Append(ValuePropertyRows);
diff --git a/Engine/Plugins/Experimental/DynamicMaterial/Source/DynamicMaterialEditor/Private/Slate/Properties/Generators/DMMaterialStageFunctionPropertyRowGenerator.cpp b/Engine/Plugins/Experimental/DynamicMaterial/Source/DynamicMaterialEditor/Private/Slate/Properties/Generators/DMMaterialStageFunctionPropertyRowGenerator.cpp
--- a/Engine/Plugins/Experimental/DynamicMaterial/Source/DynamicMaterialEditor/Private/Slate/Properties/Generators/DMMaterialStageFunctionPropertyRowGenerator.cpp
+++ b/Engine/Plugins/Experimental/DynamicMaterial/Source/DynamicMaterialEditor/Private/Slate/Properties/Generators/DMMaterialStageFunctionPropertyRowGenerator.cpp
@@ -102,10 +102,12 @@ void FDMMaterialStageFunctionPropertyRowGenerator::AddComponentProperties(const
 					}
 
 					const FText Description = FText::FromString(Input.ExpressionInput->Description);
+					const FName MaterialFunctionCategory = TEXT("Function Inputs");
 
 					for (FDMPropertyHandle& ValuePropertyRow : ValuePropertyRows)
 					{
 						ValuePropertyRow.NameToolTipOverride = Description;
+						ValuePropertyRow.CategoryOverrideName = MaterialFunctionCategory;
 					}
 
 					AllValuePropertyRows.Append(ValuePropertyRows);
diff --git a/Engine/Plugins/Experimental/DynamicMaterial/Source/DynamicMaterialEditor/Private/Slate/Properties/Generators/DMTextureUVPropertyRowGenerator.cpp b/Engine/Plugins/Experimental/DynamicMaterial/Source/DynamicMaterialEditor/Private/Slate/Properties/Generators/DMTextureUVPropertyRowGenerator.cpp
--- a/Engine/Plugins/Experimental/DynamicMaterial/Source/DynamicMaterialEditor/Private/Slate/Properties/Generators/DMTextureUVPropertyRowGenerator.cpp
+++ b/Engine/Plugins/Experimental/DynamicMaterial/Source/DynamicMaterialEditor/Private/Slate/Properties/Generators/DMTextureUVPropertyRowGenerator.cpp
@@ -126,6 +126,7 @@ void UE::DynamicMaterialEditor::Private::AddVisualizerRow(const TSharedRef<SWidg
 	VisualizerHandle.NameToolTipOverride = LOCTEXT("VisualizerToolTip", "A graphical Texture UV editor.\n\n- Offset Mode: Change the Texture UV offset.\n- Pivot Mode: Change the Texture UV pivot, rotation and scale.\n\nControl+click to reset values to default.");
 	VisualizerHandle.ValueName = FName(*InTextureUV->GetComponentPath());
 	VisualizerHandle.ValueWidget = SNew(SDMTextureUVVisualizerProperty, Stage, InTextureUV);
+	VisualizerHandle.CategoryOverrideName = TEXT("Texture UV");
 	InOutPropertyRows.Add(VisualizerHandle);
 }
 
diff --git a/Engine/Plugins/Experimental/DynamicMaterial/Source/DynamicMaterialEditor/Private/Slate/SDMComponentEdit.cpp b/Engine/Plugins/Experimental/DynamicMaterial/Source/DynamicMaterialEditor/Private/Slate/SDMComponentEdit.cpp
--- a/Engine/Plugins/Experimental/DynamicMaterial/Source/DynamicMaterialEditor/Private/Slate/SDMComponentEdit.cpp
+++ b/Engine/Plugins/Experimental/DynamicMaterial/Source/DynamicMaterialEditor/Private/Slate/SDMComponentEdit.cpp
@@ -30,6 +30,7 @@
 #include "IDetailTreeNode.h"
 #include "ISinglePropertyView.h"
 #include "Items/CustomDetailsViewItemId.h"
+#include "Items/ICustomDetailsViewCustomCategoryItem.h"
 #include "Items/ICustomDetailsViewCustomItem.h"
 #include "Items/ICustomDetailsViewItem.h"
 #include "Menus/DMMaterialStageSourceMenus.h"
@@ -342,6 +343,40 @@ TSharedRef<SWidget> SDMComponentEdit::CreateEditWidget()
 			}
 		}
 
+		FName CategoryName = EditRow.CategoryOverrideName;
+
+		if (CategoryName == NAME_None && EditRow.PropertyHandle.IsValid())
+		{
+			// Sub category (possibly)
+			if (TSharedPtr<IPropertyHandle> SubCategoryProperty = EditRow.PropertyHandle->GetParentHandle())
+			{
+				if (SubCategoryProperty->IsCategoryHandle())
+				{
+					// "Material Designer" (possibly)
+					if (TSharedPtr<IPropertyHandle> MaterialDesignerCategoryProperty = SubCategoryProperty->GetParentHandle())
+					{
+						if (MaterialDesignerCategoryProperty->IsCategoryHandle())
+						{
+							CategoryName = *SubCategoryProperty->GetPropertyDisplayName().ToString();
+						}
+					}
+				}
+			}
+		}
+
+		TSharedPtr<ICustomDetailsViewItem> CategoryItem;
+
+		if (CategoryName != NAME_None)
+		{
+			CategoryItem = DetailsView->FindCustomItem(CategoryName);
+
+			if (!CategoryItem.IsValid())
+			{
+				CategoryItem = DetailsView->CreateCustomCategoryItem(CategoryName, FText::FromName(CategoryName))->AsItem();
+				DetailsView->ExtendTree(RootId, ECustomDetailsTreeInsertPosition::Child, CategoryItem.ToSharedRef());
+			}
+		}
+
 		if (bHasValidCustomWidget)
 		{
 			TSharedPtr<ICustomDetailsViewCustomItem> Item = DetailsView->CreateCustomItem(
@@ -356,7 +391,16 @@ TSharedRef<SWidget> SDMComponentEdit::CreateEditWidget()
 			}
 
 			Item->SetValueWidget(EditRow.ValueWidget.ToSharedRef());
-			DetailsView->ExtendTree(RootId, Position, Item->AsItem());
+
+			if (CategoryItem.IsValid())
+			{
+				DetailsView->ExtendTree(CategoryItem->GetItemId(), Position, Item->AsItem());
+			}
+			else
+			{
+				DetailsView->ExtendTree(RootId, Position, Item->AsItem());
+			}
+
 			continue;
 		}
 
@@ -388,7 +432,14 @@ TSharedRef<SWidget> SDMComponentEdit::CreateEditWidget()
 			Item->SetResetToDefaultOverride(EditRow.ResetToDefaultOverride.GetValue());
 		}
 
-		DetailsView->ExtendTree(RootId, Position, Item);
+		if (CategoryItem.IsValid())
+		{
+			DetailsView->ExtendTree(CategoryItem->GetItemId(), Position, Item);
+		}
+		else
+		{
+			DetailsView->ExtendTree(RootId, Position, Item);
+		}
 	}
 
 	DetailsView->RebuildTree(ECustomDetailsViewBuildType::InstantBuild);
diff --git a/Engine/Plugins/Experimental/DynamicMaterial/Source/DynamicMaterialEditor/Private/Slate/SDMEditor.cpp b/Engine/Plugins/Experimental/DynamicMaterial/Source/DynamicMaterialEditor/Private/Slate/SDMEditor.cpp
--- a/Engine/Plugins/Experimental/DynamicMaterial/Source/DynamicMaterialEditor/Private/Slate/SDMEditor.cpp
+++ b/Engine/Plugins/Experimental/DynamicMaterial/Source/DynamicMaterialEditor/Private/Slate/SDMEditor.cpp
@@ -35,6 +35,7 @@
 #include "Styling/StyleColors.h"
 #include "ThumbnailRendering/ThumbnailManager.h"
 #include "Utils/DMBlueprintFunctionLibrary.h"
+#include "Utils/DMDetailsViewUtils.h"
 #include "Utils/DMPrivate.h"
 #include "Widgets/Input/SButton.h"
 #include "Widgets/Input/SCheckBox.h"
@@ -53,33 +54,7 @@ namespace UE::DynamicMaterialEditor::Private
 {
 	TSharedPtr<IDetailTreeNode> SearchGeneratorForNode(const TSharedRef<IPropertyRowGenerator>& InGenerator, FName InPropertyName)
 	{
-		for (const TSharedRef<IDetailTreeNode>& CategoryNode : InGenerator->GetRootTreeNodes())
-		{
-			if (CategoryNode->GetNodeName() != TEXT("Material Designer"))
-			{
-				continue;
-			}
-
-			TArray<TSharedRef<IDetailTreeNode>> ChildNodes;
-			CategoryNode->GetChildren(ChildNodes);
-
-			for (const TSharedRef<IDetailTreeNode>& ChildNode : ChildNodes)
-			{
-				if (ChildNode->GetNodeType() != EDetailNodeType::Item)
-				{
-					continue;
-				}
-
-				if (ChildNode->GetNodeName() != InPropertyName)
-				{
-					continue;
-				}
-
-				return ChildNode;
-			}
-		}
-
-		return nullptr;
+		return FDMDetailsViewUtils::SearchNodesForProperty(InGenerator->GetRootTreeNodes(), InPropertyName);
 	}
 
 	TSharedPtr<IPropertyRowGenerator> SearchForGenerator(const TArray<FDMPropertyHandle>& InPropertyHandles, UObject* InObject)
diff --git a/Engine/Plugins/Experimental/DynamicMaterial/Source/DynamicMaterialEditor/Public/DMEDefs.h b/Engine/Plugins/Experimental/DynamicMaterial/Source/DynamicMaterialEditor/Public/DMEDefs.h
--- a/Engine/Plugins/Experimental/DynamicMaterial/Source/DynamicMaterialEditor/Public/DMEDefs.h
+++ b/Engine/Plugins/Experimental/DynamicMaterial/Source/DynamicMaterialEditor/Public/DMEDefs.h
@@ -121,7 +121,8 @@ struct FDMPropertyHandle
 	TOptional<FText> NameToolTipOverride;
 	TOptional<FResetToDefaultOverride> ResetToDefaultOverride;
 	TSharedPtr<SWidget> ValueWidget;
-	FName ValueName;
+	FName ValueName = NAME_None;
+	FName CategoryOverrideName = NAME_None;
 };
 
 UENUM(BlueprintType, meta = (Bitflags, UseEnumValuesAsMaskValuesInEditor = "true"))
