Use final art for the normalize icon.

diff --git a/Engine/Plugins/Experimental/DynamicMaterial/Source/DynamicMaterialEditor/Private/DynamicMaterialEditorStyle.cpp b/Engine/Plugins/Experimental/DynamicMaterial/Source/DynamicMaterialEditor/Private/DynamicMaterialEditorStyle.cpp
--- a/Engine/Plugins/Experimental/DynamicMaterial/Source/DynamicMaterialEditor/Private/DynamicMaterialEditorStyle.cpp
+++ b/Engine/Plugins/Experimental/DynamicMaterial/Source/DynamicMaterialEditor/Private/DynamicMaterialEditorStyle.cpp
@@ -159,7 +159,7 @@ TSharedRef<FSlateStyleSet> FDynamicMaterialEditorStyle::Create()
 
 	Style->Set("Icons.Remove", new IMAGE_BRUSH("Icons/EditorIcons/Remove_16px", Icon16x16));
 
-	Style->Set("Icons.Normalize", new IMAGE_BRUSH("Icons/EditorIcons/Normalize", Icon16x16));
+	Style->Set("Icons.Normalize", new IMAGE_BRUSH_SVG("Icons/EditorIcons/Normalize", Icon16x16));
 
 	Style->Set("Icons.Stage.EnabledButton", new IMAGE_BRUSH("Icons/EditorIcons/WhiteBall", Icon8x8));
 	Style->Set("Icons.Stage.BaseToggle", new IMAGE_BRUSH("Icons/EditorIcons/BaseToggle_16x", Icon16x16));
diff --git a/Engine/Source/Runtime/SlateCore/Private/Styling/StarshipCoreStyle.cpp b/Engine/Source/Runtime/SlateCore/Private/Styling/StarshipCoreStyle.cpp
--- a/Engine/Source/Runtime/SlateCore/Private/Styling/StarshipCoreStyle.cpp
+++ b/Engine/Source/Runtime/SlateCore/Private/Styling/StarshipCoreStyle.cpp
@@ -354,7 +354,7 @@ TSharedRef<ISlateStyle> FStarshipCoreStyle::Create()
 		Style->Set("Icons.Lock", new IMAGE_BRUSH_SVG("Starship/Common/lock", Icon16x16));
 		Style->Set("Icons.Unlock", new IMAGE_BRUSH_SVG("Starship/Common/lock-unlocked", Icon16x16));
 
-		Style->Set("Icons.Normalize", new IMAGE_BRUSH("Starship/Common/normalize", Icon16x16));
+		Style->Set("Icons.Normalize", new IMAGE_BRUSH_SVG("Starship/Common/normalize", Icon16x16));
 
 		Style->Set("Icons.CircleArrowLeft", new IMAGE_BRUSH_SVG("Starship/Common/circle-arrow-left", Icon16x16));
 		Style->Set("Icons.CircleArrowRight", new IMAGE_BRUSH_SVG("Starship/Common/circle-arrow-right", Icon16x16));
