Material Designer
- Added an option to change the type of vector3 function inputs to linear colours instead of fvectors.

diff --git a/Engine/Plugins/Experimental/DynamicMaterial/Source/DynamicMaterialEditor/Private/Components/DMMaterialEffectFunction.cpp b/Engine/Plugins/Experimental/DynamicMaterial/Source/DynamicMaterialEditor/Private/Components/DMMaterialEffectFunction.cpp
--- a/Engine/Plugins/Experimental/DynamicMaterial/Source/DynamicMaterialEditor/Private/Components/DMMaterialEffectFunction.cpp
+++ b/Engine/Plugins/Experimental/DynamicMaterial/Source/DynamicMaterialEditor/Private/Components/DMMaterialEffectFunction.cpp
@@ -5,12 +5,14 @@
 #include "Components/DMMaterialValue.h"
 #include "Components/MaterialValues/DMMaterialValueFloat1.h"
 #include "Components/MaterialValues/DMMaterialValueFloat2.h"
+#include "Components/MaterialValues/DMMaterialValueFloat3RGB.h"
 #include "Components/MaterialValues/DMMaterialValueFloat3XYZ.h"
 #include "Components/MaterialValues/DMMaterialValueFloat4.h"
 #include "DMComponentPath.h"
 #include "Dom/JsonObject.h"
 #include "Dom/JsonValue.h"
 #include "DynamicMaterialEditorModule.h"
+#include "DynamicMaterialEditorSettings.h"
 #include "Materials/MaterialExpressionComponentMask.h"
 #include "Materials/MaterialExpressionFunctionInput.h"
 #include "Materials/MaterialExpressionFunctionOutput.h"
@@ -491,7 +493,14 @@ bool UDMMaterialEffectFunction::NeedsFunctionInit() const
 				break;
 
 			case EFunctionInputType::FunctionInput_Vector3:
-				ValueType = EDMValueType::VT_Float3_XYZ;
+				if (UDynamicMaterialEditorSettings::IsUseLinearColorForVectorsEnabled())
+				{
+					ValueType = EDMValueType::VT_Float3_RGB;
+				}
+				else
+				{
+					ValueType = EDMValueType::VT_Float3_XYZ;
+				}
 				break;
 
 			case EFunctionInputType::FunctionInput_Vector4:
@@ -626,7 +635,14 @@ void UDMMaterialEffectFunction::InitFunction()
 				break;
 
 			case EFunctionInputType::FunctionInput_Vector3:
-				ValueType = EDMValueType::VT_Float3_XYZ;
+				if (UDynamicMaterialEditorSettings::IsUseLinearColorForVectorsEnabled())
+				{
+					ValueType = EDMValueType::VT_Float3_RGB;
+				}
+				else
+				{
+					ValueType = EDMValueType::VT_Float3_XYZ;
+				}
 				break;
 
 			case EFunctionInputType::FunctionInput_Vector4:
@@ -669,10 +685,15 @@ void UDMMaterialEffectFunction::InitFunction()
 					break;
 
 				case EFunctionInputType::FunctionInput_Vector3:
-					if (UDMMaterialValueFloat3XYZ* Float3Value = Cast<UDMMaterialValueFloat3XYZ>(Value))
+					if (UDMMaterialValueFloat3XYZ* Float3XYZ = Cast<UDMMaterialValueFloat3XYZ>(Value))
+					{
+						Float3XYZ->SetDefaultValue({FunctionInput->PreviewValue.X, FunctionInput->PreviewValue.Y, FunctionInput->PreviewValue.Z});
+						Float3XYZ->ApplyDefaultValue();
+					}
+					else if (UDMMaterialValueFloat3RGB* Float3RGB = Cast<UDMMaterialValueFloat3RGB>(Value))
 					{
-						Float3Value->SetDefaultValue({FunctionInput->PreviewValue.X, FunctionInput->PreviewValue.Y, FunctionInput->PreviewValue.Z});
-						Float3Value->ApplyDefaultValue();
+						Float3RGB->SetDefaultValue({FunctionInput->PreviewValue.X, FunctionInput->PreviewValue.Y, FunctionInput->PreviewValue.Z});
+						Float3RGB->ApplyDefaultValue();
 					}
 					break;
 
diff --git a/Engine/Plugins/Experimental/DynamicMaterial/Source/DynamicMaterialEditor/Private/Components/DMMaterialStageFunction.cpp b/Engine/Plugins/Experimental/DynamicMaterial/Source/DynamicMaterialEditor/Private/Components/DMMaterialStageFunction.cpp
--- a/Engine/Plugins/Experimental/DynamicMaterial/Source/DynamicMaterialEditor/Private/Components/DMMaterialStageFunction.cpp
+++ b/Engine/Plugins/Experimental/DynamicMaterial/Source/DynamicMaterialEditor/Private/Components/DMMaterialStageFunction.cpp
@@ -9,9 +9,11 @@
 #include "Components/MaterialStageInputs/DMMSIValue.h"
 #include "Components/MaterialValues/DMMaterialValueFloat1.h"
 #include "Components/MaterialValues/DMMaterialValueFloat2.h"
+#include "Components/MaterialValues/DMMaterialValueFloat3RGB.h"
 #include "Components/MaterialValues/DMMaterialValueFloat3XYZ.h"
 #include "Components/MaterialValues/DMMaterialValueFloat4.h"
 #include "DynamicMaterialEditorModule.h"
+#include "DynamicMaterialEditorSettings.h"
 #include "Materials/MaterialExpressionFunctionInput.h"
 #include "Materials/MaterialExpressionFunctionOutput.h"
 #include "Materials/MaterialExpressionMaterialFunctionCall.h"
@@ -376,7 +378,14 @@ bool UDMMaterialStageFunction::NeedsFunctionInit() const
 				break;
 
 			case EFunctionInputType::FunctionInput_Vector3:
-				ValueType = EDMValueType::VT_Float3_XYZ;
+				if (UDynamicMaterialEditorSettings::IsUseLinearColorForVectorsEnabled())
+				{
+					ValueType = EDMValueType::VT_Float3_RGB;
+				}
+				else
+				{
+					ValueType = EDMValueType::VT_Float3_XYZ;
+				}
 				break;
 
 			case EFunctionInputType::FunctionInput_Vector4:
@@ -523,7 +532,14 @@ void UDMMaterialStageFunction::InitFunction()
 				break;
 
 			case EFunctionInputType::FunctionInput_Vector3:
-				ValueType = EDMValueType::VT_Float3_XYZ;
+				if (UDynamicMaterialEditorSettings::IsUseLinearColorForVectorsEnabled())
+				{
+					ValueType = EDMValueType::VT_Float3_RGB;
+				}
+				else
+				{
+					ValueType = EDMValueType::VT_Float3_XYZ;
+				}
 				break;
 
 			case EFunctionInputType::FunctionInput_Vector4:
@@ -590,10 +606,15 @@ void UDMMaterialStageFunction::InitFunction()
 					break;
 
 				case EFunctionInputType::FunctionInput_Vector3:
-					if (UDMMaterialValueFloat3XYZ* Float3Value = Cast<UDMMaterialValueFloat3XYZ>(Value))
+					if (UDMMaterialValueFloat3XYZ* Float3XYZ = Cast<UDMMaterialValueFloat3XYZ>(Value))
 					{
-						Float3Value->SetDefaultValue({FunctionInput->PreviewValue.X, FunctionInput->PreviewValue.Y, FunctionInput->PreviewValue.Z});
-						Float3Value->ApplyDefaultValue();
+						Float3XYZ->SetDefaultValue({FunctionInput->PreviewValue.X, FunctionInput->PreviewValue.Y, FunctionInput->PreviewValue.Z});
+						Float3XYZ->ApplyDefaultValue();
+					}
+					else if (UDMMaterialValueFloat3RGB* Float3RGB = Cast<UDMMaterialValueFloat3RGB>(Value))
+					{
+						Float3RGB->SetDefaultValue({FunctionInput->PreviewValue.X, FunctionInput->PreviewValue.Y, FunctionInput->PreviewValue.Z});
+						Float3RGB->ApplyDefaultValue();
 					}
 					break;
 
diff --git a/Engine/Plugins/Experimental/DynamicMaterial/Source/DynamicMaterialEditor/Private/DynamicMaterialEditorSettings.cpp b/Engine/Plugins/Experimental/DynamicMaterial/Source/DynamicMaterialEditor/Private/DynamicMaterialEditorSettings.cpp
--- a/Engine/Plugins/Experimental/DynamicMaterial/Source/DynamicMaterialEditor/Private/DynamicMaterialEditorSettings.cpp
+++ b/Engine/Plugins/Experimental/DynamicMaterial/Source/DynamicMaterialEditor/Private/DynamicMaterialEditorSettings.cpp
@@ -83,6 +83,7 @@ UDynamicMaterialEditorSettings::UDynamicMaterialEditorSettings()
 	SectionName = TEXT("Material Designer");
 
 	bFollowSelection = true;
+	bUseLinearColorForVectors = true;
 
 	ResetAllLayoutSettings();
 
@@ -169,6 +170,16 @@ UDynamicMaterialEditorSettings* UDynamicMaterialEditorSettings::Get()
 	return DefaultSettings;
 }
 
+bool UDynamicMaterialEditorSettings::IsUseLinearColorForVectorsEnabled()
+{
+	if (const UDynamicMaterialEditorSettings* Settings = UDynamicMaterialEditorSettings::Get())
+	{
+		return Settings->bUseLinearColorForVectors;
+	}
+
+	return true;
+}
+
 void UDynamicMaterialEditorSettings::PreEditChange(FEditPropertyChain& InPropertyAboutToChange)
 {
 	Super::PreEditChange(InPropertyAboutToChange);
diff --git a/Engine/Plugins/Experimental/DynamicMaterial/Source/DynamicMaterialEditor/Private/DynamicMaterialEditorSettings.h b/Engine/Plugins/Experimental/DynamicMaterial/Source/DynamicMaterialEditor/Private/DynamicMaterialEditorSettings.h
--- a/Engine/Plugins/Experimental/DynamicMaterial/Source/DynamicMaterialEditor/Private/DynamicMaterialEditorSettings.h
+++ b/Engine/Plugins/Experimental/DynamicMaterial/Source/DynamicMaterialEditor/Private/DynamicMaterialEditorSettings.h
@@ -198,6 +198,12 @@ public:
 	UPROPERTY(Config, EditAnywhere, BlueprintReadWrite, Category = "Channels", meta = (TitleProperty = Name))
 	TArray<FDMMaterialChannelListPreset> MaterialChannelPresets;
 
+	UPROPERTY(Config, EditAnywhere, BlueprintReadWrite, Category = "Material Functions")
+	bool bUseLinearColorForVectors;
+
+	/** This variable is accessed in multiple places, so this is a quick accessor. */
+	static bool IsUseLinearColorForVectorsEnabled();
+
 	FOnFinishedChangingProperties OnSettingsChanged;
 
 	//~ Begin UObject
