Material Designer
- Renamed "Settings" bubble to "Global"
- Added "Global Settings" category to global settings display

diff --git a/Engine/Plugins/Experimental/DynamicMaterial/Source/DynamicMaterialEditor/Private/Slate/SDMComponentEdit.cpp b/Engine/Plugins/Experimental/DynamicMaterial/Source/DynamicMaterialEditor/Private/Slate/SDMComponentEdit.cpp
--- a/Engine/Plugins/Experimental/DynamicMaterial/Source/DynamicMaterialEditor/Private/Slate/SDMComponentEdit.cpp
+++ b/Engine/Plugins/Experimental/DynamicMaterial/Source/DynamicMaterialEditor/Private/Slate/SDMComponentEdit.cpp
@@ -834,12 +834,13 @@ void SDMComponentEdit::OnUndo()
 void SDMComponentEdit::GenerateMaterialModelPropertyRows(const TSharedRef<SDMEditor> InEditorWidget, UDynamicMaterialModel* InMaterialModel,
 	TArray<FDMPropertyHandle>& InOutPropertyRows, TSet<UDMMaterialComponent*>& InOutProcessedObjects)
 {
+	const FName Category = FName("Global Settings");
 	const FName UIMin = FName("UIMin");
 	const FName UIMax = FName("UIMax");
 	const FName ClampMin = FName("ClampMin");
 	const FName ClampMax = FName("ClampMax");
 
-	auto AddGlobalVar = [InEditorWidget, InMaterialModel, &InOutPropertyRows, &InOutProcessedObjects, &UIMin, &UIMax, &ClampMin, &ClampMax]
+	auto AddGlobalVar = [InEditorWidget, InMaterialModel, &InOutPropertyRows, &InOutProcessedObjects, &Category , &UIMin, &UIMax, &ClampMin, &ClampMax]
 		(UDMMaterialValue* InValue, const FText& InNameOverride)
 		{
 			if (IsValid(InValue))
@@ -849,6 +850,7 @@ void SDMComponentEdit::GenerateMaterialModelPropertyRows(const TSharedRef<SDMEdi
 					FDMPropertyHandle& ValueHandle = InOutPropertyRows.Add_GetRef(InEditorWidget->GetPropertyHandle(&*InEditorWidget,
 						InValue, UDMMaterialValue::ValueName));
 
+					ValueHandle.CategoryOverrideName = Category;
 					ValueHandle.NameOverride = InNameOverride;
 
 					if (UDMMaterialValueFloat* FloatValue = Cast<UDMMaterialValueFloat>(InValue))
diff --git a/Engine/Plugins/Experimental/DynamicMaterial/Source/DynamicMaterialEditor/Private/Slate/SDMEditor.cpp b/Engine/Plugins/Experimental/DynamicMaterial/Source/DynamicMaterialEditor/Private/Slate/SDMEditor.cpp
--- a/Engine/Plugins/Experimental/DynamicMaterial/Source/DynamicMaterialEditor/Private/Slate/SDMEditor.cpp
+++ b/Engine/Plugins/Experimental/DynamicMaterial/Source/DynamicMaterialEditor/Private/Slate/SDMEditor.cpp
@@ -688,7 +688,7 @@ TSharedRef<SWidget> SDMEditor::CreateSlotPickerWidget()
 			.Content()
 			[
 				SNew(STextBlock)
-				.Text(LOCTEXT("Settings", "Settings"))
+				.Text(LOCTEXT("GlobalSettings", "Global"))
 			]
 		];
 
