Material Designer
- Moved settings from the dropdown menu into a global settings category called "Material Type".

diff --git a/Engine/Plugins/Experimental/DynamicMaterial/Source/DynamicMaterialEditor/Private/Menus/DMToolBarMenus.cpp b/Engine/Plugins/Experimental/DynamicMaterial/Source/DynamicMaterialEditor/Private/Menus/DMToolBarMenus.cpp
--- a/Engine/Plugins/Experimental/DynamicMaterial/Source/DynamicMaterialEditor/Private/Menus/DMToolBarMenus.cpp
+++ b/Engine/Plugins/Experimental/DynamicMaterial/Source/DynamicMaterialEditor/Private/Menus/DMToolBarMenus.cpp
@@ -32,461 +32,9 @@
 namespace UE::DynamicMaterialEditor::Private
 {
 	static FName ToolBarEditorLayoutMenuName = TEXT("MaterialDesigner.EditorLayout");
-	static FName ToolBarMaterialInstanceSectionName = TEXT("MaterialInstance");
 	static FName ToolBarMaterialExportSectionName = TEXT("MaterialExport");
 	static FName ToolBarMaterialDesignerSettingsSectionName = TEXT("MaterialDesignerSettings");
 
-	void AddToolBarMaterialInstanceChannelListMenu_Execute(const FToolMenuContext& InContext, FName InPresetName)
-	{
-		if (UDMMenuContext* Context = InContext.FindContext<UDMMenuContext>())
-		{
-			if (UDynamicMaterialModel* MaterialModel = Context->GetModel())
-			{
-				if (UDynamicMaterialModelEditorOnlyData* EditorOnlyData = UDynamicMaterialModelEditorOnlyData::Get(MaterialModel))
-				{
-					EditorOnlyData->SetChannelListPreset(InPresetName);
-
-					if (TSharedPtr<SDMEditor> EditorWidget = Context->GetEditorWidget())
-					{
-						EditorWidget->RefreshSlotPickerList();
-						EditorWidget->RefreshSlotWidget();
-						EditorWidget->RefreshComponentEditWidget();
-					}
-				}
-			}
-		}
-	}
-
-	ECheckBoxState AddToolBarMaterialInstanceChannelListMenu_GetCheckState(const FToolMenuContext& InContext, FName InPresetName)
-	{
-		if (UDMMenuContext* Context = InContext.FindContext<UDMMenuContext>())
-		{
-			if (UDynamicMaterialModel* MaterialModel = Context->GetModel())
-			{
-				if (UDynamicMaterialModelEditorOnlyData* EditorOnlyData = UDynamicMaterialModelEditorOnlyData::Get(MaterialModel))
-				{
-					return InPresetName == EditorOnlyData->GetChannelListPreset()
-						? ECheckBoxState::Checked
-						: ECheckBoxState::Unchecked;
-				}
-			}
-		}
-
-		return ECheckBoxState::Undetermined;
-	}
-
-	void AddToolBarMaterialInstanceChannelListMenu(UToolMenu* InMenu)
-	{
-		FToolMenuSection& NewSection = InMenu->AddSection("MaterialType", LOCTEXT("MaterialType", "Material Type"));
-
-		for (const FDMMaterialChannelListPreset& Preset : GetDefault<UDynamicMaterialEditorSettings>()->MaterialChannelPresets)
-		{
-			FToolUIAction MaterialChannelListType;
-			MaterialChannelListType.ExecuteAction = FToolMenuExecuteAction::CreateStatic(&AddToolBarMaterialInstanceChannelListMenu_Execute, Preset.Name);
-			MaterialChannelListType.GetActionCheckState = FToolMenuGetActionCheckState::CreateStatic(&AddToolBarMaterialInstanceChannelListMenu_GetCheckState, Preset.Name);
-
-			NewSection.AddMenuEntry(
-				Preset.Name,
-				FText::FromName(Preset.Name),
-				FText::GetEmpty(),
-				TAttribute<FSlateIcon>(),
-				FToolUIActionChoice(MaterialChannelListType),
-				EUserInterfaceActionType::RadioButton
-			);
-		}
-	}
-
-	void AddToolBarMaterialInstanceDomainMenu_Execute(const FToolMenuContext& InContext, EMaterialDomain InDomain)
-	{
-		if (UDMMenuContext* Context = InContext.FindContext<UDMMenuContext>())
-		{
-			if (UDynamicMaterialModel* MaterialModel = Context->GetModel())
-			{
-				if (UDynamicMaterialModelEditorOnlyData* EditorOnlyData = UDynamicMaterialModelEditorOnlyData::Get(MaterialModel))
-				{
-					EditorOnlyData->SetDomain(InDomain);
-
-					if (TSharedPtr<SDMEditor> EditorWidget = Context->GetEditorWidget())
-					{
-						EditorWidget->RefreshSlotPickerList();
-						EditorWidget->RefreshSlotWidget();
-						EditorWidget->RefreshComponentEditWidget();
-					}
-				}
-			}
-		}
-	}
-
-	bool AddToolBarMaterialInstanceDomainMenu_CanExecute(const FToolMenuContext& InContext, EMaterialDomain InDomain)
-	{
-		if (UDMMenuContext* Context = InContext.FindContext<UDMMenuContext>())
-		{
-			if (UDynamicMaterialModel* MaterialModel = Context->GetModel())
-			{
-				return !!UDynamicMaterialModelEditorOnlyData::Get(MaterialModel);
-			}
-		}
-
-		return false;
-	}
-
-	ECheckBoxState AddToolBarMaterialInstanceDomainMenu_GetCheckState(const FToolMenuContext& InContext, EMaterialDomain InDomain)
-	{
-		if (UDMMenuContext* Context = InContext.FindContext<UDMMenuContext>())
-		{
-			if (UDynamicMaterialModel* MaterialModel = Context->GetModel())
-			{
-				if (UDynamicMaterialModelEditorOnlyData* EditorOnlyData = UDynamicMaterialModelEditorOnlyData::Get(MaterialModel))
-				{
-					return EditorOnlyData->GetDomain() == InDomain
-						? ECheckBoxState::Checked
-						: ECheckBoxState::Unchecked;
-				}
-			}
-		}
-
-		return ECheckBoxState::Undetermined;
-	}
-
-	void AddToolBarMaterialInstanceDomainMenu(UToolMenu* InMenu)
-	{
-		UEnum* DomainEnum = StaticEnum<EMaterialDomain>();
-
-		FToolMenuSection& NewSection = InMenu->AddSection("MaterialDomains", LOCTEXT("MaterialDomains", "Material Domains"));
-
-		for (EMaterialDomain Domain : UDynamicMaterialModelEditorOnlyData::SupportedDomains)
-		{
-			FToolUIAction DomainAction;
-			DomainAction.ExecuteAction = FToolMenuExecuteAction::CreateStatic(&AddToolBarMaterialInstanceDomainMenu_Execute, Domain);
-			DomainAction.CanExecuteAction = FToolMenuCanExecuteAction::CreateStatic(&AddToolBarMaterialInstanceDomainMenu_CanExecute, Domain);
-			DomainAction.GetActionCheckState = FToolMenuGetActionCheckState::CreateStatic(&AddToolBarMaterialInstanceDomainMenu_GetCheckState, Domain);
-
-			NewSection.AddMenuEntry(
-				DomainEnum->GetNameByValue(Domain),
-				DomainEnum->GetDisplayNameTextByValue(Domain),
-				FText::GetEmpty(),
-				TAttribute<FSlateIcon>(),
-				FToolUIActionChoice(DomainAction),
-				EUserInterfaceActionType::RadioButton
-			);
-		}
-	}
-
-	void AddToolBarMaterialInstanceBlendModeMenu_Execute(const FToolMenuContext& InContext, EBlendMode InBlendMode)
-	{
-		if (UDMMenuContext* Context = InContext.FindContext<UDMMenuContext>())
-		{
-			if (UDynamicMaterialModel* MaterialModel = Context->GetModel())
-			{
-				if (UDynamicMaterialModelEditorOnlyData* EditorOnlyData = UDynamicMaterialModelEditorOnlyData::Get(MaterialModel))
-				{
-					EditorOnlyData->SetBlendMode(InBlendMode);
-
-					if (TSharedPtr<SDMEditor> EditorWidget = Context->GetEditorWidget())
-					{
-						EditorWidget->RefreshSlotPickerList();
-						EditorWidget->RefreshSlotWidget();
-						EditorWidget->RefreshComponentEditWidget();
-					}
-				}
-			}
-		}
-	}
-
-	bool AddToolBarMaterialInstanceBlendModeMenu_CanExecute(const FToolMenuContext& InContext, EBlendMode InBlendMode)
-	{
-		if (UDMMenuContext* Context = InContext.FindContext<UDMMenuContext>())
-		{
-			if (UDynamicMaterialModel* MaterialModel = Context->GetModel())
-			{
-				return !!UDynamicMaterialModelEditorOnlyData::Get(MaterialModel);
-			}
-		}
-
-		return false;
-	}
-
-	ECheckBoxState AddToolBarMaterialInstanceBlendModeMenu_GetCheckState(const FToolMenuContext& InContext, EBlendMode InBlendMode)
-	{
-		if (UDMMenuContext* Context = InContext.FindContext<UDMMenuContext>())
-		{
-			if (UDynamicMaterialModel* MaterialModel = Context->GetModel())
-			{
-				if (UDynamicMaterialModelEditorOnlyData* EditorOnlyData = UDynamicMaterialModelEditorOnlyData::Get(MaterialModel))
-				{
-					return EditorOnlyData->GetBlendMode() == InBlendMode
-						? ECheckBoxState::Checked
-						: ECheckBoxState::Unchecked;
-				}
-			}
-		}
-
-		return ECheckBoxState::Undetermined;
-	}
-
-	void AddToolBarMaterialInstanceBlendModeMenu(UToolMenu* InMenu)
-	{
-		UEnum* BlendModeEnum = StaticEnum<EBlendMode>();
-
-		FToolMenuSection& NewSection = InMenu->AddSection("MaterialBlendModes", LOCTEXT("MaterialBlendModes", "Material Blend Modes"));
-
-		for (EBlendMode BlendMode : UDynamicMaterialModelEditorOnlyData::SupportedBlendModes)
-		{
-			FToolUIAction BlendModeAction;
-			BlendModeAction.ExecuteAction = FToolMenuExecuteAction::CreateStatic(&AddToolBarMaterialInstanceBlendModeMenu_Execute, BlendMode);
-			BlendModeAction.CanExecuteAction = FToolMenuCanExecuteAction::CreateStatic(&AddToolBarMaterialInstanceBlendModeMenu_CanExecute, BlendMode);
-			BlendModeAction.GetActionCheckState = FToolMenuGetActionCheckState::CreateStatic(&AddToolBarMaterialInstanceBlendModeMenu_GetCheckState, BlendMode);
-
-			NewSection.AddMenuEntry(
-				BlendModeEnum->GetNameByValue(BlendMode),
-				BlendModeEnum->GetDisplayNameTextByValue(BlendMode),
-				FText::GetEmpty(),
-				TAttribute<FSlateIcon>(),
-				FToolUIActionChoice(BlendModeAction),
-				EUserInterfaceActionType::RadioButton
-			);
-		}
-	}
-
-	void AddToolBarMaterialInstanceUnlitMenu_Execute(const FToolMenuContext& InContext)
-	{
-		if (UDMMenuContext* Context = InContext.FindContext<UDMMenuContext>())
-		{
-			if (UDynamicMaterialModel* MaterialModel = Context->GetModel())
-			{
-				if (UDynamicMaterialModelEditorOnlyData* EditorOnlyData = UDynamicMaterialModelEditorOnlyData::Get(MaterialModel))
-				{
-					switch (EditorOnlyData->GetShadingModel())
-					{
-						case EDMMaterialShadingModel::DefaultLit:
-							EditorOnlyData->SetShadingModel(EDMMaterialShadingModel::Unlit);
-							break;
-
-						case EDMMaterialShadingModel::Unlit:
-							EditorOnlyData->SetShadingModel(EDMMaterialShadingModel::DefaultLit);
-							break;
-					}
-
-					if (TSharedPtr<SDMEditor> EditorWidget = Context->GetEditorWidget())
-					{
-						EditorWidget->RefreshSlotPickerList();
-						EditorWidget->RefreshSlotWidget();
-						EditorWidget->RefreshComponentEditWidget();
-					}
-				}
-			}
-		}
-	}
-
-	bool AddToolBarMaterialInstanceUnlitMenu_CanExecute(const FToolMenuContext& InContext)
-	{
-		if (UDMMenuContext* Context = InContext.FindContext<UDMMenuContext>())
-		{
-			if (UDynamicMaterialModel* MaterialModel = Context->GetModel())
-			{
-				return !!UDynamicMaterialModelEditorOnlyData::Get(MaterialModel);
-			}
-		}
-
-		return false;
-	}
-
-	ECheckBoxState AddToolBarMaterialInstanceUnlitMenu_GetCheckState(const FToolMenuContext& InContext)
-	{
-		if (UDMMenuContext* Context = InContext.FindContext<UDMMenuContext>())
-		{
-			if (UDynamicMaterialModel* MaterialModel = Context->GetModel())
-			{
-				if (UDynamicMaterialModelEditorOnlyData* EditorOnlyData = UDynamicMaterialModelEditorOnlyData::Get(MaterialModel))
-				{
-					switch (EditorOnlyData->GetShadingModel())
-					{
-						case EDMMaterialShadingModel::DefaultLit:
-							return ECheckBoxState::Unchecked;
-
-						case EDMMaterialShadingModel::Unlit:
-							return ECheckBoxState::Checked;
-					}
-				}
-			}
-		}
-
-		return ECheckBoxState::Undetermined;
-	}
-
-	void AddToolBarMaterialInstanceAnimatedMenu_Execute(const FToolMenuContext& InContext)
-	{
-		if (UDMMenuContext* Context = InContext.FindContext<UDMMenuContext>())
-		{
-			if (UDynamicMaterialModel* MaterialModel = Context->GetModel())
-			{
-				if (UDynamicMaterialModelEditorOnlyData* EditorOnlyData = UDynamicMaterialModelEditorOnlyData::Get(MaterialModel))
-				{
-					if (TSharedPtr<SDMEditor> EditorWidget = Context->GetEditorWidget())
-					{
-						EditorWidget->RefreshSlotPickerList();
-						EditorWidget->RefreshSlotWidget();
-						EditorWidget->RefreshComponentEditWidget();
-					}					
-				}
-			}
-		}
-	}
-
-	bool AddToolBarMaterialInstanceAnimatedMenu_CanExecute(const FToolMenuContext& InContext)
-	{
-		if (UDMMenuContext* Context = InContext.FindContext<UDMMenuContext>())
-		{
-			if (UDynamicMaterialModel* MaterialModel = Context->GetModel())
-			{
-				return !!UDynamicMaterialModelEditorOnlyData::Get(MaterialModel);
-			}
-		}
-
-		return false;
-	}
-
-	ECheckBoxState AddToolBarMaterialInstanceAnimatedMenu_GetCheckState(const FToolMenuContext& InContext)
-	{
-		if (UDMMenuContext* Context = InContext.FindContext<UDMMenuContext>())
-		{
-			if (UDynamicMaterialModel* MaterialModel = Context->GetModel())
-			{
-				if (UDynamicMaterialModelEditorOnlyData* EditorOnlyData = UDynamicMaterialModelEditorOnlyData::Get(MaterialModel))
-				{
-					return EditorOnlyData->IsPixelAnimationFlagSet()
-						? ECheckBoxState::Checked
-						: ECheckBoxState::Unchecked;
-				}
-			}
-		}
-
-		return ECheckBoxState::Undetermined;
-	}
-
-	void AddToolBarMaterialInstanceTwoSidedMenu_Execute(const FToolMenuContext& InContext)
-	{
-		if (UDMMenuContext* Context = InContext.FindContext<UDMMenuContext>())
-		{
-			if (UDynamicMaterialModel* MaterialModel = Context->GetModel())
-			{
-				if (UDynamicMaterialModelEditorOnlyData* EditorOnlyData = UDynamicMaterialModelEditorOnlyData::Get(MaterialModel))
-				{
-					if (TSharedPtr<SDMEditor> EditorWidget = Context->GetEditorWidget())
-					{
-						EditorWidget->RefreshSlotPickerList();
-						EditorWidget->RefreshSlotWidget();
-						EditorWidget->RefreshComponentEditWidget();
-					}					
-				}
-			}
-		}
-	}
-
-	bool AddToolBarMaterialInstanceTwoSidedMenu_CanExecute(const FToolMenuContext& InContext)
-	{
-		if (UDMMenuContext* Context = InContext.FindContext<UDMMenuContext>())
-		{
-			if (UDynamicMaterialModel* MaterialModel = Context->GetModel())
-			{
-				return !!UDynamicMaterialModelEditorOnlyData::Get(MaterialModel);
-			}
-		}
-
-		return false;
-	}
-
-	ECheckBoxState AddToolBarMaterialInstanceTwoSidedMenu_GetCheckState(const FToolMenuContext& InContext)
-	{
-		if (UDMMenuContext* Context = InContext.FindContext<UDMMenuContext>())
-		{
-			if (UDynamicMaterialModel* MaterialModel = Context->GetModel())
-			{
-				if (UDynamicMaterialModelEditorOnlyData* EditorOnlyData = UDynamicMaterialModelEditorOnlyData::Get(MaterialModel))
-				{
-					return EditorOnlyData->IsTwoSidedFlagSet()
-						? ECheckBoxState::Checked
-						: ECheckBoxState::Unchecked;
-				}
-			}
-		}
-
-		return ECheckBoxState::Undetermined;
-	}
-
-	void AddToolBarInstanceMenu(UToolMenu* InMenu)
-	{
-		if (!IsValid(InMenu) || InMenu->ContainsSection(ToolBarMaterialInstanceSectionName))
-		{
-			return;
-		}
-
-		FToolMenuSection& NewSection = InMenu->AddSection(ToolBarMaterialInstanceSectionName, LOCTEXT("MaterialInstanceSection", "Material Instance"));
-
-		NewSection.AddSubMenu(
-			"ChannelList",
-			LOCTEXT("MaterialInstanceMaterialType", "Material Type"),
-			LOCTEXT("MaterialInstanceMaterialType_ToolTip", "Set the material type (channel list) for the active material designer instance."),
-			FNewToolMenuDelegate::CreateStatic(&AddToolBarMaterialInstanceChannelListMenu)
-		);
-
-		NewSection.AddSubMenu(
-			"Domain",
-			LOCTEXT("MaterialInstanceDomain", "Material Domain"),
-			LOCTEXT("MaterialInstanceDomain_ToolTip", "Set the material domain for the active material designer instance."),
-			FNewToolMenuDelegate::CreateStatic(&AddToolBarMaterialInstanceDomainMenu)
-		);
-
-		NewSection.AddSubMenu(
-			"BlendMode",
-			LOCTEXT("MaterialInstanceBlendMode", "Material Blend Mode"),
-			LOCTEXT("MaterialInstanceBlendMode_ToolTip", "Set the material blend mode for the active material designer instance."),
-			FNewToolMenuDelegate::CreateStatic(&AddToolBarMaterialInstanceBlendModeMenu)
-		);
-
-		FToolUIAction UnlitAction;
-		UnlitAction.ExecuteAction = FToolMenuExecuteAction::CreateStatic(&AddToolBarMaterialInstanceUnlitMenu_Execute);
-		UnlitAction.CanExecuteAction = FToolMenuCanExecuteAction::CreateStatic(&AddToolBarMaterialInstanceUnlitMenu_CanExecute);
-		UnlitAction.GetActionCheckState = FToolMenuGetActionCheckState::CreateStatic(&AddToolBarMaterialInstanceUnlitMenu_GetCheckState);
-
-		NewSection.AddMenuEntry(
-			"Unlit",
-			LOCTEXT("Unlit", "Unlit"),
-			LOCTEXT("UnlitTooltip", "Whether this material requires light to be seen."),
-			TAttribute<FSlateIcon>(),
-			FToolUIActionChoice(UnlitAction),
-			EUserInterfaceActionType::ToggleButton
-		);
-
-		FToolUIAction AnimatedAction;
-		AnimatedAction.ExecuteAction = FToolMenuExecuteAction::CreateStatic(&AddToolBarMaterialInstanceAnimatedMenu_Execute);
-		AnimatedAction.CanExecuteAction = FToolMenuCanExecuteAction::CreateStatic(&AddToolBarMaterialInstanceAnimatedMenu_CanExecute);
-		AnimatedAction.GetActionCheckState = FToolMenuGetActionCheckState::CreateStatic(&AddToolBarMaterialInstanceAnimatedMenu_GetCheckState);
-
-		NewSection.AddMenuEntry(
-			"Animated",
-			LOCTEXT("Animated", "Animated"),
-			LOCTEXT("AnimatedTooltip", "Enables motion vectors in supported materials."),
-			TAttribute<FSlateIcon>(),
-			FToolUIActionChoice(AnimatedAction),
-			EUserInterfaceActionType::ToggleButton
-		);
-
-		FToolUIAction TwoSidedAction;
-		TwoSidedAction.ExecuteAction = FToolMenuExecuteAction::CreateStatic(&AddToolBarMaterialInstanceTwoSidedMenu_Execute);
-		TwoSidedAction.CanExecuteAction = FToolMenuCanExecuteAction::CreateStatic(&AddToolBarMaterialInstanceTwoSidedMenu_CanExecute);
-		TwoSidedAction.GetActionCheckState = FToolMenuGetActionCheckState::CreateStatic(&AddToolBarMaterialInstanceTwoSidedMenu_GetCheckState);
-
-		NewSection.AddMenuEntry(
-			"TwoSided",
-			LOCTEXT("TwoSided", "Two Sided"),
-			LOCTEXT("TwoSidedTooltip", "Whether this material will render on both the back and front of geometry."),
-			TAttribute<FSlateIcon>(),
-			FToolUIActionChoice(TwoSidedAction),
-			EUserInterfaceActionType::ToggleButton
-		);
-	}
-
 	void OpenMaterialEditorFromContext(UDMMenuContext* InMenuContext)
 	{
 		if (IsValid(InMenuContext))
@@ -936,7 +484,6 @@ namespace UE::DynamicMaterialEditor::Private
 
 	void AddToolBarEditorLayoutMenu(UToolMenu* InMenu)
 	{
-		AddToolBarInstanceMenu(InMenu);
 		AddToolBarExportMenu(InMenu);
 		AddToolBarSettingsMenu(InMenu);
 	}
diff --git a/Engine/Plugins/Experimental/DynamicMaterial/Source/DynamicMaterialEditor/Private/Model/DynamicMaterialModelEditorOnlyData.cpp b/Engine/Plugins/Experimental/DynamicMaterial/Source/DynamicMaterialEditor/Private/Model/DynamicMaterialModelEditorOnlyData.cpp
--- a/Engine/Plugins/Experimental/DynamicMaterial/Source/DynamicMaterialEditor/Private/Model/DynamicMaterialModelEditorOnlyData.cpp
+++ b/Engine/Plugins/Experimental/DynamicMaterial/Source/DynamicMaterialEditor/Private/Model/DynamicMaterialModelEditorOnlyData.cpp
@@ -137,15 +137,14 @@ UDynamicMaterialModelEditorOnlyData* UDynamicMaterialModelEditorOnlyData::Get(ID
 }
 
 UDynamicMaterialModelEditorOnlyData::UDynamicMaterialModelEditorOnlyData()
+	: State(EDMState::Idle)
+	, Domain(EMaterialDomain::MD_Surface)
+	, BlendMode(EBlendMode::BLEND_Opaque)
+	, ShadingModel(EDMMaterialShadingModel::DefaultLit)
+	, bPixelAnimationFlag(false)
+	, bTwoSidedFlag(true)
+	, ChannelListPreset(NAME_None)
 {
-	State = EDMState::Idle;
-
-	BlendMode = EBlendMode::BLEND_Translucent;
-	ShadingModel = EDMMaterialShadingModel::Unlit;
-	bPixelAnimationFlag = false;
-	bTwoSidedFlag = true;
-	ChannelListPreset = NAME_None;
-
 	Properties.Emplace(EDMMaterialPropertyType::BaseColor,           CreateDefaultSubobject<UDMMaterialPropertyBaseColor>(          "MaterialProperty_BaseColor"));
 	Properties.Emplace(EDMMaterialPropertyType::EmissiveColor,       CreateDefaultSubobject<UDMMaterialPropertyEmissiveColor>(      "MaterialProperty_EmissiveColor"));
 	Properties.Emplace(EDMMaterialPropertyType::Opacity,             CreateDefaultSubobject<UDMMaterialPropertyOpacity>(            "MaterialProperty_Opacity"));
@@ -186,6 +185,26 @@ void UDynamicMaterialModelEditorOnlyData::AssignPropertyAlphaValues()
 	Properties[EDMMaterialPropertyType::PixelDepthOffset   ]->AddComponent(AlphaValueName, MaterialModel->GetGlobalParameterValue(UDynamicMaterialModel::GlobalPixelDepthOffsetValueName));
 }
 
+TArray<FName> UDynamicMaterialModelEditorOnlyData::GetPresetOptions() const
+{
+	UDynamicMaterialEditorSettings* Settings = UDynamicMaterialEditorSettings::Get();
+
+	if (!Settings)
+	{
+		return {ChannelListPreset};
+	}
+
+	TArray<FName> PresetNames;
+	PresetNames.Reserve(Settings->MaterialChannelPresets.Num());
+
+	for (const FDMMaterialChannelListPreset& Preset : Settings->MaterialChannelPresets)
+	{
+		PresetNames.Add(Preset.Name);
+	}
+
+	return PresetNames;
+}
+
 void UDynamicMaterialModelEditorOnlyData::Initialize()
 {
 	if (!Slots.IsEmpty())
diff --git a/Engine/Plugins/Experimental/DynamicMaterial/Source/DynamicMaterialEditor/Private/Slate/SDMComponentEdit.cpp b/Engine/Plugins/Experimental/DynamicMaterial/Source/DynamicMaterialEditor/Private/Slate/SDMComponentEdit.cpp
--- a/Engine/Plugins/Experimental/DynamicMaterial/Source/DynamicMaterialEditor/Private/Slate/SDMComponentEdit.cpp
+++ b/Engine/Plugins/Experimental/DynamicMaterial/Source/DynamicMaterialEditor/Private/Slate/SDMComponentEdit.cpp
@@ -834,13 +834,35 @@ void SDMComponentEdit::OnUndo()
 void SDMComponentEdit::GenerateMaterialModelPropertyRows(const TSharedRef<SDMEditor> InEditorWidget, UDynamicMaterialModel* InMaterialModel,
 	TArray<FDMPropertyHandle>& InOutPropertyRows, TSet<UDMMaterialComponent*>& InOutProcessedObjects)
 {
-	const FName Category = FName("Global Settings");
+	const FName MaterialTypeCategory = FName("Material Type");
+	auto AddVariable = [InEditorWidget, &InOutPropertyRows, &MaterialTypeCategory](UObject* InObject, FName InPropertyName)
+		{
+			if (IsValid(InObject))
+			{
+				FDMPropertyHandle& ValueHandle = InOutPropertyRows.Add_GetRef(InEditorWidget->GetPropertyHandle(&*InEditorWidget,
+					InObject, InPropertyName));
+
+				ValueHandle.CategoryOverrideName = MaterialTypeCategory;
+			}
+		};
+
+	if (UDynamicMaterialModelEditorOnlyData* EditorOnlyData = UDynamicMaterialModelEditorOnlyData::Get(InMaterialModel))
+	{
+		AddVariable(EditorOnlyData, GET_MEMBER_NAME_CHECKED(UDynamicMaterialModelEditorOnlyData, ChannelListPreset));
+		AddVariable(EditorOnlyData, GET_MEMBER_NAME_CHECKED(UDynamicMaterialModelEditorOnlyData, Domain));
+		AddVariable(EditorOnlyData, GET_MEMBER_NAME_CHECKED(UDynamicMaterialModelEditorOnlyData, BlendMode));
+		AddVariable(EditorOnlyData, GET_MEMBER_NAME_CHECKED(UDynamicMaterialModelEditorOnlyData, ShadingModel));
+		AddVariable(EditorOnlyData, GET_MEMBER_NAME_CHECKED(UDynamicMaterialModelEditorOnlyData, bPixelAnimationFlag));
+		AddVariable(EditorOnlyData, GET_MEMBER_NAME_CHECKED(UDynamicMaterialModelEditorOnlyData, bTwoSidedFlag));
+	}
+
+	const FName MaterialSettingsCategory = FName("Material Settings");
 	const FName UIMin = FName("UIMin");
 	const FName UIMax = FName("UIMax");
 	const FName ClampMin = FName("ClampMin");
 	const FName ClampMax = FName("ClampMax");
 
-	auto AddGlobalVar = [InEditorWidget, InMaterialModel, &InOutPropertyRows, &InOutProcessedObjects, &Category , &UIMin, &UIMax, &ClampMin, &ClampMax]
+	auto AddGlobalValue = [InEditorWidget, InMaterialModel, &InOutPropertyRows, &InOutProcessedObjects, &MaterialSettingsCategory , &UIMin, &UIMax, &ClampMin, &ClampMax]
 		(UDMMaterialValue* InValue, const FText& InNameOverride)
 		{
 			if (IsValid(InValue))
@@ -850,7 +872,7 @@ void SDMComponentEdit::GenerateMaterialModelPropertyRows(const TSharedRef<SDMEdi
 					FDMPropertyHandle& ValueHandle = InOutPropertyRows.Add_GetRef(InEditorWidget->GetPropertyHandle(&*InEditorWidget,
 						InValue, UDMMaterialValue::ValueName));
 
-					ValueHandle.CategoryOverrideName = Category;
+					ValueHandle.CategoryOverrideName = MaterialSettingsCategory;
 					ValueHandle.NameOverride = InNameOverride;
 
 					if (UDMMaterialValueFloat* FloatValue = Cast<UDMMaterialValueFloat>(InValue))
@@ -869,15 +891,15 @@ void SDMComponentEdit::GenerateMaterialModelPropertyRows(const TSharedRef<SDMEdi
 			}
 		};
 
-	AddGlobalVar(InMaterialModel->GetGlobalParameterValue(UDynamicMaterialModel::GlobalOffsetValueName), LOCTEXT("GlobalOffset", "Global Offset"));
-	AddGlobalVar(InMaterialModel->GetGlobalParameterValue(UDynamicMaterialModel::GlobalTilingValueName), LOCTEXT("GlobalTiling", "Global Tiling"));
-	AddGlobalVar(InMaterialModel->GetGlobalParameterValue(UDynamicMaterialModel::GlobalRotationValueName), LOCTEXT("GlobalRotation", "Global Rotation"));
+	AddGlobalValue(InMaterialModel->GetGlobalParameterValue(UDynamicMaterialModel::GlobalOffsetValueName), LOCTEXT("GlobalOffset", "Global Offset"));
+	AddGlobalValue(InMaterialModel->GetGlobalParameterValue(UDynamicMaterialModel::GlobalTilingValueName), LOCTEXT("GlobalTiling", "Global Tiling"));
+	AddGlobalValue(InMaterialModel->GetGlobalParameterValue(UDynamicMaterialModel::GlobalRotationValueName), LOCTEXT("GlobalRotation", "Global Rotation"));
 
 	if (UDynamicMaterialModelEditorOnlyData* EditorOnlyData = UDynamicMaterialModelEditorOnlyData::Get(InMaterialModel))
 	{
 		if (EditorOnlyData->GetBlendMode() != BLEND_Opaque)
 		{
-			AddGlobalVar(InMaterialModel->GetGlobalParameterValue(UDynamicMaterialModel::GlobalOpacityValueName), LOCTEXT("GlobalOpacity", "Global Opacity"));
+			AddGlobalValue(InMaterialModel->GetGlobalParameterValue(UDynamicMaterialModel::GlobalOpacityValueName), LOCTEXT("GlobalOpacity", "Global Opacity"));
 		}
 
 		const FText PropertyFormat = LOCTEXT("PropertyFormat", "Global {0}");
@@ -897,7 +919,7 @@ void SDMComponentEdit::GenerateMaterialModelPropertyRows(const TSharedRef<SDMEdi
 					{
 						if (UDMMaterialValueFloat1* AlphaValue = Cast<UDMMaterialValueFloat1>(MaterialProperty->GetComponent(UDynamicMaterialModelEditorOnlyData::AlphaValueName)))
 						{
-							AddGlobalVar(
+							AddGlobalValue(
 								AlphaValue,
 								FText::Format(
 									PropertyFormat,
diff --git a/Engine/Plugins/Experimental/DynamicMaterial/Source/DynamicMaterialEditor/Private/Slate/SDMEditor.cpp b/Engine/Plugins/Experimental/DynamicMaterial/Source/DynamicMaterialEditor/Private/Slate/SDMEditor.cpp
--- a/Engine/Plugins/Experimental/DynamicMaterial/Source/DynamicMaterialEditor/Private/Slate/SDMEditor.cpp
+++ b/Engine/Plugins/Experimental/DynamicMaterial/Source/DynamicMaterialEditor/Private/Slate/SDMEditor.cpp
@@ -866,6 +866,9 @@ void SDMEditor::Tick(const FGeometry& AllottedGeometry, const double InCurrentTi
 
 void SDMEditor::OnMaterialBuilt(UDynamicMaterialModel* InMaterialModel)
 {
+	RefreshSlotPickerList();
+	RefreshSlotWidget();
+	RefreshComponentEditWidget();
 }
 
 void SDMEditor::OnValuesUpdated(UDynamicMaterialModel* InMaterialModel)
diff --git a/Engine/Plugins/Experimental/DynamicMaterial/Source/DynamicMaterialEditor/Public/Model/DynamicMaterialModelEditorOnlyData.h b/Engine/Plugins/Experimental/DynamicMaterial/Source/DynamicMaterialEditor/Public/Model/DynamicMaterialModelEditorOnlyData.h
--- a/Engine/Plugins/Experimental/DynamicMaterial/Source/DynamicMaterialEditor/Public/Model/DynamicMaterialModelEditorOnlyData.h
+++ b/Engine/Plugins/Experimental/DynamicMaterial/Source/DynamicMaterialEditor/Public/Model/DynamicMaterialModelEditorOnlyData.h
@@ -45,6 +45,7 @@ class DYNAMICMATERIALEDITOR_API UDynamicMaterialModelEditorOnlyData : public UOb
 	GENERATED_BODY()
 
 	friend class UDynamicMaterialModelFactory;
+	friend class SDMComponentEdit;
 
 public:
 	static const FString SlotsPathToken;
@@ -217,24 +218,26 @@ protected:
 	TObjectPtr<UDynamicMaterialModel> MaterialModel;
 
 	UPROPERTY(VisibleInstanceOnly, BlueprintReadOnly, Transient, DuplicateTransient, TextExportTransient, Category = "Material Designer")
-	EDMState State = EDMState::Idle;
+	EDMState State;
 
-	UPROPERTY(EditInstanceOnly, BlueprintReadOnly, Category = "Material Designer")
-	TEnumAsByte<EMaterialDomain> Domain = EMaterialDomain::MD_Surface;
+	UPROPERTY(EditInstanceOnly, BlueprintReadOnly, Category = "Material Designer", 
+		meta = (ValidEnumValues = "MD_Surface,MD_PostProcess,MD_DeferredDecal"))
+	TEnumAsByte<EMaterialDomain> Domain;
 
-	UPROPERTY(EditInstanceOnly, BlueprintReadOnly, Category = "Material Designer")
-	TEnumAsByte<EBlendMode> BlendMode = EBlendMode::BLEND_Translucent;
+	UPROPERTY(EditInstanceOnly, BlueprintReadOnly, Category = "Material Designer",
+		meta = (ValidEnumValues = "BLEND_Opaque,BLEND_Translucent,BLEND_Masked,BLEND_Additive,BLEND_Modulate"))
+	TEnumAsByte<EBlendMode> BlendMode;
 
-	UPROPERTY(EditInstanceOnly, BlueprintReadOnly, Category = "Material Designer")
-	EDMMaterialShadingModel ShadingModel = EDMMaterialShadingModel::Unlit;
+	UPROPERTY(EditInstanceOnly, BlueprintReadOnly, Category = "Material Designer", meta = (ValidEnumValues = "Unlit,DefaultLit"))
+	EDMMaterialShadingModel ShadingModel;
 
-	UPROPERTY(VisibleInstanceOnly, BlueprintReadOnly, Category = "Material Designer")
+	UPROPERTY(EditInstanceOnly, BlueprintReadOnly, Category = "Material Designer")
 	bool bPixelAnimationFlag;
 
-	UPROPERTY(VisibleInstanceOnly, BlueprintReadOnly, Category = "Material Designer")
+	UPROPERTY(EditInstanceOnly, BlueprintReadOnly, Category = "Material Designer")
 	bool bTwoSidedFlag;
 
-	UPROPERTY(VisibleInstanceOnly, BlueprintReadOnly, Category = "Material Designer")
+	UPROPERTY(EditInstanceOnly, BlueprintReadOnly, Category = "Material Designer", meta = (GetOptions = GetPresetOptions, NoResetToDefault))
 	FName ChannelListPreset;
 
 	UPROPERTY(VisibleInstanceOnly, BlueprintReadOnly, Instanced, Category = "Material Designer")
@@ -285,6 +288,9 @@ protected:
 
 	void AssignPropertyAlphaValues();
 
+	UFUNCTION()
+	TArray<FName> GetPresetOptions() const;
+
 	void LoadDeprecatedModelData_Base(bool bInCreateMaterialPackage, EBlendMode InBlendMode, EDMMaterialShadingModel InShadingModel);
 	void LoadDeprecatedModelData_Expressions(TArray<TObjectPtr<UMaterialExpression>>& InExpressions);
 	void LoadDeprecatedModelData_Properties(TMap<EDMMaterialPropertyType, TObjectPtr<UObject>>& InProperties);
