Material Designer
- AllowPreserveRatio added to Float2 and Float3XYZ values.

diff --git a/Engine/Plugins/Experimental/DynamicMaterial/Source/DynamicMaterial/Public/Components/MaterialValues/DMMaterialValueFloat2.h b/Engine/Plugins/Experimental/DynamicMaterial/Source/DynamicMaterial/Public/Components/MaterialValues/DMMaterialValueFloat2.h
--- a/Engine/Plugins/Experimental/DynamicMaterial/Source/DynamicMaterial/Public/Components/MaterialValues/DMMaterialValueFloat2.h
+++ b/Engine/Plugins/Experimental/DynamicMaterial/Source/DynamicMaterial/Public/Components/MaterialValues/DMMaterialValueFloat2.h
@@ -45,7 +45,7 @@ public:
  
 protected: 
 	UPROPERTY(EditInstanceOnly, BlueprintReadWrite, Getter = GetValue, Setter = SetValue, BlueprintSetter = SetValue, Category = "Material Designer",
-		meta = (AllowPrivateAccess = "true", DisplayName = "Vector"))
+		meta = (AllowPrivateAccess = "true", DisplayName = "Vector", AllowPreserveRatio))
 	FVector2D Value;
  
 #if WITH_EDITORONLY_DATA
diff --git a/Engine/Plugins/Experimental/DynamicMaterial/Source/DynamicMaterial/Public/Components/MaterialValues/DMMaterialValueFloat3XYZ.h b/Engine/Plugins/Experimental/DynamicMaterial/Source/DynamicMaterial/Public/Components/MaterialValues/DMMaterialValueFloat3XYZ.h
--- a/Engine/Plugins/Experimental/DynamicMaterial/Source/DynamicMaterial/Public/Components/MaterialValues/DMMaterialValueFloat3XYZ.h
+++ b/Engine/Plugins/Experimental/DynamicMaterial/Source/DynamicMaterial/Public/Components/MaterialValues/DMMaterialValueFloat3XYZ.h
@@ -45,7 +45,7 @@ public:
  
 protected: 
 	UPROPERTY(EditInstanceOnly, BlueprintReadWrite, Getter = GetValue, Setter = SetValue, BlueprintSetter = SetValue, Category = "Material Designer",
-		meta = (AllowPrivateAccess = "true"), DisplayName = "Vector")
+		meta = (AllowPrivateAccess = "true", DisplayName = "Vector", AllowPreserveRatio))
 	FVector Value;
  
 #if WITH_EDITORONLY_DATA
