Material Designer
- Localisation key conflict

diff --git a/Engine/Plugins/Experimental/DynamicMaterial/Source/DynamicMaterialEditor/Private/Slate/SDMEditor.cpp b/Engine/Plugins/Experimental/DynamicMaterial/Source/DynamicMaterialEditor/Private/Slate/SDMEditor.cpp
--- a/Engine/Plugins/Experimental/DynamicMaterial/Source/DynamicMaterialEditor/Private/Slate/SDMEditor.cpp
+++ b/Engine/Plugins/Experimental/DynamicMaterial/Source/DynamicMaterialEditor/Private/Slate/SDMEditor.cpp
@@ -684,7 +684,7 @@ TSharedRef<SWidget> SDMEditor::CreateSlotPickerWidget()
 			.IsChecked(this, &SDMEditor::GetSlotCheckState, EDMMaterialPropertyType::None)
 			.OnCheckStateChanged(this, &SDMEditor::OnSlotCheckStateChanged, EDMMaterialPropertyType::None)
 			.Padding(FVector2D(5.f, 3.f))
-			.ToolTipText(LOCTEXT("GlobalSettings", "Global Material Settings"))
+			.ToolTipText(LOCTEXT("GlobalSettingsToolTip", "Global Material Settings"))
 			.Content()
 			[
 				SNew(STextBlock)
