Material Designer
- Fix for changing material type settings in the global settings section.
- Added tooltip to, and renamed, channel list preset so it's more clear what it does.
- Fixed a crash when changing material properties.

diff --git a/Engine/Plugins/Experimental/DynamicMaterial/Source/DynamicMaterialEditor/Private/Model/DynamicMaterialModelEditorOnlyData.cpp b/Engine/Plugins/Experimental/DynamicMaterial/Source/DynamicMaterialEditor/Private/Model/DynamicMaterialModelEditorOnlyData.cpp
--- a/Engine/Plugins/Experimental/DynamicMaterial/Source/DynamicMaterialEditor/Private/Model/DynamicMaterialModelEditorOnlyData.cpp
+++ b/Engine/Plugins/Experimental/DynamicMaterial/Source/DynamicMaterialEditor/Private/Model/DynamicMaterialModelEditorOnlyData.cpp
@@ -205,6 +205,193 @@
 	return PresetNames;
 }
 
+void UDynamicMaterialModelEditorOnlyData::OnChannelListPresetChanged()
+{
+	const FDMMaterialChannelListPreset* Preset = GetDefault<UDynamicMaterialEditorSettings>()->GetPresetByName(ChannelListPreset);
+
+	if (!Preset)
+	{
+		return;
+	}
+
+	for (uint8 PropertyIndex = static_cast<uint8>(EDMMaterialPropertyType::None) + 1;
+		PropertyIndex < static_cast<uint8>(EDMMaterialPropertyType::Any);
+		++PropertyIndex)
+	{
+		const EDMMaterialPropertyType Property = static_cast<EDMMaterialPropertyType>(PropertyIndex);
+
+		if (Property == EDMMaterialPropertyType::EmissiveColor || Property == EDMMaterialPropertyType::OpacityMask)
+		{
+			continue;
+		}
+
+		if (Preset->IsPropertyEnabled(Property))
+		{
+			AddSlotForMaterialProperty(Property);
+		}
+		else
+		{
+			RemoveSlotForMaterialProperty(Property);
+		}
+	}
+
+	SetBlendMode(Preset->DefaultBlendMode);
+	SetShadingModel(Preset->DefaultShadingModel);
+	SetPixelAnimationFlag(Preset->bDefaultAnimated);
+	SetTwoSidedFlag(Preset->bDefaultTwoSided);
+}
+
+void UDynamicMaterialModelEditorOnlyData::OnDomainChanged()
+{
+	if (Domain == EMaterialDomain::MD_PostProcess)
+	{
+		const FDMUpdateGuard Guard;
+
+		// Post process only supports emissive.
+		for (uint8 PropertyIndex = static_cast<uint8>(EDMMaterialPropertyType::None) + 1;
+			PropertyIndex < static_cast<uint8>(EDMMaterialPropertyType::Any);
+			++PropertyIndex)
+		{
+			const EDMMaterialPropertyType Property = static_cast<EDMMaterialPropertyType>(PropertyIndex);
+
+			switch (Property)
+			{
+				case EDMMaterialPropertyType::BaseColor:
+				case EDMMaterialPropertyType::EmissiveColor:
+					RemoveSlotForMaterialProperty(Property);
+					continue;
+
+				default:
+					// Do nothing
+					break;
+			}
+		}
+
+		EnsureSwapSlotMaterialProperty(EDMMaterialPropertyType::BaseColor, EDMMaterialPropertyType::EmissiveColor);
+
+		SetShadingModel(EDMMaterialShadingModel::Unlit);
+		SetBlendMode(EBlendMode::BLEND_Opaque);
+	}
+	else if (Domain == EMaterialDomain::MD_DeferredDecal)
+	{
+		// Post process only supports basic types.
+		for (uint8 PropertyIndex = static_cast<uint8>(EDMMaterialPropertyType::None) + 1;
+			PropertyIndex < static_cast<uint8>(EDMMaterialPropertyType::Any);
+			++PropertyIndex)
+		{
+			const EDMMaterialPropertyType Property = static_cast<EDMMaterialPropertyType>(PropertyIndex);
+
+			switch (Property)
+			{
+				case EDMMaterialPropertyType::BaseColor:
+				case EDMMaterialPropertyType::EmissiveColor:
+				case EDMMaterialPropertyType::Opacity:
+				case EDMMaterialPropertyType::OpacityMask:
+					RemoveSlotForMaterialProperty(Property);
+					continue;
+
+				default:
+					// Do nothing
+					break;
+			}
+		}
+
+		EnsureSwapSlotMaterialProperty(EDMMaterialPropertyType::EmissiveColor, EDMMaterialPropertyType::BaseColor);
+
+		SetShadingModel(EDMMaterialShadingModel::DefaultLit);
+		SetBlendMode(EBlendMode::BLEND_Translucent);
+	}
+
+	RequestMaterialBuild();
+}
+
+void UDynamicMaterialModelEditorOnlyData::OnBlendModeChanged()
+{
+	switch (BlendMode)
+	{
+		case EBlendMode::BLEND_Opaque:
+			SetPixelAnimationFlag(false);
+			RemoveSlotForMaterialProperty(EDMMaterialPropertyType::Opacity);
+			RemoveSlotForMaterialProperty(EDMMaterialPropertyType::OpacityMask);
+			break;
+
+		case EBlendMode::BLEND_Masked:
+			SetPixelAnimationFlag(false);
+			EnsureSwapSlotMaterialProperty(EDMMaterialPropertyType::Opacity, EDMMaterialPropertyType::OpacityMask);
+			break;
+
+		case EBlendMode::BLEND_Translucent:
+		case EBlendMode::BLEND_Additive:
+		case EBlendMode::BLEND_Modulate:
+			EnsureSwapSlotMaterialProperty(EDMMaterialPropertyType::OpacityMask, EDMMaterialPropertyType::Opacity);
+			break;
+	}
+
+	RequestMaterialBuild();
+}
+
+void UDynamicMaterialModelEditorOnlyData::OnShadingModelChanged()
+{
+	EDMMaterialPropertyType FromProperty;
+	EDMMaterialPropertyType ToProperty;
+
+	switch (ShadingModel)
+	{
+		case EDMMaterialShadingModel::Unlit:
+			FromProperty = EDMMaterialPropertyType::BaseColor;
+			ToProperty = EDMMaterialPropertyType::EmissiveColor;
+			break;
+
+		case EDMMaterialShadingModel::DefaultLit:
+			FromProperty = EDMMaterialPropertyType::EmissiveColor;
+			ToProperty = EDMMaterialPropertyType::BaseColor;
+			break;
+
+		default:
+			return;
+	}
+
+	EnsureSwapSlotMaterialProperty(FromProperty, ToProperty);
+
+	RequestMaterialBuild();
+}
+
+void UDynamicMaterialModelEditorOnlyData::OnPixelAnimationFlagChanged()
+{
+	if (MaterialModel)
+	{
+		if (UMaterial* Material = MaterialModel->GetGeneratedMaterial())
+		{
+			if (GUndo)
+			{
+				Material->Modify();
+			}
+
+			Material->bHasPixelAnimation = bPixelAnimationFlag;
+		}
+	}
+
+	RequestMaterialBuild();
+}
+
+void UDynamicMaterialModelEditorOnlyData::OnTwoSidedFlagChanged()
+{
+	if (MaterialModel)
+	{
+		if (UMaterial* Material = MaterialModel->GetGeneratedMaterial())
+		{
+			if (GUndo)
+			{
+				Material->Modify();
+			}
+
+			Material->TwoSided = bTwoSidedFlag;
+		}
+	}
+
+	RequestMaterialBuild();
+}
+
 void UDynamicMaterialModelEditorOnlyData::Initialize()
 {
 	if (!Slots.IsEmpty())
@@ -782,38 +969,7 @@
 {
 	ChannelListPreset = InPresetName;
 
-	const FDMMaterialChannelListPreset* Preset = GetDefault<UDynamicMaterialEditorSettings>()->GetPresetByName(InPresetName);
-
-	if (!Preset)
-	{
-		return;
-	}
-
-	for (uint8 PropertyIndex = static_cast<uint8>(EDMMaterialPropertyType::None) + 1;
-		PropertyIndex < static_cast<uint8>(EDMMaterialPropertyType::Any);
-		++PropertyIndex)
-	{
-		const EDMMaterialPropertyType Property = static_cast<EDMMaterialPropertyType>(PropertyIndex);
-
-		if (Property == EDMMaterialPropertyType::EmissiveColor || Property == EDMMaterialPropertyType::OpacityMask)
-		{
-			continue;
-		}
-
-		if (Preset->IsPropertyEnabled(Property))
-		{
-			AddSlotForMaterialProperty(Property);
-		}
-		else
-		{
-			RemoveSlotForMaterialProperty(Property);
-		}
-	}
-
-	SetBlendMode(Preset->DefaultBlendMode);
-	SetShadingModel(Preset->DefaultShadingModel);
-	SetPixelAnimationFlag(Preset->bDefaultAnimated);
-	SetTwoSidedFlag(Preset->bDefaultTwoSided);
+	OnChannelListPresetChanged();
 }
 
 UDMMaterialComponent* UDynamicMaterialModelEditorOnlyData::GetSubComponentByPath(FDMComponentPath& InPath,
@@ -951,66 +1107,7 @@
 
 	Domain = InDomain;
 
-	if (Domain == EMaterialDomain::MD_PostProcess)
-	{
-		const FDMUpdateGuard Guard;
-
-		// Post process only supports emissive.
-		for (uint8 PropertyIndex = static_cast<uint8>(EDMMaterialPropertyType::None) + 1;
-			PropertyIndex < static_cast<uint8>(EDMMaterialPropertyType::Any);
-			++PropertyIndex)
-		{
-			const EDMMaterialPropertyType Property = static_cast<EDMMaterialPropertyType>(PropertyIndex);
-
-			switch (Property)
-			{
-				case EDMMaterialPropertyType::BaseColor:
-				case EDMMaterialPropertyType::EmissiveColor:
-					RemoveSlotForMaterialProperty(Property);
-					continue;
-
-				default:
-					// Do nothing
-					break;
-			}			
-		}
-
-		EnsureSwapSlotMaterialProperty(EDMMaterialPropertyType::BaseColor, EDMMaterialPropertyType::EmissiveColor);
-
-		SetShadingModel(EDMMaterialShadingModel::Unlit);
-		SetBlendMode(EBlendMode::BLEND_Opaque);
-	}
-	else if (Domain == EMaterialDomain::MD_DeferredDecal)
-	{
-		// Post process only supports basic types.
-		for (uint8 PropertyIndex = static_cast<uint8>(EDMMaterialPropertyType::None) + 1;
-			PropertyIndex < static_cast<uint8>(EDMMaterialPropertyType::Any);
-			++PropertyIndex)
-		{
-			const EDMMaterialPropertyType Property = static_cast<EDMMaterialPropertyType>(PropertyIndex);
-
-			switch (Property)
-			{
-				case EDMMaterialPropertyType::BaseColor:
-				case EDMMaterialPropertyType::EmissiveColor:
-				case EDMMaterialPropertyType::Opacity:
-				case EDMMaterialPropertyType::OpacityMask:
-					RemoveSlotForMaterialProperty(Property);
-					continue;
-
-				default:
-					// Do nothing
-					break;
-			}			
-		}
-
-		EnsureSwapSlotMaterialProperty(EDMMaterialPropertyType::EmissiveColor, EDMMaterialPropertyType::BaseColor);
-
-		SetShadingModel(EDMMaterialShadingModel::DefaultLit);
-		SetBlendMode(EBlendMode::BLEND_Translucent);
-	}
-
-	RequestMaterialBuild();
+	OnDomainChanged();
 }
 
 void UDynamicMaterialModelEditorOnlyData::SetBlendMode(TEnumAsByte<EBlendMode> InBlendMode)
@@ -1022,27 +1119,7 @@
 
 	BlendMode = InBlendMode;
 
-	switch (InBlendMode)
-	{
-		case EBlendMode::BLEND_Opaque:
-			SetPixelAnimationFlag(false);
-			RemoveSlotForMaterialProperty(EDMMaterialPropertyType::Opacity);
-			RemoveSlotForMaterialProperty(EDMMaterialPropertyType::OpacityMask);
-			break;
-
-		case EBlendMode::BLEND_Masked:
-			SetPixelAnimationFlag(false);
-			EnsureSwapSlotMaterialProperty(EDMMaterialPropertyType::Opacity, EDMMaterialPropertyType::OpacityMask);
-			break;
-
-		case EBlendMode::BLEND_Translucent:
-		case EBlendMode::BLEND_Additive:
-		case EBlendMode::BLEND_Modulate:
-			EnsureSwapSlotMaterialProperty(EDMMaterialPropertyType::OpacityMask, EDMMaterialPropertyType::Opacity);
-			break;
-	}
-
-	RequestMaterialBuild();
+	OnBlendModeChanged();
 }
 
 void UDynamicMaterialModelEditorOnlyData::SetShadingModel(EDMMaterialShadingModel InShadingModel)
@@ -1054,28 +1131,7 @@
 
 	ShadingModel = InShadingModel;
 
-	EDMMaterialPropertyType FromProperty;
-	EDMMaterialPropertyType ToProperty;
-
-	switch (ShadingModel)
-	{
-		case EDMMaterialShadingModel::Unlit:
-			FromProperty = EDMMaterialPropertyType::BaseColor;
-			ToProperty = EDMMaterialPropertyType::EmissiveColor;
-			break;
-
-		case EDMMaterialShadingModel::DefaultLit:
-			FromProperty = EDMMaterialPropertyType::EmissiveColor;
-			ToProperty = EDMMaterialPropertyType::BaseColor;
-			break;
-
-		default:
-			return;
-	}
-
-	EnsureSwapSlotMaterialProperty(FromProperty, ToProperty);
-
-	RequestMaterialBuild();
+	OnShadingModelChanged();
 }
 
 bool UDynamicMaterialModelEditorOnlyData::IsPixelAnimationFlagSet() const
@@ -1092,20 +1148,7 @@
 
 	bPixelAnimationFlag = bInFlagValue;
 
-	if (MaterialModel)
-	{
-		if (UMaterial* Material = MaterialModel->GetGeneratedMaterial())
-		{
-			if (GUndo)
-			{
-				Material->Modify();
-			}
-
-			Material->bHasPixelAnimation = bPixelAnimationFlag;
-		}
-	}
-
-	RequestMaterialBuild();
+	OnPixelAnimationFlagChanged();
 }
 
 bool UDynamicMaterialModelEditorOnlyData::IsTwoSidedFlagSet() const
@@ -1122,20 +1165,7 @@
 
 	bTwoSidedFlag = bInFlagValue;
 
-	if (MaterialModel)
-	{
-		if (UMaterial* Material = MaterialModel->GetGeneratedMaterial())
-		{
-			if (GUndo)
-			{
-				Material->Modify();
-			}
-
-			Material->TwoSided = bTwoSidedFlag;
-		}
-	}
-
-	RequestMaterialBuild();
+	OnTwoSidedFlagChanged();
 }
 
 FName UDynamicMaterialModelEditorOnlyData::GetChannelListPreset() const
@@ -1507,6 +1537,38 @@
 	RequestMaterialBuild();
 }
 
+void UDynamicMaterialModelEditorOnlyData::PostEditChangeChainProperty(FPropertyChangedChainEvent& PropertyChangedEvent)
+{
+	UObject::PostEditChangeChainProperty(PropertyChangedEvent);
+
+	const FName Property = PropertyChangedEvent.GetMemberPropertyName();
+
+	if (Property == GET_MEMBER_NAME_CHECKED(ThisClass, ChannelListPreset))
+	{
+		OnChannelListPresetChanged();
+	}
+	else if (Property == GET_MEMBER_NAME_CHECKED(ThisClass, Domain))
+	{
+		OnDomainChanged();
+	}
+	else if (Property == GET_MEMBER_NAME_CHECKED(ThisClass, BlendMode))
+	{
+		OnBlendModeChanged();
+	}
+	else if (Property == GET_MEMBER_NAME_CHECKED(ThisClass, ShadingModel))
+	{
+		OnShadingModelChanged();
+	}
+	else if (Property == GET_MEMBER_NAME_CHECKED(ThisClass, bPixelAnimationFlag))
+	{
+		OnPixelAnimationFlagChanged();
+	}
+	else if (Property == GET_MEMBER_NAME_CHECKED(ThisClass, bTwoSidedFlag))
+	{
+		OnTwoSidedFlagChanged();
+	}
+}
+
 void UDynamicMaterialModelEditorOnlyData::OnValueUpdated(UDMMaterialValue* InValue, EDMUpdateType InUpdateType)
 {
 	check(InValue);
diff --git a/Engine/Plugins/Experimental/DynamicMaterial/Source/DynamicMaterialEditor/Private/Slate/SDMComponentEdit.cpp b/Engine/Plugins/Experimental/DynamicMaterial/Source/DynamicMaterialEditor/Private/Slate/SDMComponentEdit.cpp
--- a/Engine/Plugins/Experimental/DynamicMaterial/Source/DynamicMaterialEditor/Private/Slate/SDMComponentEdit.cpp
+++ b/Engine/Plugins/Experimental/DynamicMaterial/Source/DynamicMaterialEditor/Private/Slate/SDMComponentEdit.cpp
@@ -320,7 +320,7 @@ TSharedRef<SWidget> SDMComponentEdit::CreateEditWidget()
 
 	TArray<FDMPropertyHandle> EditRows = GetEditRows();
 
-	for (const FDMPropertyHandle& EditRow : EditRows)
+	for (FDMPropertyHandle& EditRow : EditRows)
 	{
 		const bool bHasValidCustomWidget = EditRow.ValueWidget.IsValid() && !EditRow.ValueName.IsNone() && EditRow.NameOverride.IsSet();
 
@@ -329,15 +329,20 @@ TSharedRef<SWidget> SDMComponentEdit::CreateEditWidget()
 			continue;
 		}
 
+		if (!EditRow.PropertyHandle.IsValid() && EditRow.DetailTreeNode)
+		{
+			EditRow.PropertyHandle = EditRow.DetailTreeNode->CreatePropertyHandle();
+		}
+
 		ECustomDetailsTreeInsertPosition Position = ECustomDetailsTreeInsertPosition::Child;
 
-		if (EditRow.DetailTreeNode)
+		if (EditRow.PropertyHandle.IsValid())
 		{
-			if (EditRow.DetailTreeNode->CreatePropertyHandle()->HasMetaData("HighPriority"))
+			if (EditRow.PropertyHandle->HasMetaData("HighPriority"))
 			{
 				Position = ECustomDetailsTreeInsertPosition::FirstChild;
 			}
-			else if (EditRow.DetailTreeNode->CreatePropertyHandle()->HasMetaData("LowPriority"))
+			else if (EditRow.PropertyHandle->HasMetaData("LowPriority"))
 			{
 				Position = ECustomDetailsTreeInsertPosition::LastChild;
 			}
@@ -422,7 +427,7 @@ TSharedRef<SWidget> SDMComponentEdit::CreateEditWidget()
 			);
 		}
 
-		if (EditRow.DetailTreeNode->CreatePropertyHandle()->HasMetaData("NotKeyframeable"))
+		if (EditRow.PropertyHandle.IsValid() && EditRow.PropertyHandle->HasMetaData("NotKeyframeable"))
 		{
 			Item->SetKeyframeEnabled(false);
 		}
diff --git a/Engine/Plugins/Experimental/DynamicMaterial/Source/DynamicMaterialEditor/Public/Model/DynamicMaterialModelEditorOnlyData.h b/Engine/Plugins/Experimental/DynamicMaterial/Source/DynamicMaterialEditor/Public/Model/DynamicMaterialModelEditorOnlyData.h
--- a/Engine/Plugins/Experimental/DynamicMaterial/Source/DynamicMaterialEditor/Public/Model/DynamicMaterialModelEditorOnlyData.h
+++ b/Engine/Plugins/Experimental/DynamicMaterial/Source/DynamicMaterialEditor/Public/Model/DynamicMaterialModelEditorOnlyData.h
@@ -193,6 +193,7 @@ public:
 	virtual void PostEditUndo() override;
 	virtual void PostEditImport() override;
 	virtual void PostDuplicate(bool bDuplicateForPIE) override;
+	virtual void PostEditChangeChainProperty(FPropertyChangedChainEvent& PropertyChangedEvent) override;
 	//~ End UObject
 
 	void SaveEditor();
@@ -237,7 +238,9 @@ protected:
 	UPROPERTY(EditInstanceOnly, BlueprintReadOnly, Category = "Material Designer")
 	bool bTwoSidedFlag;
 
-	UPROPERTY(EditInstanceOnly, BlueprintReadOnly, Category = "Material Designer", meta = (GetOptions = GetPresetOptions, NoResetToDefault))
+	UPROPERTY(EditInstanceOnly, BlueprintReadOnly, Category = "Material Designer", 
+		meta = (GetOptions = GetPresetOptions, NoResetToDefault, DisplayName = "Material Type Preset",
+			ToolTip = "Sets the available channels and default material properties."))
 	FName ChannelListPreset;
 
 	UPROPERTY(VisibleInstanceOnly, BlueprintReadOnly, Instanced, Category = "Material Designer")
@@ -291,6 +294,13 @@ protected:
 	UFUNCTION()
 	TArray<FName> GetPresetOptions() const;
 
+	void OnChannelListPresetChanged();
+	void OnDomainChanged();
+	void OnBlendModeChanged();
+	void OnShadingModelChanged();
+	void OnPixelAnimationFlagChanged();
+	void OnTwoSidedFlagChanged();
+
 	void LoadDeprecatedModelData_Base(bool bInCreateMaterialPackage, EBlendMode InBlendMode, EDMMaterialShadingModel InShadingModel);
 	void LoadDeprecatedModelData_Expressions(TArray<TObjectPtr<UMaterialExpression>>& InExpressions);
 	void LoadDeprecatedModelData_Properties(TMap<EDMMaterialPropertyType, TObjectPtr<UObject>>& InProperties);
