Material Designer
- Replaced some boilerplate code with a 'for each material property type' function

diff --git a/Engine/Plugins/Experimental/DynamicMaterial/Source/DynamicMaterialEditor/Private/Model/DynamicMaterialModelEditorOnlyData.cpp b/Engine/Plugins/Experimental/DynamicMaterial/Source/DynamicMaterialEditor/Private/Model/DynamicMaterialModelEditorOnlyData.cpp
--- a/Engine/Plugins/Experimental/DynamicMaterial/Source/DynamicMaterialEditor/Private/Model/DynamicMaterialModelEditorOnlyData.cpp
+++ b/Engine/Plugins/Experimental/DynamicMaterial/Source/DynamicMaterialEditor/Private/Model/DynamicMaterialModelEditorOnlyData.cpp
@@ -214,26 +214,25 @@ void UDynamicMaterialModelEditorOnlyData::OnChannelListPresetChanged()
 		return;
 	}
 
-	for (uint8 PropertyIndex = static_cast<uint8>(EDMMaterialPropertyType::None) + 1;
-		PropertyIndex < static_cast<uint8>(EDMMaterialPropertyType::Any);
-		++PropertyIndex)
-	{
-		const EDMMaterialPropertyType Property = static_cast<EDMMaterialPropertyType>(PropertyIndex);
-
-		if (Property == EDMMaterialPropertyType::EmissiveColor || Property == EDMMaterialPropertyType::OpacityMask)
+	ForEachMaterialPropertyType(
+		[this, Preset](EDMMaterialPropertyType InProperty)
 		{
-			continue;
-		}
+			if (InProperty == EDMMaterialPropertyType::EmissiveColor || InProperty == EDMMaterialPropertyType::OpacityMask)
+			{
+				return EDMIterationResult::Continue;
+			}
 
-		if (Preset->IsPropertyEnabled(Property))
-		{
-			AddSlotForMaterialProperty(Property);
-		}
-		else
-		{
-			RemoveSlotForMaterialProperty(Property);
-		}
-	}
+			if (Preset->IsPropertyEnabled(InProperty))
+			{
+				AddSlotForMaterialProperty(InProperty);
+			}
+			else
+			{
+				RemoveSlotForMaterialProperty(InProperty);
+			}
+
+			return EDMIterationResult::Continue;
+		});
 
 	SetBlendMode(Preset->DefaultBlendMode);
 	SetShadingModel(Preset->DefaultShadingModel);
@@ -248,24 +247,23 @@ void UDynamicMaterialModelEditorOnlyData::OnDomainChanged()
 		const FDMUpdateGuard Guard;
 
 		// Post process only supports emissive.
-		for (uint8 PropertyIndex = static_cast<uint8>(EDMMaterialPropertyType::None) + 1;
-			PropertyIndex < static_cast<uint8>(EDMMaterialPropertyType::Any);
-			++PropertyIndex)
-		{
-			const EDMMaterialPropertyType Property = static_cast<EDMMaterialPropertyType>(PropertyIndex);
-
-			switch (Property)
+		ForEachMaterialPropertyType(
+			[this](EDMMaterialPropertyType InProperty)
 			{
-				case EDMMaterialPropertyType::BaseColor:
-				case EDMMaterialPropertyType::EmissiveColor:
-					RemoveSlotForMaterialProperty(Property);
-					continue;
+				switch (InProperty)
+				{
+					case EDMMaterialPropertyType::BaseColor:
+					case EDMMaterialPropertyType::EmissiveColor:
+						RemoveSlotForMaterialProperty(InProperty);
+						break;
+
+					default:
+						// Do nothing
+						break;
+				}
 
-				default:
-					// Do nothing
-					break;
-			}
-		}
+				return EDMIterationResult::Continue;
+			});
 
 		EnsureSwapSlotMaterialProperty(EDMMaterialPropertyType::BaseColor, EDMMaterialPropertyType::EmissiveColor);
 
@@ -275,26 +273,25 @@ void UDynamicMaterialModelEditorOnlyData::OnDomainChanged()
 	else if (Domain == EMaterialDomain::MD_DeferredDecal)
 	{
 		// Post process only supports basic types.
-		for (uint8 PropertyIndex = static_cast<uint8>(EDMMaterialPropertyType::None) + 1;
-			PropertyIndex < static_cast<uint8>(EDMMaterialPropertyType::Any);
-			++PropertyIndex)
-		{
-			const EDMMaterialPropertyType Property = static_cast<EDMMaterialPropertyType>(PropertyIndex);
-
-			switch (Property)
+		ForEachMaterialPropertyType(
+			[this](EDMMaterialPropertyType InProperty)
 			{
-				case EDMMaterialPropertyType::BaseColor:
-				case EDMMaterialPropertyType::EmissiveColor:
-				case EDMMaterialPropertyType::Opacity:
-				case EDMMaterialPropertyType::OpacityMask:
-					RemoveSlotForMaterialProperty(Property);
-					continue;
+				switch (InProperty)
+				{
+					case EDMMaterialPropertyType::BaseColor:
+					case EDMMaterialPropertyType::EmissiveColor:
+					case EDMMaterialPropertyType::Opacity:
+					case EDMMaterialPropertyType::OpacityMask:
+						RemoveSlotForMaterialProperty(InProperty);
+						break;
+
+					default:
+						// Do nothing
+						break;
+				}
 
-				default:
-					// Do nothing
-					break;
-			}
-		}
+				return EDMIterationResult::Continue;
+			});
 
 		EnsureSwapSlotMaterialProperty(EDMMaterialPropertyType::EmissiveColor, EDMMaterialPropertyType::BaseColor);
 
@@ -965,6 +962,20 @@ void UDynamicMaterialModelEditorOnlyData::OnWizardComplete()
 	}
 }
 
+void UDynamicMaterialModelEditorOnlyData::ForEachMaterialPropertyType(TFunctionRef<EDMIterationResult(EDMMaterialPropertyType InType)> InCallable,
+	EDMMaterialPropertyType InStart, EDMMaterialPropertyType InEnd)
+{
+	for (uint8 PropertyIndex = static_cast<uint8>(InStart); PropertyIndex <= static_cast<uint8>(InEnd); ++PropertyIndex)
+	{
+		const EDMMaterialPropertyType Property = static_cast<EDMMaterialPropertyType>(PropertyIndex);
+
+		if (InCallable(Property) == EDMIterationResult::Break)
+		{
+			break;
+		}
+	}
+}
+
 void UDynamicMaterialModelEditorOnlyData::SetChannelListPreset(FName InPresetName)
 {
 	ChannelListPreset = InPresetName;
@@ -1231,22 +1242,21 @@ UDMMaterialSlot* UDynamicMaterialModelEditorOnlyData::GetSlotForMaterialProperty
 
 UDMMaterialSlot* UDynamicMaterialModelEditorOnlyData::AddSlot()
 {
-	constexpr uint8 Start = static_cast<uint8>(EDMMaterialPropertyType::None) + 1;
-	constexpr uint8 End = static_cast<uint8>(EDMMaterialPropertyType::Any);
+	UDMMaterialSlot* NewSlot = nullptr;
 
-	for (uint8 PropertyIndex = Start; PropertyIndex < End; ++PropertyIndex)
-	{
-		const EDMMaterialPropertyType Property = static_cast<EDMMaterialPropertyType>(PropertyIndex);
-
-		if (GetSlotForMaterialProperty(Property))
+	ForEachMaterialPropertyType(
+		[this, &NewSlot](EDMMaterialPropertyType InProperty)
 		{
-			continue;
-		}
+			if (GetSlotForMaterialProperty(InProperty))
+			{
+				return EDMIterationResult::Continue;
+			}
 
-		return AddSlotForMaterialProperty(Property);
-	}
+			NewSlot = AddSlotForMaterialProperty(InProperty);
+			return EDMIterationResult::Break;
+		});
 
-	return nullptr;
+	return NewSlot;
 }
 
 UDMMaterialSlot* UDynamicMaterialModelEditorOnlyData::AddSlotForMaterialProperty(EDMMaterialPropertyType InType)
diff --git a/Engine/Plugins/Experimental/DynamicMaterial/Source/DynamicMaterialEditor/Private/Slate/SDMComponentEdit.cpp b/Engine/Plugins/Experimental/DynamicMaterial/Source/DynamicMaterialEditor/Private/Slate/SDMComponentEdit.cpp
--- a/Engine/Plugins/Experimental/DynamicMaterial/Source/DynamicMaterialEditor/Private/Slate/SDMComponentEdit.cpp
+++ b/Engine/Plugins/Experimental/DynamicMaterial/Source/DynamicMaterialEditor/Private/Slate/SDMComponentEdit.cpp
@@ -910,32 +910,34 @@ void SDMComponentEdit::GenerateMaterialModelPropertyRows(const TSharedRef<SDMEdi
 		const FText PropertyFormat = LOCTEXT("PropertyFormat", "Global {0}");
 		UEnum* MaterialPropertyEnum = StaticEnum<EDMMaterialPropertyType>();
 
-		for (uint8 PropertyIndex = static_cast<uint8>(EDMMaterialPropertyType::OpacityMask) + 1;
-			PropertyIndex < static_cast<uint8>(EDMMaterialPropertyType::Any);
-			++PropertyIndex)
-		{
-			const EDMMaterialPropertyType Property = static_cast<EDMMaterialPropertyType>(PropertyIndex);
-
-			if (!!EditorOnlyData->GetSlotForMaterialProperty(Property))
+		EditorOnlyData->ForEachMaterialPropertyType(
+			[EditorOnlyData , &AddGlobalValue, &PropertyFormat, MaterialPropertyEnum]
+			(EDMMaterialPropertyType InProperty)
 			{
-				if (UDMMaterialProperty* MaterialProperty = EditorOnlyData->GetMaterialProperty(Property))
+				if (EditorOnlyData->GetSlotForMaterialProperty(InProperty))
 				{
-					if (MaterialProperty->IsValidForModel(*EditorOnlyData))
+					if (UDMMaterialProperty* MaterialProperty = EditorOnlyData->GetMaterialProperty(InProperty))
 					{
-						if (UDMMaterialValueFloat1* AlphaValue = Cast<UDMMaterialValueFloat1>(MaterialProperty->GetComponent(UDynamicMaterialModelEditorOnlyData::AlphaValueName)))
+						if (MaterialProperty->IsValidForModel(*EditorOnlyData))
 						{
-							AddGlobalValue(
-								AlphaValue,
-								FText::Format(
-									PropertyFormat,
-									MaterialPropertyEnum->GetDisplayNameTextByValue(static_cast<int64>(Property))
-								)
-							);
+							if (UDMMaterialValueFloat1* AlphaValue = Cast<UDMMaterialValueFloat1>(MaterialProperty->GetComponent(UDynamicMaterialModelEditorOnlyData::AlphaValueName)))
+							{
+								AddGlobalValue(
+									AlphaValue,
+									FText::Format(
+										PropertyFormat,
+										MaterialPropertyEnum->GetDisplayNameTextByValue(static_cast<int64>(InProperty))
+									)
+								);
+							}
 						}
 					}
 				}
-			}
-		}
+
+				return EDMIterationResult::Continue;
+			}, 
+			/* Start from */ EDMMaterialPropertyType::OpacityMask
+		);
 	}
 }
 
diff --git a/Engine/Plugins/Experimental/DynamicMaterial/Source/DynamicMaterialEditor/Public/DMEDefs.h b/Engine/Plugins/Experimental/DynamicMaterial/Source/DynamicMaterialEditor/Public/DMEDefs.h
--- a/Engine/Plugins/Experimental/DynamicMaterial/Source/DynamicMaterialEditor/Public/DMEDefs.h
+++ b/Engine/Plugins/Experimental/DynamicMaterial/Source/DynamicMaterialEditor/Public/DMEDefs.h
@@ -24,6 +24,12 @@ class UDynamicMaterialInstance;
 class UDynamicMaterialModel;
 class UPrimitiveComponent;
 
+enum class EDMIterationResult : uint8
+{
+	Continue,
+	Break
+};
+
 UINTERFACE(MinimalAPI, Blueprintable, BlueprintType)
 class UDMBuildable : public UInterface
 {
diff --git a/Engine/Plugins/Experimental/DynamicMaterial/Source/DynamicMaterialEditor/Public/Model/DynamicMaterialModelEditorOnlyData.h b/Engine/Plugins/Experimental/DynamicMaterial/Source/DynamicMaterialEditor/Public/Model/DynamicMaterialModelEditorOnlyData.h
--- a/Engine/Plugins/Experimental/DynamicMaterial/Source/DynamicMaterialEditor/Public/Model/DynamicMaterialModelEditorOnlyData.h
+++ b/Engine/Plugins/Experimental/DynamicMaterial/Source/DynamicMaterialEditor/Public/Model/DynamicMaterialModelEditorOnlyData.h
@@ -183,6 +183,10 @@ public:
 
 	void OnWizardComplete();
 
+	void ForEachMaterialPropertyType(TFunctionRef<EDMIterationResult(EDMMaterialPropertyType InType)> InCallable,
+		EDMMaterialPropertyType InStart = static_cast<EDMMaterialPropertyType>(static_cast<uint8>(EDMMaterialPropertyType::None) + 1),
+		EDMMaterialPropertyType InEnd = static_cast<EDMMaterialPropertyType>(static_cast<uint8>(EDMMaterialPropertyType::Any) - 1));
+
 	//~ Begin FNotifyHook
 	virtual void NotifyPreChange(class FEditPropertyChain* PropertyAboutToChange) {}
 	virtual void NotifyPostChange(const FPropertyChangedEvent& PropertyChangedEvent, class FEditPropertyChain* PropertyThatChanged);
