Material Designer
- Integrated the width override changes from the Custom Details View to give more space to vectors.
- Changed ratio on scale properties to locked by default.
- Removed unneeded details code from runtime module.
- Added delta meta data to vector2 values.

diff --git a/Engine/Plugins/Experimental/DynamicMaterial/Source/DynamicMaterial/DynamicMaterial.Build.cs b/Engine/Plugins/Experimental/DynamicMaterial/Source/DynamicMaterial/DynamicMaterial.Build.cs
--- a/Engine/Plugins/Experimental/DynamicMaterial/Source/DynamicMaterial/DynamicMaterial.Build.cs
+++ b/Engine/Plugins/Experimental/DynamicMaterial/Source/DynamicMaterial/DynamicMaterial.Build.cs
@@ -33,7 +33,6 @@ public class DynamicMaterial : ModuleRules
 		{
 			PrivateDependencyModuleNames.AddRange(new string[]
 			{
-				"PropertyEditor",
 				"UnrealEd"
 			});
 		}
diff --git a/Engine/Plugins/Experimental/DynamicMaterial/Source/DynamicMaterial/Private/Components/DMMaterialValue.cpp b/Engine/Plugins/Experimental/DynamicMaterial/Source/DynamicMaterial/Private/Components/DMMaterialValue.cpp
--- a/Engine/Plugins/Experimental/DynamicMaterial/Source/DynamicMaterial/Private/Components/DMMaterialValue.cpp
+++ b/Engine/Plugins/Experimental/DynamicMaterial/Source/DynamicMaterial/Private/Components/DMMaterialValue.cpp
@@ -6,15 +6,12 @@
 #include "Model/DynamicMaterialModel.h"
 #include "Materials/Material.h"
 #include "UObject/Package.h"
-#include "Utils/DMDetailsViewUtils.h"
 
 #if WITH_EDITOR
 #include "AssetRegistry/AssetRegistryModule.h"
 #include "DMValueDefinition.h"
 #include "DynamicMaterialModule.h"
 #include "Factories/MaterialFactoryNew.h"
-#include "IDetailTreeNode.h"
-#include "IPropertyRowGenerator.h"
 #include "Model/IDMMaterialBuildStateInterface.h"
 #include "Model/IDMMaterialBuildUtilsInterface.h"
 #include "Model/IDynamicMaterialModelEditorOnlyDataInterface.h"
@@ -418,48 +415,6 @@ void UDMMaterialValue::PostEditChangeProperty(FPropertyChangedEvent& PropertyCha
 		}
 	}
 }
- 
-void UDMMaterialValue::EnsureDetailObjects()
-{
-	if (!IsComponentValid())
-	{
-		return;
-	}
-
-	if (PropertyRowGenerator.IsValid() && DetailTreeNode.IsValid() && PropertyHandle.IsValid())
-	{
-		return;
-	}
- 
-	PropertyHandle.Reset();
-	DetailTreeNode.Reset();
-	PropertyRowGenerator.Reset();
- 
-	FPropertyEditorModule& PropertyEditor = FModuleManager::Get().LoadModuleChecked<FPropertyEditorModule>(TEXT("PropertyEditor"));
- 
-	FPropertyRowGeneratorArgs RowGeneratorArgs;
-	PropertyRowGenerator = PropertyEditor.CreatePropertyRowGenerator(RowGeneratorArgs);
-	PropertyRowGenerator->SetObjects({this});
-
-	DetailTreeNode = FDMDetailsViewUtils::SearchNodesForProperty(PropertyRowGenerator->GetRootTreeNodes(), ValueName);
-
-	if (DetailTreeNode.IsValid())
-	{
-		PropertyHandle = DetailTreeNode->CreatePropertyHandle();
-	}	 
-}
-
-TSharedPtr<IDetailTreeNode> UDMMaterialValue::GetDetailTreeNode()
-{
-	EnsureDetailObjects();
-	return DetailTreeNode;
-}
- 
-TSharedPtr<IPropertyHandle> UDMMaterialValue::GetPropertyHandle()
-{
-	EnsureDetailObjects();
-	return PropertyHandle;
-}
 #endif // WITH_EDITOR
 
 UDMMaterialValue::UDMMaterialValue(EDMValueType InType)
diff --git a/Engine/Plugins/Experimental/DynamicMaterial/Source/DynamicMaterial/Private/Components/DMTextureUV.cpp b/Engine/Plugins/Experimental/DynamicMaterial/Source/DynamicMaterial/Private/Components/DMTextureUV.cpp
--- a/Engine/Plugins/Experimental/DynamicMaterial/Source/DynamicMaterial/Private/Components/DMTextureUV.cpp
+++ b/Engine/Plugins/Experimental/DynamicMaterial/Source/DynamicMaterial/Private/Components/DMTextureUV.cpp
@@ -6,15 +6,9 @@
 #include "Materials/MaterialInstanceDynamic.h"
 #include "Model/DynamicMaterialModel.h"
 #include "Serialization/CustomVersion.h"
-#include "Utils/DMDetailsViewUtils.h"
 
 #if WITH_EDITOR
-#include "IDetailTreeNode.h"
-#include "IPropertyRowGenerator.h"
 #include "Model/IDynamicMaterialModelEditorOnlyDataInterface.h"
-#include "Modules/ModuleManager.h"
-#include "PropertyEditorModule.h"
-#include "PropertyHandle.h"
 #endif
 
 namespace UE::DynamicMaterial::Private
@@ -205,20 +199,6 @@ void UDMTextureUV::SetMirrorOnY(bool bInMirrorOnY)
 
 	Update(EDMUpdateType::Structure);
 }
-
-TSharedPtr<IDetailTreeNode> UDMTextureUV::GetDetailTreeNode(FName InProperty)
-{
-	EnsureDetailObjects();
-
-	return DetailTreeNodes.FindChecked(InProperty);
-}
-
-TSharedPtr<IPropertyHandle> UDMTextureUV::GetPropertyHandle(FName InProperty)
-{
-	EnsureDetailObjects();
-
-	return PropertyHandles.FindChecked(InProperty);
-}
 #endif
 
 TArray<UDMMaterialParameter*> UDMTextureUV::GetParameters() const
@@ -379,72 +359,6 @@ void UDMTextureUV::SetMIDParameters(UMaterialInstanceDynamic* InMID)
 }
 
 #if WITH_EDITOR
-bool UDMTextureUV::CanResetToDefault(TSharedPtr<IPropertyHandle> InPropertyHandle) const
-{
-	FProperty* Property = InPropertyHandle->GetProperty();
-
-	if (!Property)
-	{
-		return false;
-	}
-
-	FName PropertyName = Property->GetFName();
-
-	if (PropertyName == NAME_None)
-	{
-		return false;
-	}
-
-	UDMTextureUV* DefaultObject = Cast<UDMTextureUV>(GetClass()->GetDefaultObject());
-
-	if (!DefaultObject)
-	{
-		return false;
-	}
-
-	if (PropertyName == NAME_UVSource)
-	{
-		return DefaultObject->GetUVSource() != GetUVSource();
-	}
-
-	if (PropertyName == NAME_bMirrorOnX)
-	{
-		return DefaultObject->GetMirrorOnX() != GetMirrorOnX();
-	}
-
-	if (PropertyName == NAME_bMirrorOnY)
-	{
-		return DefaultObject->GetMirrorOnY() != GetMirrorOnY();
-	}
-
-	if (PropertyName == NAME_Offset)
-	{
-		return !DefaultObject->GetOffset().Equals(GetOffset());
-	}
-
-	if (PropertyName == NAME_Pivot)
-	{
-		return !DefaultObject->GetPivot().Equals(GetPivot());
-	}
-
-	if (PropertyName == NAME_Rotation)
-	{
-		return DefaultObject->GetRotation() != GetRotation();
-	}
-
-	if (PropertyName == NAME_Scale)
-	{
-		return !DefaultObject->GetScale().Equals(GetScale());
-	}
-
-	return false;
-}
-
-void UDMTextureUV::ResetToDefault(TSharedPtr<IPropertyHandle> InPropertyHandle)
-{
-	InPropertyHandle->ResetToDefault();
-}
-
 bool UDMTextureUV::Modify(bool bInAlwaysMarkDirty)
 {
 	const bool bSaved = Super::Modify(bInAlwaysMarkDirty);
@@ -456,70 +370,6 @@ bool UDMTextureUV::Modify(bool bInAlwaysMarkDirty)
 
 	return bSaved;
 }
-
-void UDMTextureUV::EnsureDetailObjects()
-{
-	if (!IsComponentValid())
-	{
-		return;
-	}
-
-	bool bHasValidDetailObjects = true;
-
-	if (!PropertyRowGenerator.IsValid() || (DetailTreeNodes.Num() != TextureProperties.Num() && PropertyHandles.Num() != TextureProperties.Num()))
-	{
-		bHasValidDetailObjects = false;
-	}
-	else
-	{
-		for (const TPair<FName, TSharedPtr<IDetailTreeNode>>& DetailTreeNode : DetailTreeNodes)
-		{
-			if (DetailTreeNode.Value.IsValid() == false)
-			{
-				bHasValidDetailObjects = false;
-				break;
-			}
-		}
-
-		if (bHasValidDetailObjects)
-		{
-			for (const TPair<FName, TSharedPtr<IPropertyHandle>>& PropertyHandle : PropertyHandles)
-			{
-				if (PropertyHandle.Value.IsValid() == false)
-				{
-					bHasValidDetailObjects = false;
-					break;
-				}
-			}
-		}
-	}
-
-	if (bHasValidDetailObjects)
-	{
-		return;
-	}
-
-	PropertyHandles.Reset();
-	DetailTreeNodes.Reset();
-	PropertyRowGenerator.Reset();
-
-	FPropertyEditorModule& PropertyEditor = FModuleManager::Get().LoadModuleChecked<FPropertyEditorModule>(TEXT("PropertyEditor"));
-
-	FPropertyRowGeneratorArgs RowGeneratorArgs;
-	PropertyRowGenerator = PropertyEditor.CreatePropertyRowGenerator(RowGeneratorArgs);
-	PropertyRowGenerator->SetObjects({this});
-
-	const TArray<TSharedRef<IDetailTreeNode>>& CategoryNodes = PropertyRowGenerator->GetRootTreeNodes();
-
-	for (const TPair<FName, bool>& TextureProperty : TextureProperties)
-	{
-		if (TSharedPtr<IDetailTreeNode> DetailTreeNode = FDMDetailsViewUtils::SearchNodesForProperty(CategoryNodes, TextureProperty.Key))
-		{
-			DetailTreeNodes.Emplace(DetailTreeNode->GetNodeName(), DetailTreeNode);
-			PropertyHandles.Emplace(DetailTreeNode->GetNodeName(), DetailTreeNode->CreatePropertyHandle());
-		}
-	}
-}
 #endif
 
 void UDMTextureUV::Update(EDMUpdateType InUpdateType)
diff --git a/Engine/Plugins/Experimental/DynamicMaterial/Source/DynamicMaterial/Private/Utils/DMDetailsViewUtils.cpp b/Engine/Plugins/Experimental/DynamicMaterial/Source/DynamicMaterial/Private/Utils/DMDetailsViewUtils.cpp
deleted file mode 100644
diff --git a/Engine/Plugins/Experimental/DynamicMaterial/Source/DynamicMaterial/Public/Components/DMMaterialValue.h b/Engine/Plugins/Experimental/DynamicMaterial/Source/DynamicMaterial/Public/Components/DMMaterialValue.h
--- a/Engine/Plugins/Experimental/DynamicMaterial/Source/DynamicMaterial/Public/Components/DMMaterialValue.h
+++ b/Engine/Plugins/Experimental/DynamicMaterial/Source/DynamicMaterial/Public/Components/DMMaterialValue.h
@@ -9,6 +9,7 @@
 #include "Utils/DMJsonUtils.h"
 #include "DMMaterialValue.generated.h"
 
+class IPropertyHandle;
 class UDMMaterialParameter;
 class UDynamicMaterialModel;
 class UMaterial;
@@ -16,16 +17,9 @@ class UMaterialExpressionParameter;
 class UMaterialInstanceDynamic;
 
 #if WITH_EDITOR
-class FAssetThumbnailPool;
-class IDetailTreeNode;
-class IPropertyHandle;
-class IPropertyRowGenerator;
-class SWidget;
 struct IDMMaterialBuildStateInterface;
 #endif
 
-class UDMMaterialValue;
-
 /**
  * A value used inside a material. Can be exported as a material parameter.
  */
@@ -117,9 +111,6 @@ public:
 	 */
 	virtual int32 GetInnateMaskOutput(int32 OutputChannels) const;
 
-	TSharedPtr<IDetailTreeNode> GetDetailTreeNode();
-	TSharedPtr<IPropertyHandle> GetPropertyHandle();
-
 	virtual FName GetMainPropertyName() const { return NAME_None; }
 
 	/** Return true if, when setting the base stage, the same value should be applied to the mask stage. */
@@ -180,12 +171,6 @@ protected:
 	UPROPERTY(VisibleInstanceOnly, BlueprintReadOnly, Transient, DuplicateTransient, TextExportTransient, Category = "Material Designer")
 	TObjectPtr<UMaterial> PreviewMaterial;
 #endif
-
-#if WITH_EDITOR 
-	TSharedPtr<IPropertyRowGenerator> PropertyRowGenerator;
-	TSharedPtr<IDetailTreeNode> DetailTreeNode;
-	TSharedPtr<IPropertyHandle> PropertyHandle;
-#endif
  
 	UDMMaterialValue(EDMValueType InType);
  
@@ -203,8 +188,6 @@ protected:
 	virtual void OnComponentAdded() override;
 	virtual void OnComponentRemoved() override;
 	//~ End UDMMaterialComponent
- 
-	void EnsureDetailObjects();
 #endif
 
 private:
diff --git a/Engine/Plugins/Experimental/DynamicMaterial/Source/DynamicMaterial/Public/Components/DMTextureUV.h b/Engine/Plugins/Experimental/DynamicMaterial/Source/DynamicMaterial/Public/Components/DMTextureUV.h
--- a/Engine/Plugins/Experimental/DynamicMaterial/Source/DynamicMaterial/Public/Components/DMTextureUV.h
+++ b/Engine/Plugins/Experimental/DynamicMaterial/Source/DynamicMaterial/Public/Components/DMTextureUV.h
@@ -19,12 +19,6 @@ class UDMTextureUV;
 class UDynamicMaterialModel;
 class UMaterialInstanceDynamic;
 
-#if WITH_EDITOR
-class IDetailTreeNode;
-class IPropertyHandle;
-class IPropertyRowGenerator;
-#endif
-
 namespace UE::DynamicMaterial::ParamID
 {
 	// Individual parameters
@@ -134,9 +128,6 @@ public:
 
 	UFUNCTION(BlueprintCallable, Category = "Material Designer")
 	void SetMirrorOnY(bool bInMirrorOnY);
-
-	TSharedPtr<IDetailTreeNode> GetDetailTreeNode(FName InProperty);
-	TSharedPtr<IPropertyHandle> GetPropertyHandle(FName InProperty);
 #endif
 
 	UFUNCTION(BlueprintPure, Category = "Material Designer")
@@ -155,11 +146,6 @@ public:
 
 	void SetMIDParameters(UMaterialInstanceDynamic* InMID);
 
-#if WITH_EDITOR
-	bool CanResetToDefault(TSharedPtr<IPropertyHandle> InPropertyHandle) const;
-	void ResetToDefault(TSharedPtr<IPropertyHandle> InPropertyHandle);
-#endif
-
 	//~ Begin UObject
 #if WITH_EDITOR
 	virtual bool Modify(bool bInAlwaysMarkDirty = true) override;
@@ -200,7 +186,7 @@ protected:
 	float Rotation = 0.f;
 
 	UPROPERTY(EditInstanceOnly, BlueprintReadWrite, Getter = GetScale, Setter = SetScale, BlueprintSetter = SetScale, Category = "Material Designer|Texture UV",
-		meta = (AllowPrivateAccess = "true", AllowPreserveRatio, Delta = 0.001))
+		meta = (AllowPrivateAccess = "true", AllowPreserveRatio, VectorRatioMode = 3, Delta = 0.001))
 	FVector2D Scale = FVector2D(1.f, 1.f);
 
 #if WITH_EDITORONLY_DATA
@@ -222,15 +208,9 @@ protected:
 #if WITH_EDITOR
 	mutable TMap<int32, bool> MaterialNodesCreated;
 
-	TSharedPtr<IPropertyRowGenerator> PropertyRowGenerator;
-	TMap<FName, TSharedPtr<IDetailTreeNode>> DetailTreeNodes;
-	TMap<FName, TSharedPtr<IPropertyHandle>> PropertyHandles;
-
 	bool bNeedsPostLoadValueUpdate = false;
 	bool bNeedsPostLoadStructureUpdate = false;
 
-	void EnsureDetailObjects();
-
 	void CreateParameterNames();
 	void RemoveParameterNames();
 #endif
diff --git a/Engine/Plugins/Experimental/DynamicMaterial/Source/DynamicMaterial/Public/Components/MaterialValues/DMMaterialValueFloat2.h b/Engine/Plugins/Experimental/DynamicMaterial/Source/DynamicMaterial/Public/Components/MaterialValues/DMMaterialValueFloat2.h
--- a/Engine/Plugins/Experimental/DynamicMaterial/Source/DynamicMaterial/Public/Components/MaterialValues/DMMaterialValueFloat2.h
+++ b/Engine/Plugins/Experimental/DynamicMaterial/Source/DynamicMaterial/Public/Components/MaterialValues/DMMaterialValueFloat2.h
@@ -45,7 +45,7 @@ public:
  
 protected: 
 	UPROPERTY(EditInstanceOnly, BlueprintReadWrite, Getter = GetValue, Setter = SetValue, BlueprintSetter = SetValue, Category = "Material Designer",
-		meta = (AllowPrivateAccess = "true", DisplayName = "Vector", AllowPreserveRatio))
+		meta = (AllowPrivateAccess = "true", DisplayName = "Vector", AllowPreserveRatio, Delta = 0.001))
 	FVector2D Value;
  
 #if WITH_EDITORONLY_DATA
diff --git a/Engine/Plugins/Experimental/DynamicMaterial/Source/DynamicMaterial/Public/Utils/DMDetailsViewUtils.h b/Engine/Plugins/Experimental/DynamicMaterial/Source/DynamicMaterial/Public/Utils/DMDetailsViewUtils.h
deleted file mode 100644
diff --git a/Engine/Plugins/Experimental/DynamicMaterial/Source/DynamicMaterialEditor/Private/Slate/Properties/Generators/DMMaterialValuePropertyRowGenerator.cpp b/Engine/Plugins/Experimental/DynamicMaterial/Source/DynamicMaterialEditor/Private/Slate/Properties/Generators/DMMaterialValuePropertyRowGenerator.cpp
--- a/Engine/Plugins/Experimental/DynamicMaterial/Source/DynamicMaterialEditor/Private/Slate/Properties/Generators/DMMaterialValuePropertyRowGenerator.cpp
+++ b/Engine/Plugins/Experimental/DynamicMaterial/Source/DynamicMaterialEditor/Private/Slate/Properties/Generators/DMMaterialValuePropertyRowGenerator.cpp
@@ -52,22 +52,6 @@ void FDMMaterialValuePropertyRowGenerator::AddComponentProperties(const TSharedR
 			FResetToDefaultHandler::CreateUObject(Value, &UDMMaterialValue::ResetToDefault)
 		);
 
-		if (UDMMaterialValueFloat* FloatValue = Cast<UDMMaterialValueFloat>(Value))
-		{
-			if (FloatValue->HasValueRange())
-			{
-				static const FName UIMin = FName("UIMin");
-				static const FName UIMax = FName("UIMax");
-				static const FName ClampMin = FName("ClampMin");
-				static const FName ClampMax = FName("ClampMax");
-
-				Handle.PropertyHandle->SetInstanceMetaData(UIMin, FString::SanitizeFloat(FloatValue->GetValueRange().Min));
-				Handle.PropertyHandle->SetInstanceMetaData(ClampMin, FString::SanitizeFloat(FloatValue->GetValueRange().Min));
-				Handle.PropertyHandle->SetInstanceMetaData(UIMax, FString::SanitizeFloat(FloatValue->GetValueRange().Max));
-				Handle.PropertyHandle->SetInstanceMetaData(ClampMax, FString::SanitizeFloat(FloatValue->GetValueRange().Max));
-			}
-		}
-
 		InOutPropertyRows.Add(Handle);
 	}
 
diff --git a/Engine/Plugins/Experimental/DynamicMaterial/Source/DynamicMaterialEditor/Private/Slate/Properties/Generators/DMTextureUVPropertyRowGenerator.cpp b/Engine/Plugins/Experimental/DynamicMaterial/Source/DynamicMaterialEditor/Private/Slate/Properties/Generators/DMTextureUVPropertyRowGenerator.cpp
--- a/Engine/Plugins/Experimental/DynamicMaterial/Source/DynamicMaterialEditor/Private/Slate/Properties/Generators/DMTextureUVPropertyRowGenerator.cpp
+++ b/Engine/Plugins/Experimental/DynamicMaterial/Source/DynamicMaterialEditor/Private/Slate/Properties/Generators/DMTextureUVPropertyRowGenerator.cpp
@@ -25,12 +25,19 @@ const TSharedRef<FDMTextureUVPropertyRowGenerator>& FDMTextureUVPropertyRowGener
 
 namespace UE::DynamicMaterialEditor::Private
 {
-	void AddPropertyRow(const TSharedRef<SWidget>& InComponentEditWidget, UDMTextureUV* InTextureUV, FName InProperty, TArray<FDMPropertyHandle>& InOutPropertyRows);
-	void AddVisualizerRow(const TSharedRef<SWidget>& InComponentEditWidget, UDMTextureUV* InTextureUV, TArray<FDMPropertyHandle>& InOutPropertyRows);
+	void AddPropertyRow(const TSharedRef<SWidget>& InComponentEditWidget, UDMTextureUV* InTextureUV, FName InProperty, 
+		TArray<FDMPropertyHandle>& InOutPropertyRows);
+
+	void AddVisualizerRow(const TSharedRef<SWidget>& InComponentEditWidget, UDMTextureUV* InTextureUV, 
+		TArray<FDMPropertyHandle>& InOutPropertyRows);
+
+	bool CanResetTextureUVPropertyToDefault(TSharedPtr<IPropertyHandle> InPropertyHandle);
+
+	void ResetTextureUVPropertyToDefault(TSharedPtr<IPropertyHandle> InPropertyHandle);
 }
 
-void FDMTextureUVPropertyRowGenerator::AddComponentProperties(const TSharedRef<SDMComponentEdit>& InComponentEditWidget, UDMMaterialComponent* InComponent,
-	TArray<FDMPropertyHandle>& InOutPropertyRows, TSet<UDMMaterialComponent*>& InOutProcessedObjects)
+void FDMTextureUVPropertyRowGenerator::AddComponentProperties(const TSharedRef<SDMComponentEdit>& InComponentEditWidget, 
+	UDMMaterialComponent* InComponent, TArray<FDMPropertyHandle>& InOutPropertyRows, TSet<UDMMaterialComponent*>& InOutProcessedObjects)
 {
 	if (!IsValid(InComponent))
 	{
@@ -62,7 +69,8 @@ void FDMTextureUVPropertyRowGenerator::AddComponentProperties(const TSharedRef<S
 	AddVisualizerRow(InComponentEditWidget, TextureUV, InOutPropertyRows);
 }
 
-void FDMTextureUVPropertyRowGenerator::AddPopoutComponentProperties(const TSharedRef<SWidget>& InParentWidget, UDMMaterialComponent* InComponent, TArray<FDMPropertyHandle>& InOutPropertyRows)
+void FDMTextureUVPropertyRowGenerator::AddPopoutComponentProperties(const TSharedRef<SWidget>& InParentWidget, UDMMaterialComponent* InComponent, 
+	TArray<FDMPropertyHandle>& InOutPropertyRows)
 {
 	if (!IsValid(InComponent))
 	{
@@ -101,17 +109,19 @@ bool FDMTextureUVPropertyRowGenerator::AllowKeyframeButton(UDMMaterialComponent*
 	return FDMComponentPropertyRowGenerator::AllowKeyframeButton(InComponent, InProperty);
 }
 
-void UE::DynamicMaterialEditor::Private::AddPropertyRow(const TSharedRef<SWidget>& InComponentEditWidget, UDMTextureUV* InTextureUV, FName InProperty, TArray<FDMPropertyHandle>& InOutPropertyRows)
+void UE::DynamicMaterialEditor::Private::AddPropertyRow(const TSharedRef<SWidget>& InComponentEditWidget, UDMTextureUV* InTextureUV, 
+	FName InProperty, TArray<FDMPropertyHandle>& InOutPropertyRows)
 {
 	FDMPropertyHandle& NewHandle = InOutPropertyRows.Add_GetRef(SDMEditor::GetPropertyHandle(&*InComponentEditWidget, InTextureUV, InProperty));
 
 	NewHandle.ResetToDefaultOverride = FResetToDefaultOverride::Create(
-		FIsResetToDefaultVisible::CreateUObject(InTextureUV, &UDMTextureUV::CanResetToDefault),
-		FResetToDefaultHandler::CreateUObject(InTextureUV, &UDMTextureUV::ResetToDefault)
+		FIsResetToDefaultVisible::CreateStatic(&UE::DynamicMaterialEditor::Private::CanResetTextureUVPropertyToDefault),
+		FResetToDefaultHandler::CreateStatic(&UE::DynamicMaterialEditor::Private::ResetTextureUVPropertyToDefault)
 	);
 }
 
-void UE::DynamicMaterialEditor::Private::AddVisualizerRow(const TSharedRef<SWidget>& InComponentEditWidget, UDMTextureUV* InTextureUV, TArray<FDMPropertyHandle>& InOutPropertyRows)
+void UE::DynamicMaterialEditor::Private::AddVisualizerRow(const TSharedRef<SWidget>& InComponentEditWidget, UDMTextureUV* InTextureUV, 
+	TArray<FDMPropertyHandle>& InOutPropertyRows)
 {
 	// Make sure we don't get a substage
 	UDMMaterialStage* Stage = InTextureUV->GetTypedParent<UDMMaterialStage>(/* Allow Subclasses */ false);
@@ -130,4 +140,86 @@ void UE::DynamicMaterialEditor::Private::AddVisualizerRow(const TSharedRef<SWidg
 	InOutPropertyRows.Add(VisualizerHandle);
 }
 
+
+bool UE::DynamicMaterialEditor::Private::CanResetTextureUVPropertyToDefault(TSharedPtr<IPropertyHandle> InPropertyHandle)
+{
+	FProperty* Property = InPropertyHandle->GetProperty();
+
+	if (!Property)
+	{
+		return false;
+	}
+
+	FName PropertyName = Property->GetFName();
+
+	if (PropertyName == NAME_None)
+	{
+		return false;
+	}
+
+	TArray<UObject*> Outers;
+	InPropertyHandle->GetOuterObjects(Outers);
+
+	if (Outers.IsEmpty())
+	{
+		return false;
+	}
+
+	const UDMTextureUV* PropertyObject = Cast<UDMTextureUV>(Outers[0]);
+
+	if (!PropertyObject)
+	{
+		return false;
+	}
+
+	const UDMTextureUV* DefaultObject = GetDefault<UDMTextureUV>();
+
+	if (!DefaultObject)
+	{
+		return false;
+	}
+
+	if (PropertyName == UDMTextureUV::NAME_UVSource)
+	{
+		return DefaultObject->GetUVSource() != PropertyObject->GetUVSource();
+	}
+
+	if (PropertyName == UDMTextureUV::NAME_bMirrorOnX)
+	{
+		return DefaultObject->GetMirrorOnX() != PropertyObject->GetMirrorOnX();
+	}
+
+	if (PropertyName == UDMTextureUV::NAME_bMirrorOnY)
+	{
+		return DefaultObject->GetMirrorOnY() != PropertyObject->GetMirrorOnY();
+	}
+
+	if (PropertyName == UDMTextureUV::NAME_Offset)
+	{
+		return !DefaultObject->GetOffset().Equals(PropertyObject->GetOffset());
+	}
+
+	if (PropertyName == UDMTextureUV::NAME_Pivot)
+	{
+		return !DefaultObject->GetPivot().Equals(PropertyObject->GetPivot());
+	}
+
+	if (PropertyName == NAME_Rotation)
+	{
+		return DefaultObject->GetRotation() != PropertyObject->GetRotation();
+	}
+
+	if (PropertyName == UDMTextureUV::NAME_Scale)
+	{
+		return !DefaultObject->GetScale().Equals(PropertyObject->GetScale());
+	}
+
+	return false;
+}
+
+void UE::DynamicMaterialEditor::Private::ResetTextureUVPropertyToDefault(TSharedPtr<IPropertyHandle> InPropertyHandle)
+{
+	InPropertyHandle->ResetToDefault();
+}
+
 #undef LOCTEXT_NAMESPACE
diff --git a/Engine/Plugins/Experimental/DynamicMaterial/Source/DynamicMaterialEditor/Private/Slate/Properties/SDMDetailsGrid.cpp b/Engine/Plugins/Experimental/DynamicMaterial/Source/DynamicMaterialEditor/Private/Slate/Properties/SDMDetailsGrid.cpp
--- a/Engine/Plugins/Experimental/DynamicMaterial/Source/DynamicMaterialEditor/Private/Slate/Properties/SDMDetailsGrid.cpp
+++ b/Engine/Plugins/Experimental/DynamicMaterial/Source/DynamicMaterialEditor/Private/Slate/Properties/SDMDetailsGrid.cpp
@@ -297,6 +297,7 @@ FOnGenerateGlobalRowExtensionArgs SDMDetailsGrid::CreateExtensionArgs(UDMMateria
 {
 	FOnGenerateGlobalRowExtensionArgs ExtensionRowArgs;
 
+	/*
 	if (ensure(InValue))
 	{
 		FProperty* ValueProperty = InValue->GetClass()->FindPropertyByName("Value");
@@ -311,6 +312,7 @@ FOnGenerateGlobalRowExtensionArgs SDMDetailsGrid::CreateExtensionArgs(UDMMateria
 			ExtensionRowArgs.PropertyHandle = ValuePropertyHandle;
 		}
 	}
+	*/
 
 	return ExtensionRowArgs;
 }
@@ -350,6 +352,7 @@ void SDMDetailsGrid::CreateKey(UObject* InObject, TSharedPtr<IPropertyHandle> In
 
 void SDMDetailsGrid::CreateKey(UDMMaterialValue* InMaterialValue)
 {
+	/*
 	if (!IsValid(InMaterialValue))
 	{
 		return;
@@ -362,6 +365,7 @@ void SDMDetailsGrid::CreateKey(UDMMaterialValue* InMaterialValue)
 	}
 
 	CreateKey(InMaterialValue, PropertyHandle);
+	*/
 }
 
 #undef LOCTEXT_NAMESPACE
diff --git a/Engine/Plugins/Experimental/DynamicMaterial/Source/DynamicMaterialEditor/Private/Slate/Properties/SDMPropertyEdit.cpp b/Engine/Plugins/Experimental/DynamicMaterial/Source/DynamicMaterialEditor/Private/Slate/Properties/SDMPropertyEdit.cpp
--- a/Engine/Plugins/Experimental/DynamicMaterial/Source/DynamicMaterialEditor/Private/Slate/Properties/SDMPropertyEdit.cpp
+++ b/Engine/Plugins/Experimental/DynamicMaterial/Source/DynamicMaterialEditor/Private/Slate/Properties/SDMPropertyEdit.cpp
@@ -193,7 +193,7 @@ void SDMPropertyEdit::CreateKey(TWeakObjectPtr<UDMMaterialValue> InValueWeak)
 		return;
 	}
  
-	TSharedPtr<IPropertyHandle> PropertyHandle = InValueWeak->GetPropertyHandle();
+	TSharedPtr<IPropertyHandle> PropertyHandle = nullptr; // InValueWeak->GetPropertyHandle();
 	if (!PropertyHandle.IsValid())
 	{
 		return;
diff --git a/Engine/Plugins/Experimental/DynamicMaterial/Source/DynamicMaterialEditor/Private/Slate/Properties/SDMTextureUVVisualizerPopout.cpp b/Engine/Plugins/Experimental/DynamicMaterial/Source/DynamicMaterialEditor/Private/Slate/Properties/SDMTextureUVVisualizerPopout.cpp
--- a/Engine/Plugins/Experimental/DynamicMaterial/Source/DynamicMaterialEditor/Private/Slate/Properties/SDMTextureUVVisualizerPopout.cpp
+++ b/Engine/Plugins/Experimental/DynamicMaterial/Source/DynamicMaterialEditor/Private/Slate/Properties/SDMTextureUVVisualizerPopout.cpp
@@ -277,13 +277,13 @@ TSharedRef<SWidget> SDMTextureUVVisualizerPopout::CreatePropertyWidget(UDMTextur
 
 		ECustomDetailsTreeInsertPosition Position = ECustomDetailsTreeInsertPosition::Child;
 
-		if (EditRow.DetailTreeNode)
+		if (EditRow.PropertyHandle.IsValid())
 		{
-			if (EditRow.DetailTreeNode->CreatePropertyHandle()->HasMetaData("HighPriority"))
+			if (EditRow.PropertyHandle->HasMetaData("HighPriority"))
 			{
 				Position = ECustomDetailsTreeInsertPosition::FirstChild;
 			}
-			else if (EditRow.DetailTreeNode->CreatePropertyHandle()->HasMetaData("LowPriority"))
+			else if (EditRow.PropertyHandle->HasMetaData("LowPriority"))
 			{
 				Position = ECustomDetailsTreeInsertPosition::LastChild;
 			}
@@ -325,7 +325,7 @@ TSharedRef<SWidget> SDMTextureUVVisualizerPopout::CreatePropertyWidget(UDMTextur
 			);
 		}
 
-		if (EditRow.DetailTreeNode->CreatePropertyHandle()->HasMetaData("NotKeyframeable"))
+		if (EditRow.PropertyHandle->HasMetaData("NotKeyframeable"))
 		{
 			Item->SetKeyframeEnabled(false);
 		}
diff --git a/Engine/Plugins/Experimental/DynamicMaterial/Source/DynamicMaterialEditor/Private/Slate/SDMComponentEdit.cpp b/Engine/Plugins/Experimental/DynamicMaterial/Source/DynamicMaterialEditor/Private/Slate/SDMComponentEdit.cpp
--- a/Engine/Plugins/Experimental/DynamicMaterial/Source/DynamicMaterialEditor/Private/Slate/SDMComponentEdit.cpp
+++ b/Engine/Plugins/Experimental/DynamicMaterial/Source/DynamicMaterialEditor/Private/Slate/SDMComponentEdit.cpp
@@ -329,11 +329,6 @@ TSharedRef<SWidget> SDMComponentEdit::CreateEditWidget()
 			continue;
 		}
 
-		if (!EditRow.PropertyHandle.IsValid() && EditRow.DetailTreeNode)
-		{
-			EditRow.PropertyHandle = EditRow.DetailTreeNode->CreatePropertyHandle();
-		}
-
 		ECustomDetailsTreeInsertPosition Position = ECustomDetailsTreeInsertPosition::Child;
 
 		if (EditRow.PropertyHandle.IsValid())
@@ -397,6 +392,11 @@ TSharedRef<SWidget> SDMComponentEdit::CreateEditWidget()
 
 			Item->SetValueWidget(EditRow.ValueWidget.ToSharedRef());
 
+			if (EditRow.MaxWidth.IsSet())
+			{
+				Item->AsItem()->SetValueWidgetWidthOverride(EditRow.MaxWidth);
+			}
+
 			if (CategoryItem.IsValid())
 			{
 				DetailsView->ExtendTree(CategoryItem->GetItemId(), Position, Item->AsItem());
@@ -437,6 +437,11 @@ TSharedRef<SWidget> SDMComponentEdit::CreateEditWidget()
 			Item->SetResetToDefaultOverride(EditRow.ResetToDefaultOverride.GetValue());
 		}
 
+		if (EditRow.MaxWidth.IsSet())
+		{
+			Item->SetValueWidgetWidthOverride(EditRow.MaxWidth);
+		}
+
 		if (CategoryItem.IsValid())
 		{
 			DetailsView->ExtendTree(CategoryItem->GetItemId(), Position, Item);
@@ -515,27 +520,9 @@ TSharedRef<SWidget> SDMComponentEdit::CreateExtensionButtons(const TSharedPtr<SW
 			.HeightOverride(18.f);
 	}
 
-	TSharedPtr<IPropertyHandle> PropertyHandle = nullptr;
-	TSharedPtr<IDetailTreeNode> DetailTreeNode = nullptr;
-	FString PropertyName = InPropertyName.ToString();
-
-	if (UDMTextureUV* TextureUV = Cast<UDMTextureUV>(InComponent))
-	{
-		PropertyHandle = TextureUV->GetPropertyHandle(InPropertyName);
-		DetailTreeNode = TextureUV->GetDetailTreeNode(InPropertyName);
-		PropertyName = FString("TextureUV->") + PropertyName;
-	}
-	else if (UDMMaterialValue* Value = Cast<UDMMaterialValue>(InComponent))
-	{
-		PropertyHandle = Value->GetPropertyHandle();
-		DetailTreeNode = Value->GetDetailTreeNode();
-	}
-	else
-	{
-		FDMPropertyHandle DMPropertyHandle = SDMEditor::GetPropertyHandle(InPropertyOwner.Get(), InComponent, InPropertyName);
-		PropertyHandle = DMPropertyHandle.PropertyHandle;
-		DetailTreeNode = DMPropertyHandle.DetailTreeNode;
-	}
+	FDMPropertyHandle DMPropertyHandle = SDMEditor::GetPropertyHandle(InPropertyOwner.Get(), InComponent, InPropertyName);
+	TSharedPtr<IPropertyHandle> PropertyHandle = DMPropertyHandle.PropertyHandle;
+	TSharedPtr<IDetailTreeNode> DetailTreeNode = DMPropertyHandle.DetailTreeNode;
 
 	if (PropertyHandle.IsValid())
 	{
@@ -545,7 +532,6 @@ TSharedRef<SWidget> SDMComponentEdit::CreateExtensionButtons(const TSharedPtr<SW
 	FOnGenerateGlobalRowExtensionArgs ExtensionRowArgs;
 	ExtensionRowArgs.OwnerObject = InComponent;
 	ExtensionRowArgs.Property = Property;
-	ExtensionRowArgs.PropertyPath = PropertyName;
 	ExtensionRowArgs.OwnerTreeNode = DetailTreeNode;
 	ExtensionRowArgs.PropertyHandle = PropertyHandle;
 
@@ -842,13 +828,10 @@ void SDMComponentEdit::GenerateMaterialModelPropertyRows(const TSharedRef<SDMEdi
 	const FName MaterialTypeCategory = FName("Material Type");
 	auto AddVariable = [InEditorWidget, &InOutPropertyRows, &MaterialTypeCategory](UObject* InObject, FName InPropertyName)
 		{
-			if (IsValid(InObject))
-			{
-				FDMPropertyHandle& ValueHandle = InOutPropertyRows.Add_GetRef(InEditorWidget->GetPropertyHandle(&*InEditorWidget,
-					InObject, InPropertyName));
+			FDMPropertyHandle& ValueHandle = InOutPropertyRows.Add_GetRef(InEditorWidget->GetPropertyHandle(&*InEditorWidget,
+				InObject, InPropertyName));
 
-				ValueHandle.CategoryOverrideName = MaterialTypeCategory;
-			}
+			ValueHandle.CategoryOverrideName = MaterialTypeCategory;
 		};
 
 	if (UDynamicMaterialModelEditorOnlyData* EditorOnlyData = UDynamicMaterialModelEditorOnlyData::Get(InMaterialModel))
@@ -862,38 +845,17 @@ void SDMComponentEdit::GenerateMaterialModelPropertyRows(const TSharedRef<SDMEdi
 	}
 
 	const FName MaterialSettingsCategory = FName("Material Settings");
-	const FName UIMin = FName("UIMin");
-	const FName UIMax = FName("UIMax");
-	const FName ClampMin = FName("ClampMin");
-	const FName ClampMax = FName("ClampMax");
 
-	auto AddGlobalValue = [InEditorWidget, InMaterialModel, &InOutPropertyRows, &InOutProcessedObjects, &MaterialSettingsCategory , &UIMin, &UIMax, &ClampMin, &ClampMax]
+	auto AddGlobalValue = [InEditorWidget, &InOutPropertyRows, &InOutProcessedObjects, &MaterialSettingsCategory]
 		(UDMMaterialValue* InValue, const FText& InNameOverride)
 		{
-			if (IsValid(InValue))
-			{
-				if (!InOutProcessedObjects.Contains(InValue))
-				{
-					FDMPropertyHandle& ValueHandle = InOutPropertyRows.Add_GetRef(InEditorWidget->GetPropertyHandle(&*InEditorWidget,
-						InValue, UDMMaterialValue::ValueName));
+			FDMPropertyHandle& ValueHandle = InOutPropertyRows.Add_GetRef(InEditorWidget->GetPropertyHandle(&*InEditorWidget,
+				InValue, UDMMaterialValue::ValueName));
 
-					ValueHandle.CategoryOverrideName = MaterialSettingsCategory;
-					ValueHandle.NameOverride = InNameOverride;
+			ValueHandle.CategoryOverrideName = MaterialSettingsCategory;
+			ValueHandle.NameOverride = InNameOverride;
 
-					if (UDMMaterialValueFloat* FloatValue = Cast<UDMMaterialValueFloat>(InValue))
-					{
-						if (FloatValue->HasValueRange() && ValueHandle.PropertyHandle)
-						{
-							ValueHandle.PropertyHandle->SetInstanceMetaData(UIMin, FString::SanitizeFloat(FloatValue->GetValueRange().Min));
-							ValueHandle.PropertyHandle->SetInstanceMetaData(ClampMin, FString::SanitizeFloat(FloatValue->GetValueRange().Min));
-							ValueHandle.PropertyHandle->SetInstanceMetaData(UIMax, FString::SanitizeFloat(FloatValue->GetValueRange().Max));
-							ValueHandle.PropertyHandle->SetInstanceMetaData(ClampMax, FString::SanitizeFloat(FloatValue->GetValueRange().Max));
-						}
-					}
-
-					InOutProcessedObjects.Add(InValue);
-				}
-			}
+			InOutProcessedObjects.Add(InValue);
 		};
 
 	AddGlobalValue(InMaterialModel->GetGlobalParameterValue(UDynamicMaterialModel::GlobalOffsetValueName), LOCTEXT("GlobalOffset", "Global Offset"));
diff --git a/Engine/Plugins/Experimental/DynamicMaterial/Source/DynamicMaterialEditor/Private/Slate/SDMEditor.cpp b/Engine/Plugins/Experimental/DynamicMaterial/Source/DynamicMaterialEditor/Private/Slate/SDMEditor.cpp
--- a/Engine/Plugins/Experimental/DynamicMaterial/Source/DynamicMaterialEditor/Private/Slate/SDMEditor.cpp
+++ b/Engine/Plugins/Experimental/DynamicMaterial/Source/DynamicMaterialEditor/Private/Slate/SDMEditor.cpp
@@ -6,6 +6,7 @@
 #include "Components/DMMaterialProperty.h"
 #include "Components/DMMaterialSlot.h"
 #include "Components/MaterialStageExpressions/DMMSETextureSample.h"
+#include "Components/MaterialValues/DMMaterialValueFloat.h"
 #include "Components/PrimitiveComponent.h"
 #include "DetailLayoutBuilder.h"
 #include "DMWorldSubsystem.h"
@@ -35,7 +36,6 @@
 #include "Styling/StyleColors.h"
 #include "ThumbnailRendering/ThumbnailManager.h"
 #include "Utils/DMBlueprintFunctionLibrary.h"
-#include "Utils/DMDetailsViewUtils.h"
 #include "Utils/DMPrivate.h"
 #include "Widgets/Input/SButton.h"
 #include "Widgets/Input/SCheckBox.h"
@@ -52,9 +52,44 @@
 
 namespace UE::DynamicMaterialEditor::Private
 {
+	TSharedPtr<IDetailTreeNode> SearchNodesForProperty(const TArray<TSharedRef<IDetailTreeNode>>& InNodes, FName InPropertyName)
+	{
+		for (const TSharedRef<IDetailTreeNode>& ChildNode : InNodes)
+		{
+			switch (ChildNode->GetNodeType())
+			{
+				case EDetailNodeType::Category:
+				{
+					TArray<TSharedRef<IDetailTreeNode>> CategoryChildNodes;
+					ChildNode->GetChildren(CategoryChildNodes);
+
+					if (TSharedPtr<IDetailTreeNode> FoundNode = SearchNodesForProperty(CategoryChildNodes, InPropertyName))
+					{
+						return FoundNode;
+					}
+
+					break;
+				}
+
+				case EDetailNodeType::Item:
+					if (ChildNode->GetNodeName() == InPropertyName)
+					{
+						return ChildNode;
+					}
+					break;
+
+				default:
+					// Do nothing
+					break;
+			}
+		}
+
+		return nullptr;
+	}
+
 	TSharedPtr<IDetailTreeNode> SearchGeneratorForNode(const TSharedRef<IPropertyRowGenerator>& InGenerator, FName InPropertyName)
 	{
-		return FDMDetailsViewUtils::SearchNodesForProperty(InGenerator->GetRootTreeNodes(), InPropertyName);
+		return SearchNodesForProperty(InGenerator->GetRootTreeNodes(), InPropertyName);
 	}
 
 	TSharedPtr<IPropertyRowGenerator> SearchForGenerator(const TArray<FDMPropertyHandle>& InPropertyHandles, UObject* InObject)
@@ -80,6 +115,77 @@ namespace UE::DynamicMaterialEditor::Private
 
 		return nullptr;
 	}
+
+	void AddPropertyMetaData(UObject* InObject, FName InPropertyName, FDMPropertyHandle& InPropertyHandle)
+	{
+		FProperty* Property = nullptr;
+
+		if (InPropertyHandle.PropertyHandle.IsValid())
+		{
+			Property = InPropertyHandle.PropertyHandle->GetProperty();
+
+			if (UDMMaterialValueFloat* FloatValue = Cast<UDMMaterialValueFloat>(InObject))
+			{
+				if (FloatValue->HasValueRange())
+				{
+					const FName UIMin = FName("UIMin");
+					const FName UIMax = FName("UIMax");
+					const FName ClampMin = FName("ClampMin");
+					const FName ClampMax = FName("ClampMax");
+
+					InPropertyHandle.PropertyHandle->SetInstanceMetaData(UIMin, FString::SanitizeFloat(FloatValue->GetValueRange().Min));
+					InPropertyHandle.PropertyHandle->SetInstanceMetaData(ClampMin, FString::SanitizeFloat(FloatValue->GetValueRange().Min));
+					InPropertyHandle.PropertyHandle->SetInstanceMetaData(UIMax, FString::SanitizeFloat(FloatValue->GetValueRange().Max));
+					InPropertyHandle.PropertyHandle->SetInstanceMetaData(ClampMax, FString::SanitizeFloat(FloatValue->GetValueRange().Max));
+				}
+			}
+		}
+		else
+		{
+			Property = InObject->GetClass()->FindPropertyByName(InPropertyName);
+		}
+
+		if (FStructProperty* StructProperty = CastField<FStructProperty>(Property))
+		{
+			uint8 ComponentCount = 1;
+
+			if (StructProperty->Struct == TBaseStructure<FVector2D>::Get()
+				|| StructProperty->Struct == TVariantStructure<FVector2f>::Get())
+			{
+				ComponentCount = 2;
+			}
+
+			if (StructProperty->Struct == TBaseStructure<FVector>::Get()
+				|| StructProperty->Struct == TVariantStructure<FVector3f>::Get()
+				|| StructProperty->Struct == TBaseStructure<FRotator>::Get())
+			{
+				ComponentCount = 3;
+			}
+
+			// FLinearColor doesn't need the extra space
+			if (StructProperty->Struct == TBaseStructure<FVector4>::Get()
+				|| StructProperty->Struct == TVariantStructure<FVector4f>::Get())
+			{
+				ComponentCount = 4;
+			}
+
+			switch (ComponentCount)
+			{
+				case 0:
+				case 1:
+					break;
+
+				case 2:
+					InPropertyHandle.MaxWidth = 200.f;
+					break;
+
+					// 3 and above
+				default:
+					InPropertyHandle.MaxWidth = 275.f;
+					break;
+			}
+		}
+	}
 }
 
 TSharedPtr<FAssetThumbnailPool> SDMEditor::ThumbnailPool = nullptr;
@@ -211,15 +317,20 @@ FDMPropertyHandle SDMEditor::GetPropertyHandle(const SWidget* InOwningWidget, UO
 		}
 	}
 
-	if (TSharedPtr<IPropertyRowGenerator> PropertyRowGenerator = UE::DynamicMaterialEditor::Private::SearchForGenerator(PropertyHandles, InObject))
+	using namespace UE::DynamicMaterialEditor::Private;
+
+	if (TSharedPtr<IPropertyRowGenerator> PropertyRowGenerator = SearchForGenerator(PropertyHandles, InObject))
 	{
 		FDMPropertyHandle PropertyHandle;
 		PropertyHandle.PropertyRowGenerator = PropertyRowGenerator;
 
-		if (TSharedPtr<IDetailTreeNode> DetailTreeNode = UE::DynamicMaterialEditor::Private::SearchGeneratorForNode(PropertyRowGenerator.ToSharedRef(), InPropertyName))
+		if (TSharedPtr<IDetailTreeNode> DetailTreeNode = SearchGeneratorForNode(PropertyRowGenerator.ToSharedRef(), InPropertyName))
 		{
 			PropertyHandle.DetailTreeNode = DetailTreeNode;
 			PropertyHandle.PropertyHandle = DetailTreeNode->CreatePropertyHandle();
+
+			AddPropertyMetaData(InObject, InPropertyName, PropertyHandle);
+
 			return PropertyHandle;
 		}
 
@@ -227,6 +338,14 @@ FDMPropertyHandle SDMEditor::GetPropertyHandle(const SWidget* InOwningWidget, UO
 	}
 
 	FDMPropertyHandle NewHandle = CreatePropertyHandle(InOwningWidget, InObject, InPropertyName);
+
+	if (!NewHandle.PropertyHandle.IsValid() && NewHandle.DetailTreeNode.IsValid())
+	{
+		NewHandle.PropertyHandle = NewHandle.DetailTreeNode->CreatePropertyHandle();
+	}
+
+	AddPropertyMetaData(InObject, InPropertyName, NewHandle);
+
 	PropertyHandles.Add(NewHandle);
 
 	return NewHandle;
diff --git a/Engine/Plugins/Experimental/DynamicMaterial/Source/DynamicMaterialEditor/Public/DMEDefs.h b/Engine/Plugins/Experimental/DynamicMaterial/Source/DynamicMaterialEditor/Public/DMEDefs.h
--- a/Engine/Plugins/Experimental/DynamicMaterial/Source/DynamicMaterialEditor/Public/DMEDefs.h
+++ b/Engine/Plugins/Experimental/DynamicMaterial/Source/DynamicMaterialEditor/Public/DMEDefs.h
@@ -129,6 +129,7 @@ struct FDMPropertyHandle
 	TSharedPtr<SWidget> ValueWidget;
 	FName ValueName = NAME_None;
 	FName CategoryOverrideName = NAME_None;
+	TOptional<float> MaxWidth;
 };
 
 UENUM(BlueprintType, meta = (Bitflags, UseEnumValuesAsMaskValuesInEditor = "true"))
