Material Designer
- Moved a chunk of code from EditorOnlyData to MaterialProperty
- Added SafeNormalize as an extra OutputProcessor to the Normal Material Property.

diff --git a/Engine/Plugins/Experimental/DynamicMaterial/Source/DynamicMaterial/Private/Model/DynamicMaterialModel.cpp b/Engine/Plugins/Experimental/DynamicMaterial/Source/DynamicMaterial/Private/Model/DynamicMaterialModel.cpp
--- a/Engine/Plugins/Experimental/DynamicMaterial/Source/DynamicMaterial/Private/Model/DynamicMaterialModel.cpp
+++ b/Engine/Plugins/Experimental/DynamicMaterial/Source/DynamicMaterial/Private/Model/DynamicMaterialModel.cpp
@@ -68,10 +68,10 @@ UDynamicMaterialModel::UDynamicMaterialModel()
 #endif
 
 			UDMMaterialParameter* Parameter = CreateDefaultSubobject<UDMMaterialParameter>(InParameterName);
-			Parameter->ParameterName = InPropertyName;
+			Parameter->ParameterName = InParameterName;
 			Property->Parameter = Parameter;
 
-			ParameterMap.Add(InPropertyName, Parameter);
+			ParameterMap.Add(InParameterName, Parameter);
 		};
 
 	auto AddVector2Parameter = [this](FName InPropertyName, FName InParameterName, const FVector2D& InDefaultValue)
@@ -85,10 +85,10 @@ UDynamicMaterialModel::UDynamicMaterialModel()
 #endif
 
 			UDMMaterialParameter* Parameter = CreateDefaultSubobject<UDMMaterialParameter>(InParameterName);
-			Parameter->ParameterName = InPropertyName;
+			Parameter->ParameterName = InParameterName;
 			Property->Parameter = Parameter;
 
-			ParameterMap.Add(InPropertyName, Parameter);
+			ParameterMap.Add(InParameterName, Parameter);
 		};
 
 	FDMUpdateGuard Guard;
@@ -350,6 +350,7 @@ void UDynamicMaterialModel::RenameParameter(UDMMaterialParameter* InParameter, F
 	}
 
 	InParameter->ParameterName = CreateUniqueParameterName(InBaseName);
+
 	ParameterMap.Emplace(InParameter->ParameterName, InParameter);
 
 	if (IDynamicMaterialModelEditorOnlyDataInterface* ModelEditorOnlyData = GetEditorOnlyData())
@@ -443,8 +444,6 @@ void UDynamicMaterialModel::PostLoad()
 
 	PRAGMA_ENABLE_DEPRECATION_WARNINGS
 
-	FixGlobalVars();
-
 	IDynamicMaterialModelEditorOnlyDataInterface* ModelEditorOnlyData = GetEditorOnlyData();
 
 	if (!ModelEditorOnlyData)
@@ -573,6 +572,8 @@ void UDynamicMaterialModel::ReinitComponents()
 		}
 	}
 
+	FixGlobalVars();
+
 	if (IDynamicMaterialModelEditorOnlyDataInterface* ModelEditorOnlyData = GetEditorOnlyData())
 	{
 		ModelEditorOnlyData->ReinitComponents();
@@ -592,9 +593,19 @@ void UDynamicMaterialModel::FixGlobalVars()
 
 			UDMMaterialParameter* Parameter = Property->GetParameter();
 
-			if (!Parameter)
+			if (!Parameter || Parameter->GetParameterName() != InParameterName)
 			{
-				return;
+				if (Property->SetParameterName(InParameterName))
+				{
+					return;
+				}
+
+				Parameter = Property->GetParameter();
+
+				if (!Parameter)
+				{
+					return;
+				}
 			}
 	
 			if (const TWeakObjectPtr<UDMMaterialParameter>* FoundParameterPtr = ParameterMap.Find(InParameterName))
diff --git a/Engine/Plugins/Experimental/DynamicMaterial/Source/DynamicMaterialEditor/Private/Components/DMMaterialProperty.cpp b/Engine/Plugins/Experimental/DynamicMaterial/Source/DynamicMaterialEditor/Private/Components/DMMaterialProperty.cpp
--- a/Engine/Plugins/Experimental/DynamicMaterial/Source/DynamicMaterialEditor/Private/Components/DMMaterialProperty.cpp
+++ b/Engine/Plugins/Experimental/DynamicMaterial/Source/DynamicMaterialEditor/Private/Components/DMMaterialProperty.cpp
@@ -2,8 +2,8 @@
 
 #include "Components/DMMaterialProperty.h"
 
-#include "DMComponentPath.h"
 #include "Components/DMMaterialComponent.h"
+#include "Components/DMMaterialLayer.h"
 #include "Components/DMMaterialSlot.h"
 #include "Components/DMMaterialStageThroughputLayerBlend.h"
 #include "Components/DMMaterialSubStage.h"
@@ -11,7 +11,9 @@
 #include "Components/MaterialStageExpressions/DMMSETextureSample.h"
 #include "Components/MaterialStageInputs/DMMSIExpression.h"
 #include "Components/MaterialStageInputs/DMMSIValue.h"
+#include "Components/MaterialValues/DMMaterialValueFloat1.h"
 #include "Components/MaterialValues/DMMaterialValueTexture.h"
+#include "DMComponentPath.h"
 #include "DMValueDefinition.h"
 #include "DynamicMaterialEditorSettings.h"
 #include "Materials/MaterialExpressionConstant.h"
@@ -19,13 +21,17 @@
 #include "Materials/MaterialExpressionConstant3Vector.h"
 #include "Materials/MaterialExpressionConstant4Vector.h"
 #include "Materials/MaterialExpressionMaterialFunctionCall.h"
+#include "Materials/MaterialExpressionMax.h"
+#include "Materials/MaterialExpressionMultiply.h"
 #include "Materials/MaterialFunctionInterface.h"
 #include "Math/Vector.h"
 #include "Math/Vector2D.h"
 #include "Math/Vector4.h"
 #include "Model/DMMaterialBuildState.h"
 #include "Model/DMMaterialBuildUtils.h"
+#include "Model/DynamicMaterialModel.h"
 #include "Model/DynamicMaterialModelEditorOnlyData.h"
+#include "Utils/DMMaterialFunctionLibrary.h"
 #include "Utils/DMPrivate.h"
 
 #define LOCTEXT_NAMESPACE "DMMaterialProperty"
@@ -52,7 +58,7 @@ UDMMaterialProperty* UDMMaterialProperty::CreateCustomMaterialPropertyDefaultSub
 	EDMMaterialPropertyType InMaterialProperty, const FName& InSubObjName)
 {
 	check(InModelEditorOnlyData);
-	check(UE::DynamicMaterialEditor::Private::IsCustomMaterialProperty(InMaterialProperty));
+	check(InMaterialProperty >= EDMMaterialPropertyType::Custom1 && InMaterialProperty <= EDMMaterialPropertyType::Custom4);
 
 	UDMMaterialProperty* NewMaterialProperty = InModelEditorOnlyData->CreateDefaultSubobject<UDMMaterialProperty>(InSubObjName);
 	NewMaterialProperty->MaterialProperty = InMaterialProperty;
@@ -71,6 +77,11 @@ FText UDMMaterialProperty::GetDescription() const
 	return StaticEnum<EDMMaterialPropertyType>()->GetDisplayNameTextByValue(static_cast<int64>(MaterialProperty));
 }
 
+bool UDMMaterialProperty::IsMaterialPin() const
+{
+	return MaterialProperty > EDMMaterialPropertyType::None && MaterialProperty < EDMMaterialPropertyType::Custom1;
+}
+
 void UDMMaterialProperty::ResetInputConnectionMap()
 {
 	if (!IsComponentValid())
@@ -205,12 +216,26 @@ void UDMMaterialProperty::OnSlotAdded(UDMMaterialSlot* InSlot)
 
 UDMMaterialComponent* UDMMaterialProperty::AddComponent(FName InName, UDMMaterialComponent* InComponent)
 {
+	if (!IsValid(InComponent))
+	{
+		InComponent = nullptr;
+	}
+
 	const TObjectPtr<UDMMaterialComponent>* CurrentComponentPtr = Components.Find(InName);
 
 	if (CurrentComponentPtr && IsValid(*CurrentComponentPtr))
 	{
+		if ((*CurrentComponentPtr) == InComponent)
+		{
+			return nullptr;
+		}
+
 		(*CurrentComponentPtr)->SetComponentState(EDMComponentLifetimeState::Removed);
 	}
+	else if (!InComponent)
+	{
+		return nullptr;
+	}
 
 	if (InComponent)
 	{
@@ -256,6 +281,283 @@ UDMMaterialComponent* UDMMaterialProperty::RemoveComponent(FName InName)
 	return nullptr;
 }
 
+void UDMMaterialProperty::GenerateExpressions(const TSharedRef<FDMMaterialBuildState>& InBuildState) const
+{
+	UDynamicMaterialModelEditorOnlyData* EditorOnlyData = GetMaterialModelEditorOnlyData();
+
+	if (!EditorOnlyData)
+	{
+		return;
+	}
+
+	if (EditorOnlyData->GetDomain() == EMaterialDomain::MD_PostProcess && MaterialProperty != EDMMaterialPropertyType::EmissiveColor)
+	{
+		return;
+	}
+
+	// For now we don't have channel remapping!
+	FExpressionInput* MaterialPropertyPtr = InBuildState->GetMaterialProperty(MaterialProperty);
+
+	if (!MaterialPropertyPtr)
+	{
+		return;
+	}
+
+	MaterialPropertyPtr->Expression = nullptr;
+	MaterialPropertyPtr->OutputIndex = 0;
+
+	UMaterialExpression* LastPropertyExpression = nullptr;
+	UDMMaterialSlot* Slot = EditorOnlyData->GetSlotForMaterialProperty(MaterialProperty);
+
+	if (!Slot || Slot->GetLayers().IsEmpty())
+	{
+		return;
+	}
+
+	InBuildState->SetCurrentMaterialProperty(this);
+
+	Slot->GenerateExpressions(InBuildState);
+
+	if (InBuildState->GetSlotExpressions(Slot).IsEmpty())
+	{
+		return;
+	}
+
+	LastPropertyExpression = InBuildState->GetLastSlotPropertyExpression(Slot, MaterialProperty);
+
+	if (!LastPropertyExpression)
+	{
+		return;
+	}
+
+	MaterialPropertyPtr->Expression = LastPropertyExpression;
+
+	if (InputConnectionMap.Channels.IsEmpty() == false)
+	{
+		MaterialPropertyPtr->OutputIndex = InputConnectionMap.Channels[0].OutputIndex;
+	}
+	else
+	{
+		MaterialPropertyPtr->OutputIndex = 0;
+	}
+}
+
+void UDMMaterialProperty::GenerateOpacityExpressions(const TSharedRef<FDMMaterialBuildState>& InBuildState, UDMMaterialSlot* InFromSlot, 
+	EDMMaterialPropertyType InFromProperty, UMaterialExpression*& OutExpression, int32& OutOutputIndex, int32& OutOutputChannel)
+{
+	const TArray<TObjectPtr<UDMMaterialLayerObject>>& SlotLayers = InFromSlot->GetLayers();
+	OutExpression = nullptr;
+
+	for (const TObjectPtr<UDMMaterialLayerObject>& Layer : SlotLayers)
+	{
+		if (!IsValid(Layer))
+		{
+			continue;
+		}
+
+		// Although we are working with masks, if the base is disabled, this is handled by the GenerateExpressions
+		// of the LayerBlend code (to multiply alpha together, instead of maxing it).
+		if (Layer->GetMaterialProperty() != InFromProperty || !Layer->IsEnabled() || !Layer->IsStageEnabled(EDMMaterialLayerStage::Base))
+		{
+			continue;
+		}
+
+		UDMMaterialStage* BaseStage = Layer->GetStage(EDMMaterialLayerStage::Base);
+		UDMMaterialStage* MaskStage = Layer->GetStage(EDMMaterialLayerStage::Mask);
+
+		MaskStage->GenerateExpressions(InBuildState);
+		UDMMaterialStageThroughputLayerBlend* LayerBlend = Cast<UDMMaterialStageThroughputLayerBlend>(MaskStage->GetSource());
+
+		if (!LayerBlend)
+		{
+			continue;
+		}
+
+		UMaterialExpression* MaskOutputExpression;
+		int32 MaskOutputIndex;
+		int32 MaskOutputChannel;
+		LayerBlend->GetMaskOutput(InBuildState, MaskOutputExpression, MaskOutputIndex, MaskOutputChannel);
+
+		if (!MaskOutputExpression)
+		{
+			continue;
+		}
+
+		if (LayerBlend->UsePremultiplyAlpha())
+		{
+			if (UDMMaterialStageSource* Source = BaseStage->GetSource())
+			{
+				UMaterialExpression* LayerAlphaOutputExpression;
+				int32 LayerAlphaOutputIndex;
+				int32 LayerAlphaOutputChannel;
+
+				Source->GetMaskAlphaBlendNode(InBuildState, LayerAlphaOutputExpression, LayerAlphaOutputIndex, LayerAlphaOutputChannel);
+
+				if (LayerAlphaOutputExpression)
+				{
+					UMaterialExpressionMultiply* AlphaMultiply = InBuildState->GetBuildUtils().CreateExpression<UMaterialExpressionMultiply>(UE_DM_NodeComment_Default);
+
+					AlphaMultiply->A.Expression = MaskOutputExpression;
+					AlphaMultiply->A.OutputIndex = MaskOutputIndex;
+					AlphaMultiply->A.Mask = 0;
+
+					if (MaskOutputChannel != FDMMaterialStageConnectorChannel::WHOLE_CHANNEL)
+					{
+						AlphaMultiply->A.Mask = 1;
+						AlphaMultiply->A.MaskR = !!(MaskOutputChannel & FDMMaterialStageConnectorChannel::FIRST_CHANNEL);
+						AlphaMultiply->A.MaskG = !!(MaskOutputChannel & FDMMaterialStageConnectorChannel::SECOND_CHANNEL);
+						AlphaMultiply->A.MaskB = !!(MaskOutputChannel & FDMMaterialStageConnectorChannel::THIRD_CHANNEL);
+						AlphaMultiply->A.MaskA = !!(MaskOutputChannel & FDMMaterialStageConnectorChannel::FOURTH_CHANNEL);
+					}
+
+					AlphaMultiply->B.Expression = LayerAlphaOutputExpression;
+					AlphaMultiply->B.OutputIndex = LayerAlphaOutputIndex;
+					AlphaMultiply->B.Mask = 0;
+
+					if (LayerAlphaOutputChannel != FDMMaterialStageConnectorChannel::WHOLE_CHANNEL)
+					{
+						AlphaMultiply->B.Mask = 1;
+						AlphaMultiply->B.MaskR = !!(LayerAlphaOutputChannel & FDMMaterialStageConnectorChannel::FIRST_CHANNEL);
+						AlphaMultiply->B.MaskG = !!(LayerAlphaOutputChannel & FDMMaterialStageConnectorChannel::SECOND_CHANNEL);
+						AlphaMultiply->B.MaskB = !!(LayerAlphaOutputChannel & FDMMaterialStageConnectorChannel::THIRD_CHANNEL);
+						AlphaMultiply->B.MaskA = !!(LayerAlphaOutputChannel & FDMMaterialStageConnectorChannel::FOURTH_CHANNEL);
+					}
+
+					MaskOutputExpression = AlphaMultiply;
+					MaskOutputIndex = 0;
+					MaskOutputChannel = FDMMaterialStageConnectorChannel::WHOLE_CHANNEL;
+				}
+			}
+		}
+
+		if (OutExpression == nullptr)
+		{
+			OutExpression = MaskOutputExpression;
+
+			// The first output will use the node's output info.
+			OutOutputIndex = MaskOutputIndex;
+			OutOutputChannel = MaskOutputChannel;
+			continue;
+		}
+
+		UMaterialExpressionMax* Max = InBuildState->GetBuildUtils().CreateExpression<UMaterialExpressionMax>(UE_DM_NodeComment_Default);
+		check(Max);
+
+		Max->A.Expression = OutExpression;
+		Max->A.OutputIndex = OutOutputIndex;
+		Max->A.Mask = 0;
+
+		if (OutOutputChannel != FDMMaterialStageConnectorChannel::WHOLE_CHANNEL)
+		{
+			Max->A.Mask = 1;
+			Max->A.MaskR = !!(OutOutputChannel & FDMMaterialStageConnectorChannel::FIRST_CHANNEL);
+			Max->A.MaskG = !!(OutOutputChannel & FDMMaterialStageConnectorChannel::SECOND_CHANNEL);
+			Max->A.MaskB = !!(OutOutputChannel & FDMMaterialStageConnectorChannel::THIRD_CHANNEL);
+			Max->A.MaskA = !!(OutOutputChannel & FDMMaterialStageConnectorChannel::FOURTH_CHANNEL);
+		}
+
+		Max->B.Expression = MaskOutputExpression;
+		Max->B.OutputIndex = MaskOutputIndex;
+		Max->B.Mask = 0;
+
+		if (MaskOutputChannel != FDMMaterialStageConnectorChannel::WHOLE_CHANNEL)
+		{
+			Max->B.Mask = 1;
+			Max->B.MaskR = !!(MaskOutputChannel & FDMMaterialStageConnectorChannel::FIRST_CHANNEL);
+			Max->B.MaskG = !!(MaskOutputChannel & FDMMaterialStageConnectorChannel::SECOND_CHANNEL);
+			Max->B.MaskB = !!(MaskOutputChannel & FDMMaterialStageConnectorChannel::THIRD_CHANNEL);
+			Max->B.MaskA = !!(MaskOutputChannel & FDMMaterialStageConnectorChannel::FOURTH_CHANNEL);
+		}
+
+		OutExpression = Max;
+
+		// If we have to combine, it will use the Max node's output info
+		OutOutputIndex = 0;
+		OutOutputChannel = FDMMaterialStageConnectorChannel::WHOLE_CHANNEL;
+	}
+}
+
+void UDMMaterialProperty::AddAlphaMultiplier(const TSharedRef<FDMMaterialBuildState>& InBuildState) const
+{
+	UDMMaterialValueFloat1* AlphaValue = GetTypedComponent<UDMMaterialValueFloat1>(UDynamicMaterialModelEditorOnlyData::AlphaValueName);
+
+	if (!AlphaValue)
+	{
+		return;
+	}
+
+	FExpressionInput* PropertyInputExpression = InBuildState->GetMaterialProperty(MaterialProperty);
+
+	if (!PropertyInputExpression || !PropertyInputExpression->Expression)
+	{
+		return;
+	}
+
+	AlphaValue->GenerateExpression(InBuildState);
+
+	UMaterialExpression* GlobalOpacityExpression = InBuildState->GetLastValueExpression(AlphaValue);
+
+	if (!GlobalOpacityExpression)
+	{
+		return;
+	}
+
+	UMaterialExpressionMultiply* OpacityMultiply = InBuildState->GetBuildUtils().CreateExpression<UMaterialExpressionMultiply>(UE_DM_NodeComment_Default);
+	OpacityMultiply->A.Expression = PropertyInputExpression->Expression;
+	OpacityMultiply->A.Mask = PropertyInputExpression->Mask;
+	OpacityMultiply->A.MaskR = PropertyInputExpression->MaskR;
+	OpacityMultiply->A.MaskG = PropertyInputExpression->MaskG;
+	OpacityMultiply->A.MaskB = PropertyInputExpression->MaskB;
+	OpacityMultiply->A.MaskA = PropertyInputExpression->MaskA;
+	OpacityMultiply->A.OutputIndex = PropertyInputExpression->OutputIndex;
+
+	OpacityMultiply->B.Expression = GlobalOpacityExpression;
+	OpacityMultiply->B.SetMask(1, 1, 0, 0, 0);
+	OpacityMultiply->B.OutputIndex = 0;
+
+	PropertyInputExpression->Expression = OpacityMultiply;
+}
+
+void UDMMaterialProperty::AddOutputProcessor(const TSharedRef<FDMMaterialBuildState>& InBuildState) const
+{
+	if (!OutputProcessor)
+	{
+		return;
+	}
+
+	FExpressionInput* MaterialPropertyPtr = InBuildState->GetMaterialProperty(MaterialProperty);
+
+	if (!MaterialPropertyPtr)
+	{
+		return;
+	}
+
+	UMaterialExpression* LastPropertyExpression = MaterialPropertyPtr->Expression;
+
+	if (!LastPropertyExpression)
+	{
+		return;
+	}
+
+	UMaterialExpressionMaterialFunctionCall* MaterialFunctionCall = FDMMaterialFunctionLibrary::Get().MakeExpression(
+		InBuildState->GetDynamicMaterial(),
+		OutputProcessor,
+		UE_DM_NodeComment_Default
+	);
+
+	TArrayView<FExpressionInput*> Inputs = MaterialFunctionCall->GetInputsView();
+
+	if (Inputs.IsEmpty())
+	{
+		return;
+	}
+
+	LastPropertyExpression->ConnectExpression(Inputs[0], MaterialPropertyPtr->OutputIndex);
+	MaterialFunctionCall->ConnectExpression(MaterialPropertyPtr, 0);
+
+	MaterialPropertyPtr->OutputIndex = 0;
+}
+
 void UDMMaterialProperty::Update(EDMUpdateType InUpdateType)
 {
 	if (!IsComponentValid())
diff --git a/Engine/Plugins/Experimental/DynamicMaterial/Source/DynamicMaterialEditor/Private/Components/DMMaterialSlot.cpp b/Engine/Plugins/Experimental/DynamicMaterial/Source/DynamicMaterialEditor/Private/Components/DMMaterialSlot.cpp
--- a/Engine/Plugins/Experimental/DynamicMaterial/Source/DynamicMaterialEditor/Private/Components/DMMaterialSlot.cpp
+++ b/Engine/Plugins/Experimental/DynamicMaterial/Source/DynamicMaterialEditor/Private/Components/DMMaterialSlot.cpp
@@ -852,7 +852,7 @@ void UDMMaterialSlot::UpdateBasePreviewMaterialFull(const TSharedRef<FDMMaterial
 
 	for (const TPair<EDMMaterialPropertyType, UDMMaterialProperty*>& Pair : MaterialProperties)
 	{
-		if (UE::DynamicMaterialEditor::Private::IsCustomMaterialProperty(Pair.Key))
+		if (!Pair.Value->IsMaterialPin())
 		{
 			continue;
 		}
diff --git a/Engine/Plugins/Experimental/DynamicMaterial/Source/DynamicMaterialEditor/Private/Components/DMMaterialStageThroughputLayerBlend.cpp b/Engine/Plugins/Experimental/DynamicMaterial/Source/DynamicMaterialEditor/Private/Components/DMMaterialStageThroughputLayerBlend.cpp
--- a/Engine/Plugins/Experimental/DynamicMaterial/Source/DynamicMaterialEditor/Private/Components/DMMaterialStageThroughputLayerBlend.cpp
+++ b/Engine/Plugins/Experimental/DynamicMaterial/Source/DynamicMaterialEditor/Private/Components/DMMaterialStageThroughputLayerBlend.cpp
@@ -1,7 +1,9 @@
 // Copyright Epic Games, Inc. All Rights Reserved.
  
 #include "Components/DMMaterialStageThroughputLayerBlend.h"
+ 
 #include "Components/DMMaterialLayer.h"
+#include "Components/DMMaterialProperty.h"
 #include "Components/DMMaterialSlot.h"
 #include "Components/DMMaterialStage.h"
 #include "Components/DMMaterialSubStage.h"
@@ -16,16 +18,14 @@
 #include "Materials/MaterialExpressionFunctionInput.h"
 #include "Materials/MaterialExpressionMaterialFunctionCall.h"
 #include "Materials/MaterialExpressionMultiply.h"
-#include "Materials/MaterialExpressionSceneTexture.h"
 #include "Materials/MaterialFunctionInterface.h"
 #include "Model/DMMaterialBuildState.h"
 #include "Model/DMMaterialBuildUtils.h"
-#include "Model/DynamicMaterialModel.h"
 #include "Model/DynamicMaterialModelEditorOnlyData.h"
 #include "Utils/DMInputNodeBuilder.h"
 #include "Utils/DMMaterialFunctionLibrary.h"
 #include "Utils/DMPrivate.h"
- 
+
 #define LOCTEXT_NAMESPACE "DMMaterialStageThroughputLayerBlend"
 
 namespace UE::DynamicMaterialEditor::Private
@@ -386,7 +386,7 @@ void UDMMaterialStageThroughputLayerBlend::GenerateMainExpressions(const TShared
 		UMaterialExpression* LastExpression = nullptr;
 		int32 OutputIndex;
 		int32 OutputChannel;
-		ModelEditorOnlyData->GenerateOpacityExpressions(InBuildState, RGBSlot, RGBProperty, LastExpression, OutputIndex, OutputChannel);
+		UDMMaterialProperty::GenerateOpacityExpressions(InBuildState, RGBSlot, RGBProperty, LastExpression, OutputIndex, OutputChannel);
 
 		if (LastExpression)
 		{
diff --git a/Engine/Plugins/Experimental/DynamicMaterial/Source/DynamicMaterialEditor/Private/Components/MaterialProperties/DMMPNormal.cpp b/Engine/Plugins/Experimental/DynamicMaterial/Source/DynamicMaterialEditor/Private/Components/MaterialProperties/DMMPNormal.cpp
--- a/Engine/Plugins/Experimental/DynamicMaterial/Source/DynamicMaterialEditor/Private/Components/MaterialProperties/DMMPNormal.cpp
+++ b/Engine/Plugins/Experimental/DynamicMaterial/Source/DynamicMaterialEditor/Private/Components/MaterialProperties/DMMPNormal.cpp
@@ -1,7 +1,24 @@
 // Copyright Epic Games, Inc. All Rights Reserved.
 
 #include "Components/MaterialProperties/DMMPNormal.h"
+
+#include "Materials/MaterialExpressionMaterialFunctionCall.h"
+#include "Model/DMMaterialBuildState.h"
 #include "Model/DynamicMaterialModelEditorOnlyData.h"
+#include "Utils/DMMaterialFunctionLibrary.h"
+
+namespace UE::DynamicMaterialEditor::Private
+{
+	UMaterialFunctionInterface* GetSafeNormalize()
+	{
+		static UMaterialFunctionInterface* NormalizeBlend = FDMMaterialFunctionLibrary::Get().GetFunction(
+			"SafeNormalize",
+			TEXT("/Script/Engine.MaterialFunction'/Engine/Functions/Engine_MaterialFunctions02/SafeNormalize.SafeNormalize'")
+		);
+
+		return NormalizeBlend;
+	}
+}
 
 UDMMaterialPropertyNormal::UDMMaterialPropertyNormal()
 	: UDMMaterialProperty(
@@ -37,3 +54,40 @@ TEnumAsByte<EMaterialSamplerType> UDMMaterialPropertyNormal::GetTextureSamplerTy
 {
 	return EMaterialSamplerType::SAMPLERTYPE_Normal;
 }
+
+void UDMMaterialPropertyNormal::AddOutputProcessor(const TSharedRef<FDMMaterialBuildState>& InBuildState) const
+{
+	Super::AddOutputProcessor(InBuildState);
+
+	FExpressionInput* MaterialPropertyPtr = InBuildState->GetMaterialProperty(MaterialProperty);
+
+	if (!MaterialPropertyPtr)
+	{
+		return;
+	}
+
+	UMaterialExpression* LastPropertyExpression = MaterialPropertyPtr->Expression;
+
+	if (!LastPropertyExpression)
+	{
+		return;
+	}
+
+	UMaterialExpressionMaterialFunctionCall* MaterialFunctionCall = FDMMaterialFunctionLibrary::Get().MakeExpression(
+		InBuildState->GetDynamicMaterial(),
+		UE::DynamicMaterialEditor::Private::GetSafeNormalize(),
+		UE_DM_NodeComment_Default
+	);
+
+	TArrayView<FExpressionInput*> Inputs = MaterialFunctionCall->GetInputsView();
+
+	if (Inputs.IsEmpty())
+	{
+		return;
+	}
+
+	LastPropertyExpression->ConnectExpression(Inputs[0], MaterialPropertyPtr->OutputIndex);
+	MaterialFunctionCall->ConnectExpression(MaterialPropertyPtr, 0);
+
+	MaterialPropertyPtr->OutputIndex = 0;
+}
diff --git a/Engine/Plugins/Experimental/DynamicMaterial/Source/DynamicMaterialEditor/Private/Components/MaterialProperties/DMMPNormal.h b/Engine/Plugins/Experimental/DynamicMaterial/Source/DynamicMaterialEditor/Private/Components/MaterialProperties/DMMPNormal.h
--- a/Engine/Plugins/Experimental/DynamicMaterial/Source/DynamicMaterialEditor/Private/Components/MaterialProperties/DMMPNormal.h
+++ b/Engine/Plugins/Experimental/DynamicMaterial/Source/DynamicMaterialEditor/Private/Components/MaterialProperties/DMMPNormal.h
@@ -19,5 +19,10 @@ public:
 	virtual bool IsValidForModel(UDynamicMaterialModelEditorOnlyData& InModelEditorOnlyData) const override;
 	virtual UMaterialExpression* GetDefaultInput(const TSharedRef<FDMMaterialBuildState>& InBuildState) const override;
 	virtual TEnumAsByte<EMaterialSamplerType> GetTextureSamplerType() const override;
+
+	/**
+	 * Adds the SafeNormalize function as an extra output processor.
+	 */
+	virtual void AddOutputProcessor(const TSharedRef<FDMMaterialBuildState>& InBuildState) const;
 	//~ End UDMMaterialProperty
 };
diff --git a/Engine/Plugins/Experimental/DynamicMaterial/Source/DynamicMaterialEditor/Private/Components/MaterialStageExpressions/DMMSETextureSampleBase.cpp b/Engine/Plugins/Experimental/DynamicMaterial/Source/DynamicMaterialEditor/Private/Components/MaterialStageExpressions/DMMSETextureSampleBase.cpp
--- a/Engine/Plugins/Experimental/DynamicMaterial/Source/DynamicMaterialEditor/Private/Components/MaterialStageExpressions/DMMSETextureSampleBase.cpp
+++ b/Engine/Plugins/Experimental/DynamicMaterial/Source/DynamicMaterialEditor/Private/Components/MaterialStageExpressions/DMMSETextureSampleBase.cpp
@@ -72,7 +72,7 @@ void UDMMaterialStageExpressionTextureSampleBase::GenerateExpressions(const TSha
 		return;
 	}
 
-	UDMMaterialProperty* CurrentProperty = InBuildState->GetCurrentMaterialProperty();
+	const UDMMaterialProperty* CurrentProperty = InBuildState->GetCurrentMaterialProperty();
 
 	if (!CurrentProperty)
 	{
diff --git a/Engine/Plugins/Experimental/DynamicMaterial/Source/DynamicMaterialEditor/Private/Model/DMMaterialBuildState.cpp b/Engine/Plugins/Experimental/DynamicMaterial/Source/DynamicMaterialEditor/Private/Model/DMMaterialBuildState.cpp
--- a/Engine/Plugins/Experimental/DynamicMaterial/Source/DynamicMaterialEditor/Private/Model/DMMaterialBuildState.cpp
+++ b/Engine/Plugins/Experimental/DynamicMaterial/Source/DynamicMaterialEditor/Private/Model/DMMaterialBuildState.cpp
@@ -60,12 +60,12 @@ UDynamicMaterialModel* FDMMaterialBuildState::GetMaterialModel() const
 	return MaterialModel;
 }
 
-UDMMaterialProperty* FDMMaterialBuildState::GetCurrentMaterialProperty() const
+const UDMMaterialProperty* FDMMaterialBuildState::GetCurrentMaterialProperty() const
 {
 	return CurrentProperty;
 }
 
-void FDMMaterialBuildState::SetCurrentMaterialProperty(UDMMaterialProperty* InProperty)
+void FDMMaterialBuildState::SetCurrentMaterialProperty(const UDMMaterialProperty* InProperty)
 {
 	CurrentProperty = InProperty;
 }
diff --git a/Engine/Plugins/Experimental/DynamicMaterial/Source/DynamicMaterialEditor/Private/Model/DynamicMaterialModelEditorOnlyData.cpp b/Engine/Plugins/Experimental/DynamicMaterial/Source/DynamicMaterialEditor/Private/Model/DynamicMaterialModelEditorOnlyData.cpp
--- a/Engine/Plugins/Experimental/DynamicMaterial/Source/DynamicMaterialEditor/Private/Model/DynamicMaterialModelEditorOnlyData.cpp
+++ b/Engine/Plugins/Experimental/DynamicMaterial/Source/DynamicMaterialEditor/Private/Model/DynamicMaterialModelEditorOnlyData.cpp
@@ -6,8 +6,6 @@
 #include "Components/DMMaterialLayer.h"
 #include "Components/DMMaterialProperty.h"
 #include "Components/DMMaterialSlot.h"
-#include "Components/DMMaterialStage.h"
-#include "Components/DMMaterialStageThroughputLayerBlend.h"
 #include "Components/DMMaterialValue.h"
 #include "Components/DMTextureUV.h"
 #include "Components/MaterialProperties/DMMPAmbientOcclusion.h"
@@ -24,9 +22,6 @@
 #include "Components/MaterialProperties/DMMPSpecular.h"
 #include "Components/MaterialProperties/DMMPTangent.h"
 #include "Components/MaterialProperties/DMMPWorldPositionOffset.h"
-#include "Components/MaterialStageBlends/DMMSBNormal.h"
-#include "Components/MaterialValues/DMMaterialValueFloat1.h"
-#include "Components/MaterialValues/DMMaterialValueFloat2.h"
 #include "CoreGlobals.h"
 #include "DMComponentPath.h"
 #include "DMDefs.h"
@@ -39,16 +34,11 @@
 #include "Material/DynamicMaterialInstance.h"
 #include "Materials/Material.h"
 #include "Materials/MaterialExpressionMaterialFunctionCall.h"
-#include "Materials/MaterialExpressionMax.h"
-#include "Materials/MaterialExpressionMultiply.h"
 #include "Materials/MaterialFunction.h"
 #include "Misc/Guid.h"
 #include "Model/DMMaterialBuildState.h"
-#include "Model/DMMaterialBuildUtils.h"
 #include "Model/DynamicMaterialModel.h"
 #include "UObject/Package.h"
-#include "Utils/DMMaterialFunctionLibrary.h"
-#include "Utils/DMPrivate.h"
 
 #define LOCTEXT_NAMESPACE "MaterialDesignerModel"
 
@@ -516,105 +506,26 @@ void UDynamicMaterialModelEditorOnlyData::BuildMaterial(bool bInDirtyAssets)
 
 	TSharedRef<FDMMaterialBuildState> BuildState = CreateBuildState(MaterialModel->DynamicMaterial, bInDirtyAssets);
 
-	const bool bIsPostProcessMaterial = Domain == EMaterialDomain::MD_PostProcess;
-
 	/**
 	 * Process slots to build base material inputs.
 	 */
 	for (const TPair<EDMMaterialPropertyType, TObjectPtr<UDMMaterialProperty>>& Pair : Properties)
 	{
-		if (UE::DynamicMaterialEditor::Private::IsCustomMaterialProperty(Pair.Key))
-		{
-			continue;
-		}
-
-		if (bIsPostProcessMaterial && Pair.Key != EDMMaterialPropertyType::EmissiveColor)
-		{
-			continue;
-		}
-
-		// For now we don't have channel remapping!
-		FExpressionInput* MaterialPropertyPtr = BuildState->GetMaterialProperty(Pair.Key);
-
-		if (!MaterialPropertyPtr)
-		{
-			continue;
-		}
-
-		MaterialPropertyPtr->Expression = nullptr;
-		MaterialPropertyPtr->OutputIndex = 0;
-
-		UMaterialExpression* LastPropertyExpression = nullptr;
-		UDMMaterialSlot* Slot = GetSlotForMaterialProperty(Pair.Key);
-
-		if (!Slot || Slot->GetLayers().IsEmpty())
+		if (!Pair.Value->IsMaterialPin())
 		{
 			continue;
 		}
 
-		BuildState->SetCurrentMaterialProperty(Pair.Value);
+		Pair.Value->GenerateExpressions(BuildState);
 
-		Slot->GenerateExpressions(BuildState);
-
-		if (BuildState->GetSlotExpressions(Slot).IsEmpty())
-		{
-			continue;
-		}
-
-		LastPropertyExpression = BuildState->GetLastSlotPropertyExpression(Slot, Pair.Key);
-
-		if (!LastPropertyExpression)
-		{
-			continue;
-		}
-
-		MaterialPropertyPtr->Expression = LastPropertyExpression;
-
-		if (Pair.Value->GetInputConnectionMap().Channels.IsEmpty() == false)
-		{
-			MaterialPropertyPtr->OutputIndex = Pair.Value->GetInputConnectionMap().Channels[0].OutputIndex;
-		}
-		else
+		// Global opacity is handled at later
+		if (Pair.Key != EDMMaterialPropertyType::Opacity && Pair.Key != EDMMaterialPropertyType::OpacityMask)
 		{
-			MaterialPropertyPtr->OutputIndex = 0;
-		}
-
-		// Global opacity is handled later
-		if (Pair.Key == EDMMaterialPropertyType::Opacity || Pair.Key == EDMMaterialPropertyType::OpacityMask)
-		{
-			continue;
-		}
-
-		UDMMaterialValueFloat1* AlphaValue = Cast<UDMMaterialValueFloat1>(Pair.Value->GetComponent(AlphaValueName));
-
-		if (!AlphaValue)
-		{
-			continue;
-		}
-
-		AlphaValue->GenerateExpression(BuildState);
-		UMaterialExpression* AlphaValueExpression = BuildState->GetLastValueExpression(AlphaValue);
-
-		if (MaterialPropertyPtr->Expression)
-		{
-			UMaterialExpressionMultiply* OpacityMultiply = BuildState->GetBuildUtils().CreateExpression<UMaterialExpressionMultiply>(UE_DM_NodeComment_Default);
-			OpacityMultiply->A.Expression = MaterialPropertyPtr->Expression;
-			OpacityMultiply->A.Mask = MaterialPropertyPtr->Mask;
-			OpacityMultiply->A.MaskR = MaterialPropertyPtr->MaskR;
-			OpacityMultiply->A.MaskG = MaterialPropertyPtr->MaskG;
-			OpacityMultiply->A.MaskB = MaterialPropertyPtr->MaskB;
-			OpacityMultiply->A.MaskA = MaterialPropertyPtr->MaskA;
-			OpacityMultiply->A.OutputIndex = MaterialPropertyPtr->OutputIndex;
-
-			OpacityMultiply->B.Expression = AlphaValueExpression;
-			OpacityMultiply->B.SetMask(1, 1, 0, 0, 0);
-			OpacityMultiply->B.OutputIndex = 0;
-
-			MaterialPropertyPtr->Expression = OpacityMultiply;
+			Pair.Value->AddAlphaMultiplier(BuildState);
 		}
 	}
 
-	if (!bIsPostProcessMaterial)
+	if (Domain != EMaterialDomain::MD_PostProcess)
 	{
 		/**
 		 * Generate opacity input based on base/emissive if it doesn't already have an input.
@@ -657,7 +568,7 @@ void UDynamicMaterialModelEditorOnlyData::BuildMaterial(bool bInDirtyAssets)
 				UMaterialExpression* OpacityOutputNode;
 				int32 OutputIndex;
 				int32 OutputChannel;
-				GenerateOpacityExpressions(BuildState, OpacitySlot, OpacityProperty, OpacityOutputNode, OutputIndex, OutputChannel);
+				UDMMaterialProperty::GenerateOpacityExpressions(BuildState, OpacitySlot, OpacityProperty, OpacityOutputNode, OutputIndex, OutputChannel);
 
 				if (OpacityOutputNode)
 				{
@@ -683,43 +594,20 @@ void UDynamicMaterialModelEditorOnlyData::BuildMaterial(bool bInDirtyAssets)
 		/**
 		 * Apply global opacity slider after automatic opacity generation
 		 */
-		FExpressionInput* OpacityPropertyPtr = nullptr;
+		UDMMaterialProperty* OpacityProperty = nullptr;
 
-		if (BlendMode == BLEND_Translucent)
+		if (BlendMode == BLEND_Masked)
 		{
-			OpacityPropertyPtr = BuildState->GetMaterialProperty(EDMMaterialPropertyType::Opacity);
+			OpacityProperty = Properties.FindRef(EDMMaterialPropertyType::OpacityMask);
 		}
-		else if (BlendMode == BLEND_Masked)
+		else if (BlendMode != BLEND_Opaque)
 		{
-			OpacityPropertyPtr = BuildState->GetMaterialProperty(EDMMaterialPropertyType::OpacityMask);
+			OpacityProperty = Properties.FindRef(EDMMaterialPropertyType::Opacity);
 		}
 
-		if (OpacityPropertyPtr != nullptr)
+		if (OpacityProperty != nullptr)
 		{
-			if (UDMMaterialValueFloat1* GlobalOpacityValue = MaterialModel->GetTypedGlobalParameterValue<UDMMaterialValueFloat1>(UDynamicMaterialModel::GlobalOpacityValueName))
-			{
-				GlobalOpacityValue->GenerateExpression(BuildState);
-				UMaterialExpression* GlobalOpacityExpression = BuildState->GetLastValueExpression(GlobalOpacityValue);
-
-				if (GlobalOpacityExpression && OpacityPropertyPtr->Expression)
-				{
-					UMaterialExpressionMultiply* OpacityMultiply = BuildState->GetBuildUtils().CreateExpression<UMaterialExpressionMultiply>(UE_DM_NodeComment_Default);
-					OpacityMultiply->A.Expression = OpacityPropertyPtr->Expression;
-					OpacityMultiply->A.Mask = OpacityPropertyPtr->Mask;
-					OpacityMultiply->A.MaskR = OpacityPropertyPtr->MaskR;
-					OpacityMultiply->A.MaskG = OpacityPropertyPtr->MaskG;
-					OpacityMultiply->A.MaskB = OpacityPropertyPtr->MaskB;
-					OpacityMultiply->A.MaskA = OpacityPropertyPtr->MaskA;
-					OpacityMultiply->A.OutputIndex = OpacityPropertyPtr->OutputIndex;
-
-					OpacityMultiply->B.Expression = GlobalOpacityExpression;
-					OpacityMultiply->B.SetMask(1, 1, 0, 0, 0);
-					OpacityMultiply->B.OutputIndex = 0;
-
-					OpacityPropertyPtr->Expression = OpacityMultiply;
-					OpacityPropertyPtr->SetMask(1, 1, 0, 0, 0);
-				}
-			}
+			OpacityProperty->AddAlphaMultiplier(BuildState);
 		}
 	}
 
@@ -728,54 +616,12 @@ void UDynamicMaterialModelEditorOnlyData::BuildMaterial(bool bInDirtyAssets)
 	 */
 	for (const TPair<EDMMaterialPropertyType, TObjectPtr<UDMMaterialProperty>>& Pair : Properties)
 	{
-		if (UE::DynamicMaterialEditor::Private::IsCustomMaterialProperty(Pair.Key))
+		if (!Pair.Value->IsMaterialPin())
 		{
 			continue;
 		}
 
-		UMaterialFunctionInterface* OutputProcessor = Pair.Value->GetOutputProcessor();
-
-		if (!OutputProcessor)
-		{
-			continue;
-		}
-
-		FExpressionInput* MaterialPropertyPtr = BuildState->GetMaterialProperty(Pair.Key);
-
-		if (!MaterialPropertyPtr)
-		{
-			continue;
-		}
-
-		UMaterialExpressionMaterialFunctionCall* MFC = FDMMaterialFunctionLibrary::Get().MakeExpression(
-			BuildState->GetDynamicMaterial(),
-			OutputProcessor,
-			UE_DM_NodeComment_Default
-		);
-
-		TArrayView<FExpressionInput*> Inputs = MFC->GetInputsView();
-
-		if (Inputs.IsEmpty())
-		{
-			continue;
-		}
-
-		UMaterialExpression* LastPropertyExpression = MaterialPropertyPtr->Expression;
-
-		if (!LastPropertyExpression)
-		{
-			LastPropertyExpression = Pair.Value->GetDefaultInput(BuildState);
-
-			if (!LastPropertyExpression)
-			{
-				continue;
-			}
-		}
-
-		LastPropertyExpression->ConnectExpression(Inputs[0], MaterialPropertyPtr->OutputIndex);
-		MFC->ConnectExpression(MaterialPropertyPtr, 0);
-
-		MaterialPropertyPtr->OutputIndex = 0;
+		Pair.Value->AddOutputProcessor(BuildState);
 	}
 
 	State = EDMState::Idle;
@@ -1671,136 +1517,6 @@ void UDynamicMaterialModelEditorOnlyData::OnSlotConnectorsUpdated(UDMMaterialSlo
 	}
 }
 
-void UDynamicMaterialModelEditorOnlyData::GenerateOpacityExpressions(const TSharedRef<FDMMaterialBuildState>& InBuildState, UDMMaterialSlot* InFromSlot,
-	EDMMaterialPropertyType InFromProperty, UMaterialExpression*& OutExpression, int32& OutOutputIndex, int32& OutOutputChannel) const
-{
-	const TArray<TObjectPtr<UDMMaterialLayerObject>>& SlotLayers = InFromSlot->GetLayers();
-	OutExpression = nullptr;
-
-	for (const TObjectPtr<UDMMaterialLayerObject>& Layer : SlotLayers)
-	{
-		// Although we are working with masks, if the base is disabled, this is handled by the GenerateExpressions
-		// of the LayerBlend code (to multiply alpha together, instead of maxing it).
-		if (Layer->GetMaterialProperty() != InFromProperty || !Layer->IsEnabled() || !Layer->IsStageEnabled(EDMMaterialLayerStage::Base))
-		{
-			continue;
-		}
-
-		UDMMaterialStage* BaseStage = Layer->GetStage(EDMMaterialLayerStage::Base);
-		UDMMaterialStage* MaskStage = Layer->GetStage(EDMMaterialLayerStage::Mask);
-
-		MaskStage->GenerateExpressions(InBuildState);
-		UDMMaterialStageThroughputLayerBlend* LayerBlend = Cast<UDMMaterialStageThroughputLayerBlend>(MaskStage->GetSource());
-
-		if (!LayerBlend)
-		{
-			continue;
-		}
-
-		UMaterialExpression* MaskOutputExpression;
-		int32 MaskOutputIndex;
-		int32 MaskOutputChannel;
-		LayerBlend->GetMaskOutput(InBuildState, MaskOutputExpression, MaskOutputIndex, MaskOutputChannel);
-
-		if (!MaskOutputExpression)
-		{
-			continue;
-		}
-
-		if (LayerBlend->UsePremultiplyAlpha())
-		{
-			if (UDMMaterialStageSource* Source = BaseStage->GetSource())
-			{
-				UMaterialExpression* LayerAlphaOutputExpression;
-				int32 LayerAlphaOutputIndex;
-				int32 LayerAlphaOutputChannel;
-
-				Source->GetMaskAlphaBlendNode(InBuildState, LayerAlphaOutputExpression, LayerAlphaOutputIndex, LayerAlphaOutputChannel);
-
-				if (LayerAlphaOutputExpression)
-				{
-					UMaterialExpressionMultiply* AlphaMultiply = InBuildState->GetBuildUtils().CreateExpression<UMaterialExpressionMultiply>(UE_DM_NodeComment_Default);
-
-					AlphaMultiply->A.Expression = MaskOutputExpression;
-					AlphaMultiply->A.OutputIndex = MaskOutputIndex;
-					AlphaMultiply->A.Mask = 0;
-
-					if (MaskOutputChannel != FDMMaterialStageConnectorChannel::WHOLE_CHANNEL)
-					{
-						AlphaMultiply->A.Mask = 1;
-						AlphaMultiply->A.MaskR = !!(MaskOutputChannel & FDMMaterialStageConnectorChannel::FIRST_CHANNEL);
-						AlphaMultiply->A.MaskG = !!(MaskOutputChannel & FDMMaterialStageConnectorChannel::SECOND_CHANNEL);
-						AlphaMultiply->A.MaskB = !!(MaskOutputChannel & FDMMaterialStageConnectorChannel::THIRD_CHANNEL);
-						AlphaMultiply->A.MaskA = !!(MaskOutputChannel & FDMMaterialStageConnectorChannel::FOURTH_CHANNEL);
-					}
-
-					AlphaMultiply->B.Expression = LayerAlphaOutputExpression;
-					AlphaMultiply->B.OutputIndex = LayerAlphaOutputIndex;
-					AlphaMultiply->B.Mask = 0;
-
-					if (LayerAlphaOutputChannel != FDMMaterialStageConnectorChannel::WHOLE_CHANNEL)
-					{
-						AlphaMultiply->B.Mask = 1;
-						AlphaMultiply->B.MaskR = !!(LayerAlphaOutputChannel & FDMMaterialStageConnectorChannel::FIRST_CHANNEL);
-						AlphaMultiply->B.MaskG = !!(LayerAlphaOutputChannel & FDMMaterialStageConnectorChannel::SECOND_CHANNEL);
-						AlphaMultiply->B.MaskB = !!(LayerAlphaOutputChannel & FDMMaterialStageConnectorChannel::THIRD_CHANNEL);
-						AlphaMultiply->B.MaskA = !!(LayerAlphaOutputChannel & FDMMaterialStageConnectorChannel::FOURTH_CHANNEL);
-					}
-
-					MaskOutputExpression = AlphaMultiply;
-					MaskOutputIndex = 0;
-					MaskOutputChannel = FDMMaterialStageConnectorChannel::WHOLE_CHANNEL;
-				}
-			}
-		}
-
-		if (OutExpression == nullptr)
-		{
-			OutExpression = MaskOutputExpression;
-			
-			// The first output will use the node's output info.
-			OutOutputIndex = MaskOutputIndex;
-			OutOutputChannel = MaskOutputChannel;
-			continue;
-		}
-
-		UMaterialExpressionMax* Max = InBuildState->GetBuildUtils().CreateExpression<UMaterialExpressionMax>(UE_DM_NodeComment_Default);
-		check(Max);
-
-		Max->A.Expression = OutExpression;
-		Max->A.OutputIndex = OutOutputIndex;
-		Max->A.Mask = 0;
-
-		if (OutOutputChannel != FDMMaterialStageConnectorChannel::WHOLE_CHANNEL)
-		{
-			Max->A.Mask = 1;
-			Max->A.MaskR = !!(OutOutputChannel & FDMMaterialStageConnectorChannel::FIRST_CHANNEL);
-			Max->A.MaskG = !!(OutOutputChannel & FDMMaterialStageConnectorChannel::SECOND_CHANNEL);
-			Max->A.MaskB = !!(OutOutputChannel & FDMMaterialStageConnectorChannel::THIRD_CHANNEL);
-			Max->A.MaskA = !!(OutOutputChannel & FDMMaterialStageConnectorChannel::FOURTH_CHANNEL);
-		}
-
-		Max->B.Expression = MaskOutputExpression;
-		Max->B.OutputIndex = MaskOutputIndex;
-		Max->B.Mask = 0;
-
-		if (MaskOutputChannel != FDMMaterialStageConnectorChannel::WHOLE_CHANNEL)
-		{
-			Max->B.Mask = 1;
-			Max->B.MaskR = !!(MaskOutputChannel & FDMMaterialStageConnectorChannel::FIRST_CHANNEL);
-			Max->B.MaskG = !!(MaskOutputChannel & FDMMaterialStageConnectorChannel::SECOND_CHANNEL);
-			Max->B.MaskB = !!(MaskOutputChannel & FDMMaterialStageConnectorChannel::THIRD_CHANNEL);
-			Max->B.MaskA = !!(MaskOutputChannel & FDMMaterialStageConnectorChannel::FOURTH_CHANNEL);
-		}
-
-		OutExpression = Max;
-
-		// If we have to combine, it will use the Max node's output info
-		OutOutputIndex = 0;
-		OutOutputChannel = FDMMaterialStageConnectorChannel::WHOLE_CHANNEL;
-	}
-}
-
 void UDynamicMaterialModelEditorOnlyData::ReinitComponents()
 {
 	for (int32 SlotIdx = 0; SlotIdx < Slots.Num(); ++SlotIdx)
diff --git a/Engine/Plugins/Experimental/DynamicMaterial/Source/DynamicMaterialEditor/Private/Utils/DMPrivate.cpp b/Engine/Plugins/Experimental/DynamicMaterial/Source/DynamicMaterialEditor/Private/Utils/DMPrivate.cpp
--- a/Engine/Plugins/Experimental/DynamicMaterial/Source/DynamicMaterialEditor/Private/Utils/DMPrivate.cpp
+++ b/Engine/Plugins/Experimental/DynamicMaterial/Source/DynamicMaterialEditor/Private/Utils/DMPrivate.cpp
@@ -66,11 +66,6 @@ namespace UE::DynamicMaterialEditor::Private
 		}
 	}
 
-	bool IsCustomMaterialProperty(EDMMaterialPropertyType InMaterialProperty)
-	{
-		return (InMaterialProperty >= EDMMaterialPropertyType::Custom1 && InMaterialProperty <= EDMMaterialPropertyType::Custom4);
-	}
-
 	static bool bAllowUIFeedback = false;
 	static FText LogErrorObjectFormat = LOCTEXT("LogErrorObjectFormat", "%s (Source: %s)");
 
diff --git a/Engine/Plugins/Experimental/DynamicMaterial/Source/DynamicMaterialEditor/Public/Components/DMMaterialProperty.h b/Engine/Plugins/Experimental/DynamicMaterial/Source/DynamicMaterialEditor/Public/Components/DMMaterialProperty.h
--- a/Engine/Plugins/Experimental/DynamicMaterial/Source/DynamicMaterialEditor/Public/Components/DMMaterialProperty.h
+++ b/Engine/Plugins/Experimental/DynamicMaterial/Source/DynamicMaterialEditor/Public/Components/DMMaterialProperty.h
@@ -45,11 +45,17 @@ public:
 	EDMMaterialPropertyType GetMaterialProperty() const { return MaterialProperty; }
 
 	/**
-	 * @return The description of this property based on the EDMMaterialPropertyType enum.
+	 * The description of this property based on the EDMMaterialPropertyType enum.
 	 */
 	UFUNCTION(BlueprintPure, Category = "Material Designer")
 	FText GetDescription() const;
 
+	/**
+	 * Whether this property represents a material pin. Only materials properties with a pin will
+	 * be generated by default.
+	 */
+	bool IsMaterialPin() const;
+
 	/**
 	 * @return The value type of this property. Will be either VT_Float1, VT_Float3_RGB or VT_Float3_XYZ.
 	 */
@@ -134,6 +140,12 @@ public:
 	 */
 	virtual UDMMaterialComponent* GetComponent(FName InName) const;
 
+	template<typename InClassName>
+	InClassName* GetTypedComponent(FName InName) const
+	{
+		return Cast<InClassName>(GetComponent(InName));
+	}
+
 	/**
 	 * Removes the component from the component list. Sets component state to removed.
 	 * @param InName The name of the component.
@@ -141,6 +153,27 @@ public:
 	 */
 	virtual UDMMaterialComponent* RemoveComponent(FName InName);
 
+	/**
+	 * Generates the base expressions for this slot.
+	 */
+	virtual void GenerateExpressions(const TSharedRef<FDMMaterialBuildState>& InBuildState) const;
+
+	/**
+	 * Adds expressions to opacity properties based on color channel.
+	 */
+	static void GenerateOpacityExpressions(const TSharedRef<FDMMaterialBuildState>& InBuildState, UDMMaterialSlot* InFromSlot,
+		EDMMaterialPropertyType InFromProperty, UMaterialExpression*& OutExpression, int32& OutOutputIndex, int32& OutOutputChannel);
+
+	/**
+	 * Adds the alpha multiplier.
+	 */
+	virtual void AddAlphaMultiplier(const TSharedRef<FDMMaterialBuildState>& InBuildState) const;
+
+	/**
+	 * Adds the output processor to the end of the current expression chain.
+	 */
+	virtual void AddOutputProcessor(const TSharedRef<FDMMaterialBuildState>& InBuildState) const;
+
 	//~ Begin UDMMaterialComponent
 	virtual void Update(EDMUpdateType InUpdateType) override;
 	virtual bool IsRootComponent() const override { return true; }
diff --git a/Engine/Plugins/Experimental/DynamicMaterial/Source/DynamicMaterialEditor/Public/Model/DMMaterialBuildState.h b/Engine/Plugins/Experimental/DynamicMaterial/Source/DynamicMaterialEditor/Public/Model/DMMaterialBuildState.h
--- a/Engine/Plugins/Experimental/DynamicMaterial/Source/DynamicMaterialEditor/Public/Model/DMMaterialBuildState.h
+++ b/Engine/Plugins/Experimental/DynamicMaterial/Source/DynamicMaterialEditor/Public/Model/DMMaterialBuildState.h
@@ -40,9 +40,9 @@ struct DYNAMICMATERIALEDITOR_API FDMMaterialBuildState : public IDMMaterialBuild
 
 	virtual UDynamicMaterialModel* GetMaterialModel() const override;
 
-	UDMMaterialProperty* GetCurrentMaterialProperty() const;
+	const UDMMaterialProperty* GetCurrentMaterialProperty() const;
 
-	void SetCurrentMaterialProperty(UDMMaterialProperty* InProperty);
+	void SetCurrentMaterialProperty(const UDMMaterialProperty* InProperty);
 
 	bool ShouldDirtyAssets() const { return bDirtyAssets; }
 
@@ -157,7 +157,7 @@ struct DYNAMICMATERIALEDITOR_API FDMMaterialBuildState : public IDMMaterialBuild
 private:
 	UMaterial* DynamicMaterial = nullptr;
 	UDynamicMaterialModel* MaterialModel = nullptr;
-	UDMMaterialProperty* CurrentProperty = nullptr;
+	const UDMMaterialProperty* CurrentProperty = nullptr;
 	bool bDirtyAssets;
 	bool bIgnoreUVs;
 	bool bIsPreviewMaterial;
diff --git a/Engine/Plugins/Experimental/DynamicMaterial/Source/DynamicMaterialEditor/Public/Model/DynamicMaterialModelEditorOnlyData.h b/Engine/Plugins/Experimental/DynamicMaterial/Source/DynamicMaterialEditor/Public/Model/DynamicMaterialModelEditorOnlyData.h
--- a/Engine/Plugins/Experimental/DynamicMaterial/Source/DynamicMaterialEditor/Public/Model/DynamicMaterialModelEditorOnlyData.h
+++ b/Engine/Plugins/Experimental/DynamicMaterial/Source/DynamicMaterialEditor/Public/Model/DynamicMaterialModelEditorOnlyData.h
@@ -174,9 +174,6 @@ public:
 	FDMOnSlotListUpdated& GetOnSlotListUpdateDelegate() { return OnSlotListUpdateDelegate; }
 	FDMOnTextureUVUpdated& GetOnTextureUVUpdateDelegate() { return OnTextureUVUpdateDelegate; }
 
-	void GenerateOpacityExpressions(const TSharedRef<FDMMaterialBuildState>& InBuildState, UDMMaterialSlot* InFromSlot,
-		EDMMaterialPropertyType InFromProperty, UMaterialExpression*& OutExpression, int32& OutOutputIndex, int32& OutOutputChannel) const;
-
 	TSharedRef<FDMMaterialBuildState> CreateBuildState(UMaterial* InMaterialToBuild, bool bInDirtyAssets = true) const;
 
 	bool NeedsWizard() const;
