Material Designer
- Changed Texture UV Scale to Tiling and inverted its effect.

#virtualized

diff --git a/Engine/Plugins/Experimental/DynamicMaterial/Source/DynamicMaterial/Private/Components/DMTextureUV.cpp b/Engine/Plugins/Experimental/DynamicMaterial/Source/DynamicMaterial/Private/Components/DMTextureUV.cpp
--- a/Engine/Plugins/Experimental/DynamicMaterial/Source/DynamicMaterial/Private/Components/DMTextureUV.cpp
+++ b/Engine/Plugins/Experimental/DynamicMaterial/Source/DynamicMaterial/Private/Components/DMTextureUV.cpp
@@ -16,8 +16,8 @@ namespace UE::DynamicMaterial::Private
 	TMap<int32, FName> BaseParameterNames = {
 		{ParamID::PivotX,   FName(TEXT("TextureUV_PivotX"))},
 		{ParamID::PivotY,   FName(TEXT("TextureUV_PivotY"))},
-		{ParamID::ScaleX,   FName(TEXT("TextureUV_ScaleX"))},
-		{ParamID::ScaleY,   FName(TEXT("TextureUV_ScaleY"))},
+		{ParamID::TilingX,   FName(TEXT("TextureUV_TilingX"))},
+		{ParamID::TilingY,   FName(TEXT("TextureUV_TilingY"))},
 		{ParamID::Rotation, FName(TEXT("TextureUV_Rotation"))},
 		{ParamID::OffsetX,  FName(TEXT("TextureUV_OffsetX"))},
 		{ParamID::OffsetY,  FName(TEXT("TextureUV_OffsetY"))},
@@ -28,7 +28,8 @@ enum class EDMTextureUVVersion : int32
 {
 	Initial_Pre_20221102 = 0,
 	Version_22021102 = 1,
-	LatestVersion = Version_22021102
+	ScaleToTiling = 2,
+	LatestVersion = ScaleToTiling
 };
 
 const FGuid UDMTextureUV::GUID(0xFCF57AFB, 0x50764284, 0xB9A9E659, 0xFFA02D33);
@@ -39,13 +40,13 @@ const FString UDMTextureUV::OffsetYPathToken  = FString(TEXT("OffsetY"));
 const FString UDMTextureUV::PivotXPathToken   = FString(TEXT("PivotX"));
 const FString UDMTextureUV::PivotYPathToken   = FString(TEXT("PivotY"));
 const FString UDMTextureUV::RotationPathToken = FString(TEXT("Rotation"));
-const FString UDMTextureUV::ScaleXPathToken   = FString(TEXT("ScaleX"));
-const FString UDMTextureUV::ScaleYPathToken   = FString(TEXT("ScaleY"));
+const FString UDMTextureUV::TilingXPathToken   = FString(TEXT("Tiling"));
+const FString UDMTextureUV::TilingYPathToken   = FString(TEXT("TilingY"));
 
 const FName UDMTextureUV::NAME_Offset     = GET_MEMBER_NAME_CHECKED(UDMTextureUV, Offset);
 const FName UDMTextureUV::NAME_Pivot      = GET_MEMBER_NAME_CHECKED(UDMTextureUV, Pivot);
 const FName UDMTextureUV::NAME_Rotation   = GET_MEMBER_NAME_CHECKED(UDMTextureUV, Rotation);
-const FName UDMTextureUV::NAME_Scale      = GET_MEMBER_NAME_CHECKED(UDMTextureUV, Scale);
+const FName UDMTextureUV::NAME_Tiling      = GET_MEMBER_NAME_CHECKED(UDMTextureUV, Tiling);
 
 #if WITH_EDITOR
 const FName UDMTextureUV::NAME_UVSource = GET_MEMBER_NAME_CHECKED(UDMTextureUV, UVSource);
@@ -57,7 +58,7 @@ const TMap<FName, bool> UDMTextureUV::TextureProperties = {
 	{NAME_Offset,     true},
 	{NAME_Pivot,      true},
 	{NAME_Rotation,   true},
-	{NAME_Scale,      true},
+	{NAME_Tiling,      true},
 	{NAME_bMirrorOnX, false},
 	{NAME_bMirrorOnY, false}
 };
@@ -69,7 +70,7 @@ UDMTextureUV::UDMTextureUV()
 	EditableProperties.Add(NAME_Offset);
 	EditableProperties.Add(NAME_Pivot);
 	EditableProperties.Add(NAME_Rotation);
-	EditableProperties.Add(NAME_Scale);
+	EditableProperties.Add(NAME_Tiling);
 	EditableProperties.Add(NAME_bMirrorOnX);
 	EditableProperties.Add(NAME_bMirrorOnY);
 #endif
@@ -147,20 +148,20 @@ void UDMTextureUV::SetRotation(float InRotation)
 	Update(EDMUpdateType::Value);
 }
 
-void UDMTextureUV::SetScale(const FVector2D& InScale)
+void UDMTextureUV::SetTiling(const FVector2D& InTiling)
 {
 	if (!IsComponentValid())
 	{
 		return;
 	}
 
-	if (FMath::IsNearlyEqual(Scale.X, InScale.X)
-		&& FMath::IsNearlyEqual(Scale.Y, InScale.Y))
+	if (FMath::IsNearlyEqual(Tiling.X, InTiling.X)
+		&& FMath::IsNearlyEqual(Tiling.Y, InTiling.Y))
 	{
 		return;
 	}
 
-	Scale = InScale;
+	Tiling = InTiling;
 
 	Update(EDMUpdateType::Value);
 }
@@ -255,16 +256,16 @@ UDMMaterialParameter* UDMTextureUV::GetMaterialParameter(FName InPropertyName, i
 				break;
 		}
 	}
-	else if (InPropertyName == NAME_Scale)
+	else if (InPropertyName == NAME_Tiling)
 	{
 		switch (InComponent)
 		{
 			case 0:
-				ParameterPtr = MaterialParameters.Find(ParamID::ScaleX);
+				ParameterPtr = MaterialParameters.Find(ParamID::TilingX);
 				break;
 
 			case 1:
-				ParameterPtr = MaterialParameters.Find(ParamID::ScaleY);
+				ParameterPtr = MaterialParameters.Find(ParamID::TilingY);
 				break;
 		}
 	}
@@ -351,8 +352,8 @@ void UDMTextureUV::SetMIDParameters(UMaterialInstanceDynamic* InMID)
 
 	UpdateMID(MaterialParameters[ParamID::PivotX]->GetParameterName(), GetPivot().X);
 	UpdateMID(MaterialParameters[ParamID::PivotY]->GetParameterName(), GetPivot().Y);
-	UpdateMID(MaterialParameters[ParamID::ScaleX]->GetParameterName(), GetScale().X);
-	UpdateMID(MaterialParameters[ParamID::ScaleY]->GetParameterName(), GetScale().Y);
+	UpdateMID(MaterialParameters[ParamID::TilingX]->GetParameterName(), GetTiling().X);
+	UpdateMID(MaterialParameters[ParamID::TilingY]->GetParameterName(), GetTiling().Y);
 	UpdateMID(MaterialParameters[ParamID::Rotation]->GetParameterName(), GetRotation());
 	UpdateMID(MaterialParameters[ParamID::OffsetX]->GetParameterName(), GetOffset().X);
 	UpdateMID(MaterialParameters[ParamID::OffsetY]->GetParameterName(), GetOffset().Y);
@@ -453,14 +454,14 @@ UDMMaterialComponent* UDMTextureUV::GetSubComponentByPath(FDMComponentPath& InPa
 		return GetParameter(ParamID::Rotation);
 	}
 
-	if (InPathSegment.GetToken() == ScaleXPathToken)
+	if (InPathSegment.GetToken() == TilingXPathToken)
 	{
-		return GetParameter(ParamID::ScaleX);
+		return GetParameter(ParamID::TilingX);
 	}
 
-	if (InPathSegment.GetToken() == ScaleYPathToken)
+	if (InPathSegment.GetToken() == TilingYPathToken)
 	{
-		return GetParameter(ParamID::ScaleY);
+		return GetParameter(ParamID::TilingY);
 	}
 
 	return Super::GetSubComponentByPath(InPath, InPathSegment);
@@ -512,8 +513,8 @@ void UDMTextureUV::CreateParameterNames()
 
 	CreateParam(ParamID::PivotX);
 	CreateParam(ParamID::PivotY);
-	CreateParam(ParamID::ScaleX);
-	CreateParam(ParamID::ScaleY);
+	CreateParam(ParamID::TilingX);
+	CreateParam(ParamID::TilingY);
 	CreateParam(ParamID::Rotation);
 	CreateParam(ParamID::OffsetX);
 	CreateParam(ParamID::OffsetY);
@@ -595,7 +596,7 @@ void UDMTextureUV::PostEditChangeProperty(struct FPropertyChangedEvent& Property
 	if (PropertyChangedEvent.MemberProperty->GetFName() == NAME_Offset
 		|| PropertyChangedEvent.MemberProperty->GetFName() == NAME_Pivot
 		|| PropertyChangedEvent.MemberProperty->GetFName() == NAME_Rotation
-		|| PropertyChangedEvent.MemberProperty->GetFName() == NAME_Scale)
+		|| PropertyChangedEvent.MemberProperty->GetFName() == NAME_Tiling)
 	{
 		Update(EDMUpdateType::Value);
 	}
@@ -722,13 +723,21 @@ void UDMTextureUV::Serialize(FArchive& Ar)
 			case static_cast<int32>(EDMTextureUVVersion::Initial_Pre_20221102):
 				Offset.X *= -1;
 				Rotation *= -360.f;
-				Scale = FVector2D(1.f, 1.f) / Scale;
-				TextureUVVersion = static_cast<int32>(EDMTextureUVVersion::LatestVersion);
+				Tiling = FVector2D(1.f, 1.f) / Tiling;
+				++TextureUVVersion;
 #if WITH_EDITORONLY_DATA
 				bNeedsPostLoadValueUpdate = true;
 #endif
 				break;
 
+			case static_cast<int32>(EDMTextureUVVersion::Version_22021102):
+				PRAGMA_DISABLE_DEPRECATION_WARNINGS
+				Tiling.X = (Scale.X != 0.f) ? (1.f / Scale.X) : 1.f;
+				Tiling.Y = (Scale.Y != 0.f) ? (1.f / Scale.Y) : 1.f;
+				PRAGMA_ENABLE_DEPRECATION_WARNINGS
+				++TextureUVVersion;
+				break;
+
 			case static_cast<int32>(EDMTextureUVVersion::LatestVersion):
 				// Do nothing
 				break;
diff --git a/Engine/Plugins/Experimental/DynamicMaterial/Source/DynamicMaterial/Public/Components/DMTextureUV.h b/Engine/Plugins/Experimental/DynamicMaterial/Source/DynamicMaterial/Public/Components/DMTextureUV.h
--- a/Engine/Plugins/Experimental/DynamicMaterial/Source/DynamicMaterial/Public/Components/DMTextureUV.h
+++ b/Engine/Plugins/Experimental/DynamicMaterial/Source/DynamicMaterial/Public/Components/DMTextureUV.h
@@ -24,15 +24,15 @@ namespace UE::DynamicMaterial::ParamID
 	// Individual parameters
 	constexpr int32 PivotX = 0;
 	constexpr int32 PivotY = 1;
-	constexpr int32 ScaleX = 2;
-	constexpr int32 ScaleY = 3;
+	constexpr int32 TilingX = 2;
+	constexpr int32 TilingY = 3;
 	constexpr int32 Rotation = 4;
 	constexpr int32 OffsetX = 5;
 	constexpr int32 OffsetY = 6;
 
 	// Parameter groups
 	constexpr int32 Pivot = PivotX;
-	constexpr int32 Scale = ScaleX;
+	constexpr int32 Tiling = TilingX;
 	//constexpr int32 Rotation = 4; // No additional group value needed
 	constexpr int32 Offset = OffsetX;
 }
@@ -55,15 +55,15 @@ public:
 	static const FName NAME_Offset;
 	static const FName NAME_Pivot;
 	static const FName NAME_Rotation;
-	static const FName NAME_Scale;
+	static const FName NAME_Tiling;
 
 	static const FString OffsetXPathToken;
 	static const FString OffsetYPathToken;
 	static const FString PivotXPathToken;
 	static const FString PivotYPathToken;
 	static const FString RotationPathToken;
-	static const FString ScaleXPathToken;
-	static const FString ScaleYPathToken;
+	static const FString TilingXPathToken;
+	static const FString TilingYPathToken;
 	
 	static const FGuid GUID;
 
@@ -78,7 +78,7 @@ public:
 
 #if WITH_EDITORONLY_DATA
 	UPROPERTY(EditInstanceOnly, BlueprintReadWrite, Category = "Material Designer")
-	bool bLinkScale = true;
+	bool bLinkTiling = true;
 #endif
 
 	UFUNCTION(BlueprintPure, Category = "Material Designer")
@@ -111,10 +111,10 @@ public:
 	void SetRotation(float InRotation);
 
 	UFUNCTION(BlueprintPure, Category = "Material Designer")
-	const FVector2D& GetScale() const { return Scale; }
+	const FVector2D& GetTiling() const { return Tiling; }
 
 	UFUNCTION(BlueprintCallable, Category = "Material Designer")
-	void SetScale(const FVector2D& InScale);
+	void SetTiling(const FVector2D& InTiling);
 
 #if WITH_EDITOR
 	UFUNCTION(BlueprintPure, Category = "Material Designer")
@@ -178,15 +178,19 @@ protected:
 	FVector2D Offset = FVector2D(0.f, 0.f);
 
 	UPROPERTY(EditInstanceOnly, BlueprintReadWrite, Getter = GetPivot, Setter = SetPivot, BlueprintSetter = SetPivot, Category = "Material Designer|Texture UV",
-		meta = (AllowPrivateAccess = "true", ToolTip="Pivot for rotation and scale.", Delta = 0.001))
+		meta = (AllowPrivateAccess = "true", ToolTip="Pivot for rotation and tiling.", Delta = 0.001))
 	FVector2D Pivot = FVector2D(0.5, 0.5);
 
 	UPROPERTY(EditInstanceOnly, BlueprintReadWrite, Getter = GetRotation, Setter = SetRotation, BlueprintSetter = SetRotation, Category = "Material Designer|Texture UV",
 		meta = (AllowPrivateAccess = "true", Delta = 1.0))
 	float Rotation = 0.f;
 
-	UPROPERTY(EditInstanceOnly, BlueprintReadWrite, Getter = GetScale, Setter = SetScale, BlueprintSetter = SetScale, Category = "Material Designer|Texture UV",
+	UPROPERTY(EditInstanceOnly, BlueprintReadWrite, Getter = GetTiling, Setter = SetTiling, BlueprintSetter = SetTiling, Category = "Material Designer|Texture UV",
 		meta = (AllowPrivateAccess = "true", AllowPreserveRatio, VectorRatioMode = 3, Delta = 0.001))
+	FVector2D Tiling = FVector2D(1.f, 1.f);
+
+	UE_DEPRECATED(5.5, "Changed to tiling.")
+	UPROPERTY()
 	FVector2D Scale = FVector2D(1.f, 1.f);
 
 #if WITH_EDITORONLY_DATA
diff --git a/Engine/Plugins/Experimental/DynamicMaterial/Source/DynamicMaterialEditor/Private/Components/MaterialStageInputs/DMMSITextureUV.cpp b/Engine/Plugins/Experimental/DynamicMaterial/Source/DynamicMaterialEditor/Private/Components/MaterialStageInputs/DMMSITextureUV.cpp
--- a/Engine/Plugins/Experimental/DynamicMaterial/Source/DynamicMaterialEditor/Private/Components/MaterialStageInputs/DMMSITextureUV.cpp
+++ b/Engine/Plugins/Experimental/DynamicMaterial/Source/DynamicMaterialEditor/Private/Components/MaterialStageInputs/DMMSITextureUV.cpp
@@ -429,6 +429,12 @@ TArray<UMaterialExpression*> UDMMaterialStageInputTextureUV::CreateTextureUVExpr
 		(int32 ParamGroupId, int32 ParamId, FName InputName, float DefaultValue)
 		{
 			check(InTextureUV->MaterialParameters.Contains(ParamId));
+
+			if (!NameToInputIndex.Contains(InputName))
+			{
+				return;
+			}
+
 			check(NameToInputIndex.Contains(InputName));
 
 			UMaterialExpressionScalarParameter* ParamNode = CreateScalarParameter(InBuildState, InTextureUV->MaterialParameters[ParamId]->GetParameterName());
@@ -445,8 +451,8 @@ TArray<UMaterialExpression*> UDMMaterialStageInputTextureUV::CreateTextureUVExpr
 	CreateParameterExpression(ParamID::Pivot,    ParamID::PivotX,   TEXT("PivotX"),   InTextureUV->GetPivot().X);
 	CreateParameterExpression(ParamID::Pivot,    ParamID::PivotY,   TEXT("PivotY"),   InTextureUV->GetPivot().Y);
 	CreateParameterExpression(ParamID::Rotation, ParamID::Rotation, TEXT("Rotation"), InTextureUV->GetRotation());
-	CreateParameterExpression(ParamID::Scale,    ParamID::ScaleX,   TEXT("ScaleX"),   InTextureUV->GetScale().X);
-	CreateParameterExpression(ParamID::Scale,    ParamID::ScaleY,   TEXT("ScaleY"),   InTextureUV->GetScale().Y);
+	CreateParameterExpression(ParamID::Tiling,    ParamID::TilingX,   TEXT("TilingX"),   InTextureUV->GetTiling().X);
+	CreateParameterExpression(ParamID::Tiling,    ParamID::TilingY,   TEXT("TilingY"),   InTextureUV->GetTiling().Y);
 
 	if (UMaterialExpression* GlobalOffset = InBuildState->GetGlobalExpression(UDynamicMaterialModel::GlobalOffsetValueName))
 	{
diff --git a/Engine/Plugins/Experimental/DynamicMaterial/Source/DynamicMaterialEditor/Private/DetailsPanel/DMValueDetailsRowExtensions.cpp b/Engine/Plugins/Experimental/DynamicMaterial/Source/DynamicMaterialEditor/Private/DetailsPanel/DMValueDetailsRowExtensions.cpp
--- a/Engine/Plugins/Experimental/DynamicMaterial/Source/DynamicMaterialEditor/Private/DetailsPanel/DMValueDetailsRowExtensions.cpp
+++ b/Engine/Plugins/Experimental/DynamicMaterial/Source/DynamicMaterialEditor/Private/DetailsPanel/DMValueDetailsRowExtensions.cpp
@@ -221,7 +221,7 @@ void FDMValueDetailsRowExtensions::FillPropertyRightClickMenu_TextureUV(FToolMen
 
 	if (InPropertyName == UDMTextureUV::NAME_Offset 
 		|| InPropertyName == UDMTextureUV::NAME_Pivot
-		|| InPropertyName == UDMTextureUV::NAME_Scale)
+		|| InPropertyName == UDMTextureUV::NAME_Tiling)
 	{
 		CreateMenuEntry(ParameterNameMenuEntryX, ComponentX);
 		CreateMenuEntry(ParameterNameMenuEntryY, ComponentY);
diff --git a/Engine/Plugins/Experimental/DynamicMaterial/Source/DynamicMaterialEditor/Private/Slate/Properties/Generators/DMTextureUVPropertyRowGenerator.cpp b/Engine/Plugins/Experimental/DynamicMaterial/Source/DynamicMaterialEditor/Private/Slate/Properties/Generators/DMTextureUVPropertyRowGenerator.cpp
--- a/Engine/Plugins/Experimental/DynamicMaterial/Source/DynamicMaterialEditor/Private/Slate/Properties/Generators/DMTextureUVPropertyRowGenerator.cpp
+++ b/Engine/Plugins/Experimental/DynamicMaterial/Source/DynamicMaterialEditor/Private/Slate/Properties/Generators/DMTextureUVPropertyRowGenerator.cpp
@@ -62,7 +62,7 @@ void FDMTextureUVPropertyRowGenerator::AddComponentProperties(const TSharedRef<S
 
 	AddPropertyRow(InComponentEditWidget, TextureUV, UDMTextureUV::NAME_Offset, InOutPropertyRows);
 	AddPropertyRow(InComponentEditWidget, TextureUV, UDMTextureUV::NAME_Rotation, InOutPropertyRows);
-	AddPropertyRow(InComponentEditWidget, TextureUV, UDMTextureUV::NAME_Scale, InOutPropertyRows);
+	AddPropertyRow(InComponentEditWidget, TextureUV, UDMTextureUV::NAME_Tiling, InOutPropertyRows);
 	AddPropertyRow(InComponentEditWidget, TextureUV, UDMTextureUV::NAME_Pivot, InOutPropertyRows);
 	AddPropertyRow(InComponentEditWidget, TextureUV, UDMTextureUV::NAME_bMirrorOnX, InOutPropertyRows);
 	AddPropertyRow(InComponentEditWidget, TextureUV, UDMTextureUV::NAME_bMirrorOnY, InOutPropertyRows);
@@ -88,7 +88,7 @@ void FDMTextureUVPropertyRowGenerator::AddPopoutComponentProperties(const TShare
 
 	AddPropertyRow(InParentWidget, TextureUV, UDMTextureUV::NAME_Offset, InOutPropertyRows);
 	AddPropertyRow(InParentWidget, TextureUV, UDMTextureUV::NAME_Rotation, InOutPropertyRows);
-	AddPropertyRow(InParentWidget, TextureUV, UDMTextureUV::NAME_Scale, InOutPropertyRows);
+	AddPropertyRow(InParentWidget, TextureUV, UDMTextureUV::NAME_Tiling, InOutPropertyRows);
 	AddPropertyRow(InParentWidget, TextureUV, UDMTextureUV::NAME_Pivot, InOutPropertyRows);
 	AddPropertyRow(InParentWidget, TextureUV, UDMTextureUV::NAME_bMirrorOnX, InOutPropertyRows);
 	AddPropertyRow(InParentWidget, TextureUV, UDMTextureUV::NAME_bMirrorOnY, InOutPropertyRows);
@@ -133,7 +133,7 @@ void UE::DynamicMaterialEditor::Private::AddVisualizerRow(const TSharedRef<SWidg
 
 	FDMPropertyHandle VisualizerHandle;
 	VisualizerHandle.NameOverride = LOCTEXT("Visualizer", "UV Visualizer");
-	VisualizerHandle.NameToolTipOverride = LOCTEXT("VisualizerToolTip", "A graphical Texture UV editor.\n\n- Offset Mode: Change the Texture UV offset.\n- Pivot Mode: Change the Texture UV pivot, rotation and scale.\n\nControl+click to reset values to default.");
+	VisualizerHandle.NameToolTipOverride = LOCTEXT("VisualizerToolTip", "A graphical Texture UV editor.\n\n- Offset Mode: Change the Texture UV offset.\n- Pivot Mode: Change the Texture UV pivot, rotation and tiling.\n\nControl+click to reset values to default.");
 	VisualizerHandle.ValueName = FName(*InTextureUV->GetComponentPath());
 	VisualizerHandle.ValueWidget = SNew(SDMTextureUVVisualizerProperty, Stage, InTextureUV);
 	VisualizerHandle.CategoryOverrideName = TEXT("Texture UV");
@@ -209,9 +209,9 @@ bool UE::DynamicMaterialEditor::Private::CanResetTextureUVPropertyToDefault(TSha
 		return DefaultObject->GetRotation() != PropertyObject->GetRotation();
 	}
 
-	if (PropertyName == UDMTextureUV::NAME_Scale)
+	if (PropertyName == UDMTextureUV::NAME_Tiling)
 	{
-		return !DefaultObject->GetScale().Equals(PropertyObject->GetScale());
+		return !DefaultObject->GetTiling().Equals(PropertyObject->GetTiling());
 	}
 
 	return false;
diff --git a/Engine/Plugins/Experimental/DynamicMaterial/Source/DynamicMaterialEditor/Private/Slate/Properties/SDMTextureUVVisualizer.cpp b/Engine/Plugins/Experimental/DynamicMaterial/Source/DynamicMaterialEditor/Private/Slate/Properties/SDMTextureUVVisualizer.cpp
--- a/Engine/Plugins/Experimental/DynamicMaterial/Source/DynamicMaterialEditor/Private/Slate/Properties/SDMTextureUVVisualizer.cpp
+++ b/Engine/Plugins/Experimental/DynamicMaterial/Source/DynamicMaterialEditor/Private/Slate/Properties/SDMTextureUVVisualizer.cpp
@@ -53,7 +53,7 @@ SDMTextureUVVisualizer::SDMTextureUVVisualizer()
 	, ScrubbingStartAbsoluteMouse(FVector2f::ZeroVector)
 	, HandleAxis(EHandleAxis::None)
 	, ValueStart(FVector2D::ZeroVector)
-	, bInvertScale(false)
+	, bInvertTiling(false)
 {
 	Brush.SetUVRegion(FBox2f(FVector2f(0, 0), FVector2f(1, 1)));
 }
@@ -339,7 +339,7 @@ int32 SDMTextureUVVisualizer::OnPaint(const FPaintArgs& InArgs, const FGeometry&
 
 	if (bPivotEditMode)
 	{
-		/** Scale Handles */
+		/** Tiling Handles */
 		float RadiusAtAngle[16];
 
 		for (int32 Direction = 0; Direction < 16; ++Direction)
@@ -350,28 +350,28 @@ int32 SDMTextureUVVisualizer::OnPaint(const FPaintArgs& InArgs, const FGeometry&
 		DrawRotatedBorderBox(
 			FVector2f(0.f, RadiusAtAngle[0]),
 			FVector2f(TextureUVVisualizerLargeRadius, TextureUVVisualizerLargeRadius),
-			(ScrubbingMode == EScrubbingMode::Scale && HandleAxis == EHandleAxis::Y) ? HighlightColor : NormalColor,
+			(ScrubbingMode == EScrubbingMode::Tiling && HandleAxis == EHandleAxis::Y) ? HighlightColor : NormalColor,
 			RotationRadians
 		);
 
 		DrawRotatedBorderBox(
 			FVector2f(RadiusAtAngle[4], 0.f),
 			FVector2f(TextureUVVisualizerLargeRadius, TextureUVVisualizerLargeRadius),
-			(ScrubbingMode == EScrubbingMode::Scale && HandleAxis == EHandleAxis::X) ? HighlightColor : NormalColor,
+			(ScrubbingMode == EScrubbingMode::Tiling && HandleAxis == EHandleAxis::X) ? HighlightColor : NormalColor,
 			RotationRadians
 		);
 
 		DrawRotatedBorderBox(
 			FVector2f(0.f, -RadiusAtAngle[8]),
 			FVector2f(TextureUVVisualizerLargeRadius, TextureUVVisualizerLargeRadius),
-			(ScrubbingMode == EScrubbingMode::Scale && HandleAxis == EHandleAxis::Y) ? HighlightColor : NormalColor,
+			(ScrubbingMode == EScrubbingMode::Tiling && HandleAxis == EHandleAxis::Y) ? HighlightColor : NormalColor,
 			RotationRadians
 		);
 
 		DrawRotatedBorderBox(
 			FVector2f(-RadiusAtAngle[12], 0.f),
 			FVector2f(TextureUVVisualizerLargeRadius, TextureUVVisualizerLargeRadius),
-			(ScrubbingMode == EScrubbingMode::Scale && HandleAxis == EHandleAxis::X) ? HighlightColor : NormalColor,
+			(ScrubbingMode == EScrubbingMode::Tiling && HandleAxis == EHandleAxis::X) ? HighlightColor : NormalColor,
 			RotationRadians
 		);
 
@@ -469,7 +469,7 @@ FVector2f SDMTextureUVVisualizer::ApplyTextureUVTransform(const FVector2f& InUV)
 	Offset.Y *= -1.f;
 
 	const float Rotation = TextureUV->GetRotation();
-	const FVector2f& Scale = static_cast<FVector2f>(TextureUV->GetScale());
+	const FVector2f& Tiling = static_cast<FVector2f>(TextureUV->GetTiling());
 	const FVector2f& Pivot = static_cast<FVector2f>(TextureUV->GetPivot());
 
 	FVector2f TransformedUV = InUV;
@@ -480,16 +480,16 @@ FVector2f SDMTextureUVVisualizer::ApplyTextureUVTransform(const FVector2f& InUV)
 		TransformedUV = TransformedUV.GetRotated(Rotation);
 	}
 
-	TransformedUV *= Scale;
+	TransformedUV /= Tiling;
 	TransformedUV += Pivot;
 
 	if (!FMath::IsNearlyZero(Rotation))
 	{
-		TransformedUV += (Offset * Scale).GetRotated(Rotation);
+		TransformedUV += (Offset / Tiling).GetRotated(Rotation);
 	}
 	else
 	{
-		TransformedUV += Offset * Scale;
+		TransformedUV += Offset / Tiling;
 	}
 
 	return TransformedUV;
@@ -544,20 +544,20 @@ float SDMTextureUVVisualizer::GetCircleHandleRadiusAtAngle(float InAngle) const
 		return BaseDistance;
 	}
 
-	const FVector2D& Scale2D = TextureUV->GetScale();
+	const FVector2D& Tiling = TextureUV->GetTiling();
 
-	if (FMath::IsNearlyEqual(Scale2D.X, Scale2D.Y))
+	if (FMath::IsNearlyEqual(Tiling.X, Tiling.Y))
 	{
-		return BaseDistance * Scale2D.X;
+		return BaseDistance / Tiling.X;
 	}
 
 	const float RadianAngle = FMath::DegreesToRadians(InAngle);
 
-	const FVector2f Scale = static_cast<FVector2f>(Scale2D);
+	const FVector2f TilingFloat = static_cast<FVector2f>(Tiling);
 
 	const FVector2f Offset = {
-		FMath::Sin(RadianAngle) * Scale.X,
-		FMath::Cos(RadianAngle) * Scale.Y
+		FMath::Sin(RadianAngle) / TilingFloat.X,
+		FMath::Cos(RadianAngle) / TilingFloat.Y
 	};
 
 	return BaseDistance * Offset.Size();
@@ -633,12 +633,12 @@ SDMTextureUVVisualizer::EHandleAxis SDMTextureUVVisualizer::GetCircleHandleAxis(
 
 	const FVector2f HandleOffset = InAbsolutePosition - GetAbsolutePivotLocation();
 	const float DistanceFromHandle = HandleOffset.Size();
-	const FVector2D& Scale2D = TextureUV->GetScale();
+	const FVector2D& Tiling2D = TextureUV->GetTiling();
 	// We manage angle clockwise from +Y axis. Atan2 handles it anti-clockwise from +X axis.
 	float Angle = 90.f - FMath::RadiansToDegrees(FMath::Atan2(HandleOffset.Y, HandleOffset.X)) + TextureUV->GetRotation();
 	float DistanceFromCircleHandle;
 
-	if (FMath::IsNearlyEqual(Scale2D.X, Scale2D.Y))
+	if (FMath::IsNearlyEqual(Tiling2D.X, Tiling2D.Y))
 	{
 		DistanceFromCircleHandle = FMath::Abs(DistanceFromHandle - GetCircleHandleBaseRadius());
 	}
@@ -796,7 +796,7 @@ bool SDMTextureUVVisualizer::TryClickCircleHandle(const FVector2f& InMousePositi
 	const FVector2f AbsolutePivotLocation = GetAbsolutePivotLocation();
 
 	// When the uv is rotated, do the opposite action
-	const bool bRegularInvertScale = Rotation <= 90.f || Rotation > 270.f;
+	const bool bRegularInvertTiling = Rotation <= 90.f || Rotation > 270.f;
 
 	// Work out if the mouse is over part of the circle handle.
 	EHandleAxis Axis = GetCircleHandleAxis(InMousePosition);
@@ -827,16 +827,16 @@ bool SDMTextureUVVisualizer::TryClickCircleHandle(const FVector2f& InMousePositi
 		{
 			if (!bInResetToDefault)
 			{
-				SetScrubbingMode(EScrubbingMode::Scale, Axis);
-				bInvertScale = (InMousePosition.X <= AbsolutePivotLocation.X) == bRegularInvertScale;
+				SetScrubbingMode(EScrubbingMode::Tiling, Axis);
+				bInvertTiling = (InMousePosition.X <= AbsolutePivotLocation.X) == bRegularInvertTiling;
 			}
 			else
 			{
-				FScopedTransaction Transaction(LOCTEXT("ResetScaleX", "Reset Scale X to Default."));
+				FScopedTransaction Transaction(LOCTEXT("ResetTilingeX", "Reset Tiling X to Default."));
 				TextureUV->Modify();
-				FVector2D NewScale = TextureUV->GetScale();
-				NewScale.X = 1.0;
-				TextureUV->SetScale(NewScale);
+				FVector2D NewTiling = TextureUV->GetTiling();
+				NewTiling.X = 1.0;
+				TextureUV->SetTiling(NewTiling);
 			}
 
 			break;
@@ -846,16 +846,16 @@ bool SDMTextureUVVisualizer::TryClickCircleHandle(const FVector2f& InMousePositi
 		{
 			if (!bInResetToDefault)
 			{
-				SetScrubbingMode(EScrubbingMode::Scale, Axis);
-				bInvertScale = (InMousePosition.Y <= AbsolutePivotLocation.Y) == bRegularInvertScale;
+				SetScrubbingMode(EScrubbingMode::Tiling, Axis);
+				bInvertTiling = (InMousePosition.Y <= AbsolutePivotLocation.Y) == bRegularInvertTiling;
 			}
 			else
 			{
-				FScopedTransaction Transaction(LOCTEXT("ResetScaleY", "Reset Scale Y to Default."));
+				FScopedTransaction Transaction(LOCTEXT("ResetTilingY", "Reset Tiling Y to Default."));
 				TextureUV->Modify();
-				FVector2D NewScale = TextureUV->GetScale();
-				NewScale.Y = 1.0;
-				TextureUV->SetScale(NewScale);
+				FVector2D NewTiling = TextureUV->GetTiling();
+				NewTiling.Y = 1.0;
+				TextureUV->SetTiling(NewTiling);
 			}
 
 			break;
@@ -932,10 +932,10 @@ void SDMTextureUVVisualizer::SetScrubbingMode(EScrubbingMode InMode, EHandleAxis
 			break;
 		}
 
-		case EScrubbingMode::Scale:
-			ValueStart = TextureUV->GetScale();
-			TextureProperty = UDMTextureUV::StaticClass()->FindPropertyByName(UDMTextureUV::NAME_Scale);
-			UE_LOG(LogDynamicMaterialEditor, Verbose, TEXT("Started Scale mode"));
+		case EScrubbingMode::Tiling:
+			ValueStart = TextureUV->GetTiling();
+			TextureProperty = UDMTextureUV::StaticClass()->FindPropertyByName(UDMTextureUV::NAME_Tiling);
+			UE_LOG(LogDynamicMaterialEditor, Verbose, TEXT("Started Tiling mode"));
 			break;
 
 		case EScrubbingMode::Pivot:
@@ -1004,8 +1004,8 @@ void SDMTextureUVVisualizer::UpdateScrub()
 			UpdateScrub_Rotation();
 			break;
 
-		case EScrubbingMode::Scale:
-			UpdateScrub_Scale();
+		case EScrubbingMode::Tiling:
+			UpdateScrub_Tiling();
 			break;
 
 		case EScrubbingMode::Pivot:
@@ -1045,7 +1045,7 @@ void SDMTextureUVVisualizer::UpdateScrub_Offset()
 		OffsetChange = OffsetChange.GetRotated(-Rotation);
 	}
 
-	OffsetChange /= TextureUV->GetScale();
+	OffsetChange *= TextureUV->GetTiling();
 
 	if (HandleAxis == EHandleAxis::X)
 	{
@@ -1097,7 +1097,7 @@ void SDMTextureUVVisualizer::UpdateScrub_Rotation()
 	TextureUV->PostEditChangeProperty(ChangedEvent);
 }
 
-void SDMTextureUVVisualizer::UpdateScrub_Scale()
+void SDMTextureUVVisualizer::UpdateScrub_Tiling()
 {
 	UDMTextureUV* TextureUV = TextureUVWeak.Get();
 
@@ -1109,66 +1109,66 @@ void SDMTextureUVVisualizer::UpdateScrub_Scale()
 	const FVector2f MouseOffset = FSlateApplication::Get().GetCursorPos() - ScrubbingStartAbsoluteMouse;
 	const FBox2d BrushUV = Brush.GetUVRegion();
 
-	FVector2D ScaleChange = static_cast<FVector2D>(MouseOffset / CurrentAbsoluteSize);
-	ScaleChange.X *= (BrushUV.Max.X - BrushUV.Min.X);
-	ScaleChange.Y *= (BrushUV.Max.Y - BrushUV.Min.Y);
-	ScaleChange /= TextureUV->GetScale();
+	FVector2D TilingChange = static_cast<FVector2D>(MouseOffset / CurrentAbsoluteSize);
+	TilingChange.X *= (BrushUV.Max.X - BrushUV.Min.X);
+	TilingChange.Y *= (BrushUV.Max.Y - BrushUV.Min.Y);
+	TilingChange /= TextureUV->GetTiling();
 
-	if (bInvertScale)
+	if (!bInvertTiling)
 	{
-		ScaleChange *= -1.f;
+		TilingChange *= -1.f;
 	}
 
 	const float Rotation = TextureUV->GetRotation();
 
 	if (!FMath::IsNearlyZero(Rotation))
 	{
-		ScaleChange = ScaleChange.GetRotated(-Rotation);
+		TilingChange = TilingChange.GetRotated(-Rotation);
 	}
 
 	if (HandleAxis == EHandleAxis::X)
 	{
-		ScaleChange.Y = 0;
+		TilingChange.Y = 0;
 	}
 	else
 	{
-		ScaleChange.Y *= -1.0;
+		TilingChange.Y *= -1.0;
 	}
 
 	if (HandleAxis == EHandleAxis::Y)
 	{
-		ScaleChange.X = 0;
+		TilingChange.X = 0;
 	}
 
-	FVector2D NewScale = ValueStart;
+	FVector2D NewTiling = ValueStart;
 
 	if (HandleAxis == EHandleAxis::X)
 	{
-		if (ScaleChange.X > 0)
+		if (TilingChange.X > 0)
 		{
-			NewScale.X = FMath::Max(0.001, NewScale.X * (1.f + ScaleChange.X));
+			NewTiling.X = FMath::Max(0.001, NewTiling.X * (1.f + TilingChange.X));
 		}
 		else
 		{
-			NewScale.X = FMath::Max(0.001, NewScale.X / (1.f - ScaleChange.X));
+			NewTiling.X = FMath::Max(0.001, NewTiling.X / (1.f - TilingChange.X));
 		}
 	}
 	else
 	{
-		if (ScaleChange.Y > 0)
+		if (TilingChange.Y > 0)
 		{
-			NewScale.Y = FMath::Max(0.001, NewScale.Y / (1.f + ScaleChange.Y));
+			NewTiling.Y = FMath::Max(0.001, NewTiling.Y / (1.f + TilingChange.Y));
 		}
 		else
 		{
-			NewScale.Y = FMath::Max(0.001, NewScale.Y * (1.f - ScaleChange.Y));
+			NewTiling.Y = FMath::Max(0.001, NewTiling.Y * (1.f - TilingChange.Y));
 		}
 	}
 
-	TextureUV->SetScale(NewScale);
+	TextureUV->SetTiling(NewTiling);
 
 	FPropertyChangedEvent ChangedEvent(
-		UDMTextureUV::StaticClass()->FindPropertyByName(UDMTextureUV::NAME_Scale),
+		UDMTextureUV::StaticClass()->FindPropertyByName(UDMTextureUV::NAME_Tiling),
 		EPropertyChangeType::Interactive,
 		{TextureUV}
 	);
diff --git a/Engine/Plugins/Experimental/DynamicMaterial/Source/DynamicMaterialEditor/Private/Slate/Properties/SDMTextureUVVisualizer.h b/Engine/Plugins/Experimental/DynamicMaterial/Source/DynamicMaterialEditor/Private/Slate/Properties/SDMTextureUVVisualizer.h
--- a/Engine/Plugins/Experimental/DynamicMaterial/Source/DynamicMaterialEditor/Private/Slate/Properties/SDMTextureUVVisualizer.h
+++ b/Engine/Plugins/Experimental/DynamicMaterial/Source/DynamicMaterialEditor/Private/Slate/Properties/SDMTextureUVVisualizer.h
@@ -29,7 +29,7 @@ public:
 		None,
 		Offset,
 		Rotation,
-		Scale,
+		Tiling,
 		Pivot
 	};
 
@@ -87,7 +87,7 @@ protected:
 	FVector2f ScrubbingStartAbsoluteMouse;
 	EHandleAxis HandleAxis;
 	FVector2D ValueStart;
-	bool bInvertScale;
+	bool bInvertTiling;
 	TSharedPtr<FScopedTransaction> ScrubbingTransaction;
 
 	void OnStageUpdated(UDMMaterialComponent* InComponent, EDMUpdateType InUpdateType);
@@ -128,6 +128,6 @@ protected:
 	void UpdateScrub();
 	void UpdateScrub_Offset();
 	void UpdateScrub_Rotation();
-	void UpdateScrub_Scale();
+	void UpdateScrub_Tiling();
 	void UpdateScrub_Pivot();
 };
diff --git a/Engine/Plugins/Experimental/DynamicMaterial/Source/DynamicMaterialEditor/Private/Slate/Properties/SDMTextureUVVisualizerPopout.cpp b/Engine/Plugins/Experimental/DynamicMaterial/Source/DynamicMaterialEditor/Private/Slate/Properties/SDMTextureUVVisualizerPopout.cpp
--- a/Engine/Plugins/Experimental/DynamicMaterial/Source/DynamicMaterialEditor/Private/Slate/Properties/SDMTextureUVVisualizerPopout.cpp
+++ b/Engine/Plugins/Experimental/DynamicMaterial/Source/DynamicMaterialEditor/Private/Slate/Properties/SDMTextureUVVisualizerPopout.cpp
@@ -255,7 +255,7 @@ TSharedRef<SWidget> SDMTextureUVVisualizerPopout::CreatePropertyWidget(UDMTextur
 			.Padding(FVector2D(10.f, 3.f))
 			.IsChecked(this, &SDMTextureUVVisualizerPopout::GetModeCheckBoxState, /* Is Pivot */ true)
 			.OnCheckStateChanged(this, &SDMTextureUVVisualizerPopout::OnModeCheckBoxStateChanged, /* Is Pivot */ true)
-			.ToolTipText(LOCTEXT("VisualizerPivotToolTip", "Allows changing of the UV pivot, rotation and scale."))
+			.ToolTipText(LOCTEXT("VisualizerPivotToolTip", "Allows changing of the UV pivot, rotation and tiling."))
 			.Content()
 			[
 				SNew(STextBlock)
diff --git a/Engine/Plugins/Experimental/DynamicMaterial/Source/DynamicMaterialEditor/Private/Slate/Properties/SDMTextureUVVisualizerProperty.cpp b/Engine/Plugins/Experimental/DynamicMaterial/Source/DynamicMaterialEditor/Private/Slate/Properties/SDMTextureUVVisualizerProperty.cpp
--- a/Engine/Plugins/Experimental/DynamicMaterial/Source/DynamicMaterialEditor/Private/Slate/Properties/SDMTextureUVVisualizerProperty.cpp
+++ b/Engine/Plugins/Experimental/DynamicMaterial/Source/DynamicMaterialEditor/Private/Slate/Properties/SDMTextureUVVisualizerProperty.cpp
@@ -100,7 +100,7 @@ void SDMTextureUVVisualizerProperty::Construct(const FArguments& InArgs, UDMMate
 				.Padding(FVector2D(5.f, 3.f))
 				.IsChecked(this, &SDMTextureUVVisualizerProperty::GetModeCheckBoxState, /* Is Pivot */ true)
 				.OnCheckStateChanged(this, &SDMTextureUVVisualizerProperty::OnModeCheckBoxStateChanged, /* Is Pivot */ true)
-				.ToolTipText(LOCTEXT("VisualizerPivotToolTip", "Allows changing of the UV pivot, rotation and scale."))
+				.ToolTipText(LOCTEXT("VisualizerPivotToolTip", "Allows changing of the UV pivot, rotation and tiling."))
 				.Content()
 				[
 					SNew(STextBlock)
