Material Designer
- New multi-type default slot values (texture/color)
- Textures and colors (linear colors) added to a slot will use the default slot's texture or rgb value.

diff --git a/Engine/Plugins/Experimental/DynamicMaterial/Source/DynamicMaterialEditor/Private/Components/DMMaterialLayer.cpp b/Engine/Plugins/Experimental/DynamicMaterial/Source/DynamicMaterialEditor/Private/Components/DMMaterialLayer.cpp
--- a/Engine/Plugins/Experimental/DynamicMaterial/Source/DynamicMaterialEditor/Private/Components/DMMaterialLayer.cpp
+++ b/Engine/Plugins/Experimental/DynamicMaterial/Source/DynamicMaterialEditor/Private/Components/DMMaterialLayer.cpp
@@ -366,9 +366,14 @@ bool UDMMaterialLayerObject::SetStage(EDMMaterialLayerStage InStageType, UDMMate
 
 	Stages[Index] = InStage;
 
-	if (IsComponentAdded())
+	if (IsValid(InStage))
 	{
-		Stages[Index]->SetComponentState(EDMComponentLifetimeState::Added);
+		InStage->Rename(nullptr, this, UE::DynamicMaterial::RenameFlags);
+
+		if (IsComponentAdded())
+		{
+			InStage->SetComponentState(EDMComponentLifetimeState::Added);
+		}
 	}
 
 	return true;
diff --git a/Engine/Plugins/Experimental/DynamicMaterial/Source/DynamicMaterialEditor/Private/Components/DMMaterialProperty.cpp b/Engine/Plugins/Experimental/DynamicMaterial/Source/DynamicMaterialEditor/Private/Components/DMMaterialProperty.cpp
--- a/Engine/Plugins/Experimental/DynamicMaterial/Source/DynamicMaterialEditor/Private/Components/DMMaterialProperty.cpp
+++ b/Engine/Plugins/Experimental/DynamicMaterial/Source/DynamicMaterialEditor/Private/Components/DMMaterialProperty.cpp
@@ -9,9 +9,11 @@
 #include "Components/DMMaterialSubStage.h"
 #include "Components/MaterialStageBlends/DMMSBNormal.h"
 #include "Components/MaterialStageExpressions/DMMSETextureSample.h"
+#include "Components/MaterialStageExpressions/DMMSETextureSampleBase.h"
 #include "Components/MaterialStageInputs/DMMSIExpression.h"
 #include "Components/MaterialStageInputs/DMMSIValue.h"
 #include "Components/MaterialValues/DMMaterialValueFloat1.h"
+#include "Components/MaterialValues/DMMaterialValueFloat3RGB.h"
 #include "Components/MaterialValues/DMMaterialValueTexture.h"
 #include "DMComponentPath.h"
 #include "DMValueDefinition.h"
@@ -131,57 +133,98 @@ TEnumAsByte<EMaterialSamplerType> UDMMaterialProperty::GetTextureSamplerType() c
 
 void UDMMaterialProperty::OnSlotAdded(UDMMaterialSlot* InSlot)
 {
+	if (!IsValid(InSlot))
+	{
+		return;
+	}
+
+	InSlot->AddDefaultLayer(MaterialProperty);
+}
+
+void UDMMaterialProperty::AddDefaultBaseStage(UDMMaterialLayerObject* InLayer)
+{
+	if (!IsValid(InLayer))
+	{
+		return;
+	}
+
 	UDMMaterialStage* DefaultStage = UDMMaterialStageBlend::CreateStage(UDMMaterialStageBlendNormal::StaticClass());
 	check(DefaultStage);
 
-	UDMMaterialStage* MaskStage = UDMMaterialStageThroughputLayerBlend::CreateStage();
-	check(MaskStage);
+	InLayer->SetStage(EDMMaterialLayerStage::Base, DefaultStage);
 
-	InSlot->AddLayerWithMask(MaterialProperty, DefaultStage, MaskStage);
+	const FDMDefaultMaterialPropertySlotValue& DefaultValue = UDynamicMaterialEditorSettings::Get()->GetDefaultSlotValue(MaterialProperty);
 
-	if (UTexture* BaseTexture = UDynamicMaterialEditorSettings::Get()->GetDefaultTextureForSlot(MaterialProperty))
+	switch (DefaultValue.Type)
 	{
-		UDMMaterialStageInputExpression* BaseInputExpression = UDMMaterialStageInputExpression::ChangeStageInput_Expression(
-			DefaultStage,
-			UDMMaterialStageExpressionTextureSample::StaticClass(),
-			UDMMaterialStageBlendNormal::InputB,
-			FDMMaterialStageConnectorChannel::WHOLE_CHANNEL,
-			0,
-			FDMMaterialStageConnectorChannel::WHOLE_CHANNEL
-		);
-
-		if (UDMMaterialStageExpressionTextureSample* BaseInputTextureSample = Cast<UDMMaterialStageExpressionTextureSample>(BaseInputExpression->GetMaterialStageExpression()))
+		case EDMDefaultMaterialPropertySlotValueType::Texture:
 		{
-			if (UDMMaterialStage* BaseTextureInputStage = BaseInputExpression->GetSubStage())
-			{
-				const TArray<UDMMaterialStageInput*> BaseTextureStageInputs = BaseTextureInputStage->GetInputs();
+			UDMMaterialStageInputExpression* BaseInputExpression = UDMMaterialStageInputExpression::ChangeStageInput_Expression(
+				DefaultStage,
+				UDMMaterialStageExpressionTextureSample::StaticClass(),
+				UDMMaterialStageBlendNormal::InputB,
+				FDMMaterialStageConnectorChannel::WHOLE_CHANNEL,
+				0,
+				FDMMaterialStageConnectorChannel::WHOLE_CHANNEL
+			);
 
-				for (UDMMaterialStageInput* BaseTextureStageInput : BaseTextureStageInputs)
+			if (UDMMaterialStageExpressionTextureSample* BaseInputTextureSample = Cast<UDMMaterialStageExpressionTextureSample>(BaseInputExpression->GetMaterialStageExpression()))
+			{
+				if (UDMMaterialStage* BaseTextureInputStage = BaseInputExpression->GetSubStage())
 				{
-					if (UDMMaterialStageInputValue* BaseTextureInputValue = Cast<UDMMaterialStageInputValue>(BaseTextureStageInput))
+					const TArray<UDMMaterialStageInput*> BaseTextureStageInputs = BaseTextureInputStage->GetInputs();
+
+					for (UDMMaterialStageInput* BaseTextureStageInput : BaseTextureStageInputs)
 					{
-						if (UDMMaterialValueTexture* BaseTextureValue = Cast<UDMMaterialValueTexture>(BaseTextureInputValue->GetValue()))
+						if (UDMMaterialStageInputValue* BaseTextureInputValue = Cast<UDMMaterialStageInputValue>(BaseTextureStageInput))
 						{
-							BaseTextureValue->SetDefaultValue(BaseTexture);
-							BaseTextureValue->ApplyDefaultValue();
-							break;
+							if (UDMMaterialValueTexture* BaseTextureValue = Cast<UDMMaterialValueTexture>(BaseTextureInputValue->GetValue()))
+							{
+								BaseTextureValue->SetDefaultValue(DefaultValue.Texture.LoadSynchronous());
+								BaseTextureValue->ApplyDefaultValue();
+								break;
+							}
 						}
 					}
 				}
 			}
+
+			break;
+		}
+
+		case EDMDefaultMaterialPropertySlotValueType::Color:
+		{
+			UDMMaterialStageInputValue* InputValue = UDMMaterialStageInputValue::ChangeStageInput_NewLocalValue(
+				DefaultStage,
+				UDMMaterialStageBlendNormal::InputB,
+				FDMMaterialStageConnectorChannel::WHOLE_CHANNEL,
+				EDMValueType::VT_Float3_RGB,
+				0
+			);
+
+			if (UDMMaterialValueFloat3RGB* RGBValue = Cast<UDMMaterialValueFloat3RGB>(InputValue->GetValue()))
+			{
+				RGBValue->SetDefaultValue(DefaultValue.Color);
+				RGBValue->ApplyDefaultValue();
+			}
+
+			break;
 		}
 	}
-	else
+}
+
+void UDMMaterialProperty::AddDefaultMaskStage(UDMMaterialLayerObject* InLayer)
+{
+	if (!IsValid(InLayer))
 	{
-		UDMMaterialStageInputValue::ChangeStageInput_NewLocalValue(
-			DefaultStage,
-			UDMMaterialStageBlendNormal::InputB,
-			FDMMaterialStageConnectorChannel::WHOLE_CHANNEL,
-			EDMValueType::VT_Float3_RGB,
-			0
-		);
+		return;
 	}
 
+	UDMMaterialStage* MaskStage = UDMMaterialStageThroughputLayerBlend::CreateStage();
+	check(MaskStage);
+
+	InLayer->SetStage(EDMMaterialLayerStage::Mask, MaskStage);
+
 	if (UTexture* MaskTexture = UDynamicMaterialEditorSettings::Get()->DefaultMask.LoadSynchronous())
 	{
 		const TArray<UDMMaterialStageInput*> MaskStageInputs = MaskStage->GetInputs();
diff --git a/Engine/Plugins/Experimental/DynamicMaterial/Source/DynamicMaterialEditor/Private/Components/DMMaterialSlot.cpp b/Engine/Plugins/Experimental/DynamicMaterial/Source/DynamicMaterialEditor/Private/Components/DMMaterialSlot.cpp
--- a/Engine/Plugins/Experimental/DynamicMaterial/Source/DynamicMaterialEditor/Private/Components/DMMaterialSlot.cpp
+++ b/Engine/Plugins/Experimental/DynamicMaterial/Source/DynamicMaterialEditor/Private/Components/DMMaterialSlot.cpp
@@ -416,12 +416,84 @@ TSet<EDMValueType> UDMMaterialSlot::GetAllOutputConnectorTypes() const
 	return AllOutputTypes;
 }
 
+UDMMaterialLayerObject* UDMMaterialSlot::AddDefaultLayer(EDMMaterialPropertyType InMaterialProperty)
+{
+	if (!IsComponentValid())
+	{
+		return nullptr;
+	}
+
+	UDynamicMaterialModelEditorOnlyData* ModelEditorOnlyData = GetMaterialModelEditorOnlyData();
+	check(ModelEditorOnlyData);
+
+	UDMMaterialProperty* Property = ModelEditorOnlyData->GetMaterialProperty(InMaterialProperty);
+	check(Property);
+
+	UDMMaterialLayerObject* NewLayer = UDMMaterialLayerObject::CreateLayer(this, InMaterialProperty, {});
+	LayerObjects.Add(NewLayer);
+
+	{
+		const FDMUpdateGuard Guard;
+		Property->AddDefaultBaseStage(NewLayer);
+		Property->AddDefaultMaskStage(NewLayer);
+	}	
+
+	if (IsComponentAdded())
+	{
+		NewLayer->SetComponentState(EDMComponentLifetimeState::Added);
+	}
+
+	UpdateOutputConnectorTypes();
+
+	NewLayer->Update(EDMUpdateType::Structure);
+
+	OnLayersUpdateDelegate.Broadcast(this);
+
+	return NewLayer;
+}
+
 UDMMaterialLayerObject* UDMMaterialSlot::AddLayer(EDMMaterialPropertyType InMaterialProperty, UDMMaterialStage* InNewBase)
 {
-	UDMMaterialStage* NewMask = UDMMaterialStageThroughputLayerBlend::CreateStage();
-	check(NewMask);
+	if (!IsComponentValid())
+	{
+		return nullptr;
+	}
+
+	check(InNewBase);
+	check(InNewBase->GetSource());
+	check(InNewBase->GetSource()->GetOutputConnectors().IsEmpty() == false);
+
+	UDynamicMaterialModelEditorOnlyData* ModelEditorOnlyData = GetMaterialModelEditorOnlyData();
+	check(ModelEditorOnlyData);
+
+	UDMMaterialProperty* Property = ModelEditorOnlyData->GetMaterialProperty(InMaterialProperty);
+	check(Property);
+
+	if (GUndo)
+	{
+		InNewBase->Modify();
+	}
+
+	UDMMaterialLayerObject* NewLayer = UDMMaterialLayerObject::CreateLayer(this, InMaterialProperty, {InNewBase});
+	LayerObjects.Add(NewLayer);
+
+	{
+		const FDMUpdateGuard Guard;
+		Property->AddDefaultMaskStage(NewLayer);
+	}
+
+	if (IsComponentAdded())
+	{
+		NewLayer->SetComponentState(EDMComponentLifetimeState::Added);
+	}
+
+	UpdateOutputConnectorTypes();
+
+	NewLayer->Update(EDMUpdateType::Structure);
+
+	OnLayersUpdateDelegate.Broadcast(this);
 
-	return AddLayerWithMask(InMaterialProperty, InNewBase, NewMask);
+	return NewLayer;
 }
 
 UDMMaterialLayerObject* UDMMaterialSlot::AddLayerWithMask(EDMMaterialPropertyType InMaterialProperty, UDMMaterialStage* InNewBase, 
@@ -456,11 +528,11 @@ UDMMaterialLayerObject* UDMMaterialSlot::AddLayerWithMask(EDMMaterialPropertyTyp
 
 	UpdateOutputConnectorTypes();
 
-	InNewBase->Update(EDMUpdateType::Structure);
+	NewLayer->Update(EDMUpdateType::Structure);
 
 	OnLayersUpdateDelegate.Broadcast(this);
 
-	return LayerObjects.Last();
+	return NewLayer;
 }
 
 bool UDMMaterialSlot::PasteLayer(UDMMaterialLayerObject* InLayer)
diff --git a/Engine/Plugins/Experimental/DynamicMaterial/Source/DynamicMaterialEditor/Private/DynamicMaterialEditorModule.cpp b/Engine/Plugins/Experimental/DynamicMaterial/Source/DynamicMaterialEditor/Private/DynamicMaterialEditorModule.cpp
--- a/Engine/Plugins/Experimental/DynamicMaterial/Source/DynamicMaterialEditor/Private/DynamicMaterialEditorModule.cpp
+++ b/Engine/Plugins/Experimental/DynamicMaterial/Source/DynamicMaterialEditor/Private/DynamicMaterialEditorModule.cpp
@@ -303,9 +303,11 @@ void FDynamicMaterialEditorModule::StartupModule()
 
 	UDMMaterialValueTexture::GetDefaultRGBTexture.BindLambda([]()
 		{
-			if (UTexture* Texture = UDynamicMaterialEditorSettings::Get()->GetDefaultTextureForSlot(EDMMaterialPropertyType::BaseColor))
+			const FDMDefaultMaterialPropertySlotValue& DefaultValue = UDynamicMaterialEditorSettings::Get()->GetDefaultSlotValue(EDMMaterialPropertyType::BaseColor);
+
+			if (UTexture* DefaultTexture = DefaultValue.Texture.LoadSynchronous())
 			{
-				return Texture;
+				return DefaultTexture;
 			}
 
 			return UDynamicMaterialEditorSettings::Get()->DefaultMask.LoadSynchronous();
diff --git a/Engine/Plugins/Experimental/DynamicMaterial/Source/DynamicMaterialEditor/Private/DynamicMaterialEditorSettings.cpp b/Engine/Plugins/Experimental/DynamicMaterial/Source/DynamicMaterialEditor/Private/DynamicMaterialEditorSettings.cpp
--- a/Engine/Plugins/Experimental/DynamicMaterial/Source/DynamicMaterialEditor/Private/DynamicMaterialEditorSettings.cpp
+++ b/Engine/Plugins/Experimental/DynamicMaterial/Source/DynamicMaterialEditor/Private/DynamicMaterialEditorSettings.cpp
@@ -4,6 +4,7 @@
 
 #include "AssetRegistry/AssetData.h"
 #include "Components/DMMaterialComponent.h"
+#include "Components/MaterialValues/DMMaterialValueFloat3RGB.h"
 #include "DynamicMaterialEditorModule.h"
 #include "Engine/AssetManager.h"
 #include "GenericPlatform/GenericApplication.h"
@@ -18,18 +19,38 @@
 
 namespace UE::DynamicMaterialEditor::Private
 {
-	const TMap<EDMMaterialPropertyType, TSoftObjectPtr<UTexture>> BaseDefaultSlotProperties = {
+	const TMap<EDMMaterialPropertyType, FDMDefaultMaterialPropertySlotValue> DefaultSlotValues = {
 		{EDMMaterialPropertyType::BaseColor,           TSoftObjectPtr<UTexture>(FSoftObjectPath(TEXT("Texture2D'/DynamicMaterial/T_Default_Texture.T_Default_Texture'")))},
 		{EDMMaterialPropertyType::EmissiveColor,       TSoftObjectPtr<UTexture>(FSoftObjectPath(TEXT("Texture2D'/DynamicMaterial/T_Default_Texture.T_Default_Texture'")))},
 		{EDMMaterialPropertyType::Opacity,             TSoftObjectPtr<UTexture>(FSoftObjectPath(TEXT("/Script/Engine.Texture2D'/DynamicMaterial/Textures/T_DM_HorizontalGradient.T_DM_HorizontalGradient'")))},
 		{EDMMaterialPropertyType::OpacityMask,         TSoftObjectPtr<UTexture>(FSoftObjectPath(TEXT("/Script/Engine.Texture2D'/DynamicMaterial/Textures/T_DM_HorizontalGradient.T_DM_HorizontalGradient'")))},
 		{EDMMaterialPropertyType::Metallic,            TSoftObjectPtr<UTexture>(FSoftObjectPath(TEXT("/Script/Engine.Texture2D'/DynamicMaterial/Textures/SlotDefaults/T_MD_Metallic.T_MD_Metallic'")))},
+		{EDMMaterialPropertyType::Specular,            FLinearColor(0.5f, 0.5f, 0.5f, 1.f)},
 		{EDMMaterialPropertyType::Roughness,           TSoftObjectPtr<UTexture>(FSoftObjectPath(TEXT("/Script/Engine.Texture2D'/DynamicMaterial/Textures/SlotDefaults/T_MD_Roughness.T_MD_Roughness'")))},
 		{EDMMaterialPropertyType::Normal,              TSoftObjectPtr<UTexture>(FSoftObjectPath(TEXT("/Script/Engine.Texture2D'/DynamicMaterial/Textures/SlotDefaults/T_MD_Normal.T_MD_Normal'")))},
 		{EDMMaterialPropertyType::AmbientOcclusion,    TSoftObjectPtr<UTexture>(FSoftObjectPath(TEXT("/Script/Engine.Texture2D'/DynamicMaterial/Textures/SlotDefaults/T_MD_AmbientOcclusion.T_MD_AmbientOcclusion'")))},
 	};
 }
 
+FDMDefaultMaterialPropertySlotValue::FDMDefaultMaterialPropertySlotValue()
+	: Type(EDMDefaultMaterialPropertySlotValueType::Texture)
+	, Color(FLinearColor::White)
+{
+}
+
+FDMDefaultMaterialPropertySlotValue::FDMDefaultMaterialPropertySlotValue(const TSoftObjectPtr<UTexture>& InTexture)
+	: Type(EDMDefaultMaterialPropertySlotValueType::Texture)
+	, Texture(InTexture)
+	, Color(FLinearColor::Black)
+{
+}
+
+FDMDefaultMaterialPropertySlotValue::FDMDefaultMaterialPropertySlotValue(const FLinearColor& InColor)
+	: Type(EDMDefaultMaterialPropertySlotValueType::Color)
+	, Color(InColor)
+{
+}
+
 bool FDMMaterialChannelListPreset::IsPropertyEnabled(EDMMaterialPropertyType InProperty) const
 {
 	switch (InProperty)
@@ -89,8 +110,6 @@ UDynamicMaterialEditorSettings::UDynamicMaterialEditorSettings()
 
 	DefaultMask = TSoftObjectPtr<UTexture>(FSoftObjectPath(TEXT("Texture2D'/Engine/EngineResources/WhiteSquareTexture.WhiteSquareTexture'")));
 
-	OverriddenDefaultSlotTextures = {};
-
 	FDMMaterialChannelListPreset Opaque;
 	Opaque.Name = TEXT("Opaque");
 	Opaque.bRGB = true;
@@ -382,25 +401,43 @@ TArray<FDMMaterialEffectList> UDynamicMaterialEditorSettings::GetEffectList() co
 	return Effects;
 }
 
-UTexture* UDynamicMaterialEditorSettings::GetDefaultTextureForSlot(EDMMaterialPropertyType InProperty) const
+const FDMDefaultMaterialPropertySlotValue& UDynamicMaterialEditorSettings::GetDefaultSlotValue(EDMMaterialPropertyType InProperty) const
 {
-	if (const TSoftObjectPtr<UTexture>* OverridePtr = OverriddenDefaultSlotTextures.Find(InProperty))
+	if (const FDMDefaultMaterialPropertySlotValue* OverridePtr = DefaultSlotValueOverrides.Find(InProperty))
 	{
-		if (UTexture* OverrideTexture = OverridePtr->LoadSynchronous())
+		switch (OverridePtr->Type)
 		{
-			return OverrideTexture;
+			case EDMDefaultMaterialPropertySlotValueType::Texture:
+				if (!OverridePtr->Texture.IsNull() && OverridePtr->Texture.LoadSynchronous())
+				{
+					return *OverridePtr;
+				}
+				break;
+
+			case EDMDefaultMaterialPropertySlotValueType::Color:
+				return *OverridePtr;
 		}
 	}
 	
-	if (const TSoftObjectPtr<UTexture>* OverridePtr = UE::DynamicMaterialEditor::Private::BaseDefaultSlotProperties.Find(InProperty))
+	if (const FDMDefaultMaterialPropertySlotValue* DefaultPtr = UE::DynamicMaterialEditor::Private::DefaultSlotValues.Find(InProperty))
 	{
-		if (UTexture* OverrideTexture = OverridePtr->LoadSynchronous())
+		switch (DefaultPtr->Type)
 		{
-			return OverrideTexture;
+			case EDMDefaultMaterialPropertySlotValueType::Texture:
+				if (!DefaultPtr->Texture.IsNull() && DefaultPtr->Texture.LoadSynchronous())
+				{
+					return *DefaultPtr;
+				}
+				break;
+
+			case EDMDefaultMaterialPropertySlotValueType::Color:
+				return *DefaultPtr;
 		}
 	}
 
-	return nullptr;
+	static const FDMDefaultMaterialPropertySlotValue Default = {GetDefault<UDMMaterialValueFloat3RGB>()->GetValue()};
+
+	return Default;
 }
 
 const FDMMaterialChannelListPreset* UDynamicMaterialEditorSettings::GetPresetByName(FName InName) const
diff --git a/Engine/Plugins/Experimental/DynamicMaterial/Source/DynamicMaterialEditor/Private/DynamicMaterialEditorSettings.h b/Engine/Plugins/Experimental/DynamicMaterial/Source/DynamicMaterialEditor/Private/DynamicMaterialEditorSettings.h
--- a/Engine/Plugins/Experimental/DynamicMaterial/Source/DynamicMaterialEditor/Private/DynamicMaterialEditorSettings.h
+++ b/Engine/Plugins/Experimental/DynamicMaterial/Source/DynamicMaterialEditor/Private/DynamicMaterialEditorSettings.h
@@ -29,6 +29,34 @@ struct FDMMaterialEffectList
 	TArray<TSoftObjectPtr<UMaterialFunctionInterface>> Effects;
 };
 
+UENUM(BlueprintType)
+enum class EDMDefaultMaterialPropertySlotValueType : uint8
+{
+	Texture,
+	Color
+};
+
+USTRUCT(BlueprintType)
+struct FDMDefaultMaterialPropertySlotValue
+{
+	GENERATED_BODY()
+
+	FDMDefaultMaterialPropertySlotValue();
+	FDMDefaultMaterialPropertySlotValue(const TSoftObjectPtr<UTexture>& InTexture);
+	FDMDefaultMaterialPropertySlotValue(const FLinearColor& InColor);
+
+	UPROPERTY(EditAnywhere, BlueprintReadWrite, Category = "Material Designer")
+	EDMDefaultMaterialPropertySlotValueType Type;
+
+	UPROPERTY(EditAnywhere, BlueprintReadWrite, Category = "Material Designer", 
+		meta = (EditCondition = "Type == EDMDefaultMaterialPropertySlotValueType::Texture", EditConditionHides))
+	TSoftObjectPtr<UTexture> Texture;
+
+	UPROPERTY(EditAnywhere, BlueprintReadWrite, Category = "Material Designer",
+		meta = (EditCondition = "Type == EDMDefaultMaterialPropertySlotValueType::Color", EditConditionHides))
+	FLinearColor Color;
+};
+
 USTRUCT(BlueprintType)
 struct FDMMaterialChannelListPreset
 {
@@ -175,8 +203,11 @@ public:
 	UPROPERTY(Config, EditAnywhere, BlueprintReadWrite, Category = "Preview")
 	TSoftObjectPtr<UTexture> DefaultMask;
 
+	/**
+	 * Overrides the default values given to slots created in the given material property.
+	 */
 	UPROPERTY(Config, EditAnywhere, BlueprintReadWrite, Category = "Channels")
-	TMap<EDMMaterialPropertyType, TSoftObjectPtr<UTexture>> OverriddenDefaultSlotTextures;
+	TMap<EDMMaterialPropertyType, FDMDefaultMaterialPropertySlotValue> DefaultSlotValueOverrides;
 
 	/*
 	 * Add paths to search for custom effects.
@@ -219,7 +250,7 @@ public:
 
 	TArray<FDMMaterialEffectList> GetEffectList() const;
 
-	UTexture* GetDefaultTextureForSlot(EDMMaterialPropertyType InProperty) const;
+	const FDMDefaultMaterialPropertySlotValue& GetDefaultSlotValue(EDMMaterialPropertyType InProperty) const;
 
 	const FDMMaterialChannelListPreset* GetPresetByName(FName InName) const;
 
diff --git a/Engine/Plugins/Experimental/DynamicMaterial/Source/DynamicMaterialEditor/Private/Slate/SDMSlot.cpp b/Engine/Plugins/Experimental/DynamicMaterial/Source/DynamicMaterialEditor/Private/Slate/SDMSlot.cpp
--- a/Engine/Plugins/Experimental/DynamicMaterial/Source/DynamicMaterialEditor/Private/Slate/SDMSlot.cpp
+++ b/Engine/Plugins/Experimental/DynamicMaterial/Source/DynamicMaterialEditor/Private/Slate/SDMSlot.cpp
@@ -2,6 +2,7 @@
 
 #include "SDMSlot.h"
 #include "Components/DMMaterialLayer.h"
+#include "Components/DMMaterialProperty.h"
 #include "Components/DMMaterialSlot.h"
 #include "Components/DMMaterialStage.h"
 #include "Components/DMMaterialStageBlend.h"
@@ -9,6 +10,7 @@
 #include "Components/DMMaterialStageFunction.h"
 #include "Components/DMMaterialStageGradient.h"
 #include "Components/DMMaterialStageThroughputLayerBlend.h"
+#include "Components/DMMaterialSubStage.h"
 #include "Components/DMMaterialValue.h"
 #include "Components/DMRenderTargetRenderer.h"
 #include "Components/MaterialStageBlends/DMMSBNormal.h"
@@ -21,6 +23,9 @@
 #include "Components/MaterialStageInputs/DMMSITextureUV.h"
 #include "Components/MaterialStageInputs/DMMSIValue.h"
 #include "Components/MaterialValues/DMMaterialValueFloat1.h"
+#include "Components/MaterialValues/DMMaterialValueFloat3RGB.h"
+#include "Components/MaterialValues/DMMaterialValueFloat4.h"
+#include "Components/MaterialValues/DMMaterialValueTexture.h"
 #include "DetailLayoutBuilder.h"
 #include "DynamicMaterialEditorModule.h"
 #include "DynamicMaterialEditorSettings.h"
@@ -1169,7 +1174,11 @@ UDMMaterialLayerObject* SDMSlot::AddNewLayer(UDMMaterialStage* InNewBaseStage, U
 			return nullptr;
 		}
 
-		if (!InNewMaskStage)
+		if (!InNewBaseStage && !InNewMaskStage)
+		{
+			Slot->AddDefaultLayer(MaterialProperty);
+		}
+		else if (!InNewMaskStage)
 		{
 			Slot->AddLayer(MaterialProperty, InNewBaseStage);
 		}
@@ -1211,13 +1220,15 @@ void SDMSlot::AddNewLayer_NewLocalValue(TSubclassOf<UDMMaterialValue> InValueCla
 	UDMMaterialStage* NewBase = UDMMaterialStageBlend::CreateStage(UDMMaterialStageBlendNormal::StaticClass());
 	AddNewLayer(NewBase);
 
-	UDMMaterialStageInputValue::ChangeStageInput_NewLocalValue(
+	UDMMaterialStageInputValue* InputValue = UDMMaterialStageInputValue::ChangeStageInput_NewLocalValue(
 		NewBase, 
 		UDMMaterialStageBlend::InputB,
 		FDMMaterialStageConnectorChannel::WHOLE_CHANNEL, 
 		InValueClass,
 		FDMMaterialStageConnectorChannel::WHOLE_CHANNEL
 	);
+
+	OnValueAdded(NewBase, InputValue->GetValue());
 }
 
 void SDMSlot::AddNewLayer_GlobalValue(UDMMaterialValue* InValue)
@@ -1250,6 +1261,8 @@ void SDMSlot::AddNewLayer_GlobalValue(UDMMaterialValue* InValue)
 		InValue,
 		FDMMaterialStageConnectorChannel::WHOLE_CHANNEL
 	);
+
+	// Already existing global values should not have their value changed to the property default.
 }
 
 void SDMSlot::AddNewLayer_NewGlobalValue(EDMValueType InValueType)
@@ -1271,13 +1284,15 @@ void SDMSlot::AddNewLayer_NewGlobalValue(TSubclassOf<UDMMaterialValue> InValueCl
 	UDMMaterialStage* NewStage = UDMMaterialStageBlend::CreateStage(UDMMaterialStageBlendNormal::StaticClass());
 	AddNewLayer(NewStage);
 
-	UDMMaterialStageInputValue::ChangeStageInput_NewValue(
+	UDMMaterialStageInputValue* InputValue = UDMMaterialStageInputValue::ChangeStageInput_NewValue(
 		NewStage, 
 		UDMMaterialStageBlend::InputB,
 		FDMMaterialStageConnectorChannel::WHOLE_CHANNEL,
 		InValueClass,
 		FDMMaterialStageConnectorChannel::WHOLE_CHANNEL
 	);
+
+	OnValueAdded(NewStage, InputValue->GetValue());
 }
 
 void SDMSlot::AddNewLayer_Slot(UDMMaterialSlot* InSlot, EDMMaterialPropertyType InMaterialProperty)
@@ -1330,7 +1345,7 @@ void SDMSlot::AddNewLayer_Expression(TSubclassOf<UDMMaterialStageExpression> InE
 	UDMMaterialStage* NewStage = UDMMaterialStageBlend::CreateStage(UDMMaterialStageBlendNormal::StaticClass());
 	UDMMaterialLayerObject* NewLayer = AddNewLayer(NewStage);
 
-	UDMMaterialStageInputExpression::ChangeStageInput_Expression(
+	UDMMaterialStageInputExpression* NewExpression = UDMMaterialStageInputExpression::ChangeStageInput_Expression(
 		NewStage, 
 		InExpressionClass,
 		UDMMaterialStageBlend::InputB, 
@@ -1359,6 +1374,20 @@ void SDMSlot::AddNewLayer_Expression(TSubclassOf<UDMMaterialStageExpression> InE
 			// No nothing
 			break;
 	}
+
+	if (InExpressionClass->IsChildOf(UDMMaterialStageExpressionTextureSampleBase::StaticClass()))
+	{
+		if (UDMMaterialSubStage* SubStage = NewExpression->GetSubStage())
+		{
+			for (UDMMaterialStageInput* Input : SubStage->GetInputs())
+			{
+				if (UDMMaterialStageInputValue* InputValue = Cast<UDMMaterialStageInputValue>(Input))
+				{
+					OnValueAdded(NewStage, InputValue->GetValue());
+				}
+			}
+		}
+	}
 }
 
 void SDMSlot::AddNewLayer_Blend(TSubclassOf<UDMMaterialStageBlend> InBlendClass)
@@ -1521,4 +1550,58 @@ void SDMSlot::OnLayerStageSelected(const bool bInSelected, const TSharedRef<SDMS
 	}
 }
 
+void SDMSlot::OnValueAdded(UDMMaterialStage* InStage, UDMMaterialValue* InValue)
+{
+	if (!InStage || !InValue)
+	{
+		return;
+	}
+
+	UDMMaterialLayerObject* Layer = InStage->GetLayer();
+
+	if (!Layer)
+	{
+		return;
+	}
+
+	UClass* ValueClass = InValue->GetClass();
+
+	// Do not include subclasses
+	if (ValueClass != UDMMaterialValueTexture::StaticClass()
+		&& ValueClass != UDMMaterialValueFloat3RGB::StaticClass()
+		&& ValueClass != UDMMaterialValueFloat4::StaticClass())
+	{
+		return;
+	}
+
+	const FDMDefaultMaterialPropertySlotValue& DefaultValue = GetDefault<UDynamicMaterialEditorSettings>()->GetDefaultSlotValue(Layer->GetMaterialProperty());
+
+	switch (DefaultValue.Type)
+	{
+		case EDMDefaultMaterialPropertySlotValueType::Texture:
+			if (UDMMaterialValueTexture* TextureValue = Cast<UDMMaterialValueTexture>(InValue))
+			{
+				if (UTexture* Texture = DefaultValue.Texture.LoadSynchronous())
+				{
+					TextureValue->SetDefaultValue(Texture);
+					TextureValue->ApplyDefaultValue();
+				}
+			}
+			break;
+
+		case EDMDefaultMaterialPropertySlotValueType::Color:
+			if (UDMMaterialValueFloat3RGB* RGBValue = Cast<UDMMaterialValueFloat3RGB>(InValue))
+			{
+				RGBValue->SetDefaultValue(DefaultValue.Color);
+				RGBValue->ApplyDefaultValue();
+			}
+			else if (UDMMaterialValueFloat4* RGBAValue = Cast<UDMMaterialValueFloat4>(InValue))
+			{
+				RGBAValue->SetDefaultValue(DefaultValue.Color);
+				RGBAValue->ApplyDefaultValue();
+			}
+			break;
+	}
+}
+
 #undef LOCTEXT_NAMESPACE
diff --git a/Engine/Plugins/Experimental/DynamicMaterial/Source/DynamicMaterialEditor/Private/Slate/SDMSlot.h b/Engine/Plugins/Experimental/DynamicMaterial/Source/DynamicMaterialEditor/Private/Slate/SDMSlot.h
--- a/Engine/Plugins/Experimental/DynamicMaterial/Source/DynamicMaterialEditor/Private/Slate/SDMSlot.h
+++ b/Engine/Plugins/Experimental/DynamicMaterial/Source/DynamicMaterialEditor/Private/Slate/SDMSlot.h
@@ -87,7 +87,7 @@ public:
 
 	void SetSelectedLayer(UDMMaterialLayerObject* InLayer) const;
 
-	UDMMaterialLayerObject* AddNewLayer(UDMMaterialStage* InNewBaseStage, UDMMaterialStage* InNewMaskStage = nullptr);
+	UDMMaterialLayerObject* AddNewLayer(UDMMaterialStage* InNewBaseStage = nullptr, UDMMaterialStage* InNewMaskStage = nullptr);
 
 protected:
 	TWeakPtr<SDMEditor> EditorWidgetWeak;
@@ -134,6 +134,8 @@ protected:
 	void OnLayerSelected(TSharedPtr<FDMMaterialLayerReference> InLayerItem, const int32 InLayerIndex);
 	void OnLayerStageSelected(const bool bInSelected, const TSharedRef<SDMStage>& InStageWidget);
 
+	void OnValueAdded(UDMMaterialStage* InStage, UDMMaterialValue* InValue);
+
 	void RefreshMainWidget();
 	void RefreshHeaderPropertyListWidget();
 	void RefreshSlotSettingsRowWidget();
diff --git a/Engine/Plugins/Experimental/DynamicMaterial/Source/DynamicMaterialEditor/Public/Components/DMMaterialProperty.h b/Engine/Plugins/Experimental/DynamicMaterial/Source/DynamicMaterialEditor/Public/Components/DMMaterialProperty.h
--- a/Engine/Plugins/Experimental/DynamicMaterial/Source/DynamicMaterialEditor/Public/Components/DMMaterialProperty.h
+++ b/Engine/Plugins/Experimental/DynamicMaterial/Source/DynamicMaterialEditor/Public/Components/DMMaterialProperty.h
@@ -9,6 +9,7 @@
 #include "DMMaterialProperty.generated.h"
 
 class FText;
+class UDMMaterialLayerObject;
 class UDMMaterialSlot;
 class UDMMaterialStageInputExpression;
 class UDynamicMaterialModel;
@@ -118,6 +119,16 @@ public:
 	 */
 	virtual void OnSlotAdded(UDMMaterialSlot* InSlot);
 
+	/**
+	 * Adds the default base stage for a layer in this property.
+	 */
+	virtual void AddDefaultBaseStage(UDMMaterialLayerObject* InLayer);
+
+	/**
+	 * Adds the default mask stage for a layer in this property.
+	 */
+	virtual void AddDefaultMaskStage(UDMMaterialLayerObject* InLayer);
+
 	/**
 	 * Adds a component to the component list. Sets component state to added.
 	 * @param InName The name of the component in the map.
diff --git a/Engine/Plugins/Experimental/DynamicMaterial/Source/DynamicMaterialEditor/Public/Components/DMMaterialSlot.h b/Engine/Plugins/Experimental/DynamicMaterial/Source/DynamicMaterialEditor/Public/Components/DMMaterialSlot.h
--- a/Engine/Plugins/Experimental/DynamicMaterial/Source/DynamicMaterialEditor/Public/Components/DMMaterialSlot.h
+++ b/Engine/Plugins/Experimental/DynamicMaterial/Source/DynamicMaterialEditor/Public/Components/DMMaterialSlot.h
@@ -63,6 +63,9 @@ public:
 
 	const TArray<TObjectPtr<UDMMaterialLayerObject>>& GetLayers() const { return LayerObjects; }
 
+	UFUNCTION(BlueprintCallable, Category = "Material Designer")
+	UDMMaterialLayerObject* AddDefaultLayer(EDMMaterialPropertyType InMaterialProperty);
+
 	UFUNCTION(BlueprintCallable, Category = "Material Designer")
 	UDMMaterialLayerObject* AddLayer(EDMMaterialPropertyType InMaterialProperty, UDMMaterialStage* InNewBase);
 
