Material Designer
- Horde issue 465716

diff --git a/Engine/Plugins/Experimental/DynamicMaterial/Source/DynamicMaterialEditor/Private/Slate/SDMSlot.cpp b/Engine/Plugins/Experimental/DynamicMaterial/Source/DynamicMaterialEditor/Private/Slate/SDMSlot.cpp
--- a/Engine/Plugins/Experimental/DynamicMaterial/Source/DynamicMaterialEditor/Private/Slate/SDMSlot.cpp
+++ b/Engine/Plugins/Experimental/DynamicMaterial/Source/DynamicMaterialEditor/Private/Slate/SDMSlot.cpp
@@ -1124,7 +1124,7 @@ void SDMSlot::AddPropertyToSlot(EDMMaterialPropertyType Property)
 
 UDMMaterialLayerObject* SDMSlot::AddNewLayer(UDMMaterialStage* InNewBaseStage, UDMMaterialStage* InNewMaskStage)
 {
-	if (!IsValid(InNewBaseStage))
+	if (!InNewBaseStage || !IsValid(InNewBaseStage))
 	{
 		return nullptr;
 	}
@@ -1143,6 +1143,8 @@ UDMMaterialLayerObject* SDMSlot::AddNewLayer(UDMMaterialStage* InNewBaseStage, U
 		return nullptr;
 	}
 
+	UDMMaterialLayerObject* Layer = nullptr;
+
 	{
 		const FDMUpdateGuard Guard;
 
@@ -1176,33 +1178,44 @@ UDMMaterialLayerObject* SDMSlot::AddNewLayer(UDMMaterialStage* InNewBaseStage, U
 
 		if (!InNewBaseStage && !InNewMaskStage)
 		{
-			Slot->AddDefaultLayer(MaterialProperty);
+			Layer = Slot->AddDefaultLayer(MaterialProperty);
 		}
 		else if (!InNewMaskStage)
 		{
-			Slot->AddLayer(MaterialProperty, InNewBaseStage);
+			Layer = Slot->AddLayer(MaterialProperty, InNewBaseStage);
 		}
 		else
 		{
-			Slot->AddLayerWithMask(MaterialProperty, InNewBaseStage, InNewMaskStage);
+			Layer = Slot->AddLayerWithMask(MaterialProperty, InNewBaseStage, InNewMaskStage);
+		}
+	}
+
+	const bool bValidLayer = IsValid(Layer);
+	UDMMaterialStageSource* Source = nullptr;
+
+	if (bValidLayer)
+	{
+		if (UDMMaterialStage* Stage = Layer->GetStage(EDMMaterialLayerStage::Base))
+		{
+			Source = Stage->GetSource();
 		}
 	}
 
-	UDMMaterialStageSource* const Source = InNewBaseStage->GetSource();
 	if (IsValid(Source))
 	{
 		Source->Update(EDMUpdateType::Structure);
 	}
-
-	UDMMaterialLayerObject* Layer = Slot->FindLayer(InNewBaseStage);
-	if (!Layer)
+	else if (bValidLayer)
 	{
-		return nullptr;
+		Layer->Update(EDMUpdateType::Structure);
 	}
 
-	if (ensure(LayerView.IsValid()))
+	if (bValidLayer)
 	{
-		LayerView->AddLayerItem(MakeShared<FDMMaterialLayerReference>(Layer));
+		if (ensure(LayerView.IsValid()))
+		{
+			LayerView->AddLayerItem(MakeShared<FDMMaterialLayerReference>(Layer));
+		}
 	}
 
 	return Layer;
