Material Designer
- Fix for a crash when changing throughput class related to whether to keep the previous input values.

diff --git a/Engine/Plugins/Experimental/DynamicMaterial/Source/DynamicMaterialEditor/Private/Components/DMMaterialStageThroughput.cpp b/Engine/Plugins/Experimental/DynamicMaterial/Source/DynamicMaterialEditor/Private/Components/DMMaterialStageThroughput.cpp
--- a/Engine/Plugins/Experimental/DynamicMaterial/Source/DynamicMaterialEditor/Private/Components/DMMaterialStageThroughput.cpp
+++ b/Engine/Plugins/Experimental/DynamicMaterial/Source/DynamicMaterialEditor/Private/Components/DMMaterialStageThroughput.cpp
@@ -173,11 +173,11 @@ bool UDMMaterialStageThroughput::ShouldKeepInput(int32 InInputIdx)
 						const TArray<FDMMaterialStageConnector>& PreviousStageOutputConnectors = MaskStage->GetSource()->GetOutputConnectors();
 						check(PreviousStageOutputConnectors.IsValidIndex(Channel.OutputIndex));
 
-						if (UDMValueDefinitionLibrary::GetValueDefinition(PreviousStageOutputConnectors[Channel.OutputChannel].Type).IsFloatType() == false)
+						if (UDMValueDefinitionLibrary::GetValueDefinition(PreviousStageOutputConnectors[Channel.OutputIndex].Type).IsFloatType() == false)
 						{
 							check(InputConnectionMap[InInputIdx].Channels.Num() == 1);
 
-							return InputConnectors[InInputIdx].IsCompatibleWith(PreviousStageOutputConnectors[Channel.OutputChannel]);
+							return InputConnectors[InInputIdx].IsCompatibleWith(PreviousStageOutputConnectors[Channel.OutputIndex]);
 						}
 						else
 						{
@@ -186,7 +186,7 @@ bool UDMMaterialStageThroughput::ShouldKeepInput(int32 InInputIdx)
 
 						if (Channel.OutputChannel == FDMMaterialStageConnectorChannel::WHOLE_CHANNEL)
 						{
-							ThisInputCount = UDMValueDefinitionLibrary::GetValueDefinition(PreviousStageOutputConnectors[Channel.OutputChannel].Type).GetFloatCount();
+							ThisInputCount = UDMValueDefinitionLibrary::GetValueDefinition(PreviousStageOutputConnectors[Channel.OutputIndex].Type).GetFloatCount();
 						}
 						else
 						{
@@ -198,7 +198,7 @@ bool UDMMaterialStageThroughput::ShouldKeepInput(int32 InInputIdx)
 
 							ThisInputCount = FMath::Min(
 								ThisInputCount,
-								static_cast<int32>(UDMValueDefinitionLibrary::GetValueDefinition(PreviousStageOutputConnectors[Channel.OutputChannel].Type).GetFloatCount())
+								static_cast<int32>(UDMValueDefinitionLibrary::GetValueDefinition(PreviousStageOutputConnectors[Channel.OutputIndex].Type).GetFloatCount())
 							);
 						}
 					}
@@ -217,11 +217,11 @@ bool UDMMaterialStageThroughput::ShouldKeepInput(int32 InInputIdx)
 				const TArray<FDMMaterialStageConnector>& InputOutputConnectors = StageInputs[InputObjectIdx]->GetOutputConnectors();
 				check(InputOutputConnectors.IsValidIndex(Channel.OutputIndex));
 
-				if (UDMValueDefinitionLibrary::GetValueDefinition(InputOutputConnectors[Channel.OutputChannel].Type).IsFloatType() == false)
+				if (UDMValueDefinitionLibrary::GetValueDefinition(InputOutputConnectors[Channel.OutputIndex].Type).IsFloatType() == false)
 				{
 					check(InputConnectionMap[InInputIdx].Channels.Num() == 1);
 
-					return InputConnectors[InInputIdx].IsCompatibleWith(InputOutputConnectors[Channel.OutputChannel]);
+					return InputConnectors[InInputIdx].IsCompatibleWith(InputOutputConnectors[Channel.OutputIndex]);
 				}
 				else
 				{
@@ -230,7 +230,7 @@ bool UDMMaterialStageThroughput::ShouldKeepInput(int32 InInputIdx)
 
 				if (Channel.OutputChannel == FDMMaterialStageConnectorChannel::WHOLE_CHANNEL)
 				{
-					ThisInputCount = UDMValueDefinitionLibrary::GetValueDefinition(InputOutputConnectors[Channel.OutputChannel].Type).GetFloatCount();
+					ThisInputCount = UDMValueDefinitionLibrary::GetValueDefinition(InputOutputConnectors[Channel.OutputIndex].Type).GetFloatCount();
 				}
 				else
 				{
@@ -242,7 +242,7 @@ bool UDMMaterialStageThroughput::ShouldKeepInput(int32 InInputIdx)
 
 					ThisInputCount = FMath::Min(
 						ThisInputCount,
-						static_cast<int32>(UDMValueDefinitionLibrary::GetValueDefinition(InputOutputConnectors[Channel.OutputChannel].Type).GetFloatCount())
+						static_cast<int32>(UDMValueDefinitionLibrary::GetValueDefinition(InputOutputConnectors[Channel.OutputIndex].Type).GetFloatCount())
 					);
 				}
 
