Material Designer
- Order of operations fix for adding new layers resulting in the default color/texture not being applied.

diff --git a/Engine/Plugins/Experimental/DynamicMaterial/Source/DynamicMaterialEditor/Private/Components/DMMaterialSlot.cpp b/Engine/Plugins/Experimental/DynamicMaterial/Source/DynamicMaterialEditor/Private/Components/DMMaterialSlot.cpp
--- a/Engine/Plugins/Experimental/DynamicMaterial/Source/DynamicMaterialEditor/Private/Components/DMMaterialSlot.cpp
+++ b/Engine/Plugins/Experimental/DynamicMaterial/Source/DynamicMaterialEditor/Private/Components/DMMaterialSlot.cpp
@@ -432,15 +432,15 @@ UDMMaterialLayerObject* UDMMaterialSlot::AddDefaultLayer(EDMMaterialPropertyType
 	UDMMaterialLayerObject* NewLayer = UDMMaterialLayerObject::CreateLayer(this, InMaterialProperty, {});
 	LayerObjects.Add(NewLayer);
 
+	if (IsComponentAdded())
+	{
+		NewLayer->SetComponentState(EDMComponentLifetimeState::Added);
+	}
+
 	{
 		const FDMUpdateGuard Guard;
 		Property->AddDefaultBaseStage(NewLayer);
 		Property->AddDefaultMaskStage(NewLayer);
-	}	
-
-	if (IsComponentAdded())
-	{
-		NewLayer->SetComponentState(EDMComponentLifetimeState::Added);
 	}
 
 	UpdateOutputConnectorTypes();
@@ -477,14 +477,14 @@ UDMMaterialLayerObject* UDMMaterialSlot::AddLayer(EDMMaterialPropertyType InMate
 	UDMMaterialLayerObject* NewLayer = UDMMaterialLayerObject::CreateLayer(this, InMaterialProperty, {InNewBase});
 	LayerObjects.Add(NewLayer);
 
+	if (IsComponentAdded())
 	{
-		const FDMUpdateGuard Guard;
-		Property->AddDefaultMaskStage(NewLayer);
+		NewLayer->SetComponentState(EDMComponentLifetimeState::Added);
 	}
 
-	if (IsComponentAdded())
 	{
-		NewLayer->SetComponentState(EDMComponentLifetimeState::Added);
+		const FDMUpdateGuard Guard;
+		Property->AddDefaultMaskStage(NewLayer);
 	}
 
 	UpdateOutputConnectorTypes();
