Material Designer
- Fixed "UObject names" for Material Designer Instance and Model asset definitions.

diff --git a/Engine/Plugins/Experimental/DynamicMaterial/Source/DynamicMaterialEditor/Private/Material/AssetDefinition_DynamicMaterialInstance.cpp b/Engine/Plugins/Experimental/DynamicMaterial/Source/DynamicMaterialEditor/Private/Material/AssetDefinition_DynamicMaterialInstance.cpp
--- a/Engine/Plugins/Experimental/DynamicMaterial/Source/DynamicMaterialEditor/Private/Material/AssetDefinition_DynamicMaterialInstance.cpp
+++ b/Engine/Plugins/Experimental/DynamicMaterial/Source/DynamicMaterialEditor/Private/Material/AssetDefinition_DynamicMaterialInstance.cpp
@@ -15,7 +15,7 @@ FText UAssetDefinition_DynamicMaterialInstance::GetAssetDisplayName() const
 
 FText UAssetDefinition_DynamicMaterialInstance::GetAssetDisplayName(const FAssetData& InAssetData) const
 {
-	return FText::FromName(InAssetData.AssetName);
+	return LOCTEXT("MaterialDesignerInstance", "Material Designer Instance");
 }
 
 TSoftClassPtr<> UAssetDefinition_DynamicMaterialInstance::GetAssetClass() const
diff --git a/Engine/Plugins/Experimental/DynamicMaterial/Source/DynamicMaterialEditor/Private/Model/AssetDefinition_DynamicMaterialModel.cpp b/Engine/Plugins/Experimental/DynamicMaterial/Source/DynamicMaterialEditor/Private/Model/AssetDefinition_DynamicMaterialModel.cpp
--- a/Engine/Plugins/Experimental/DynamicMaterial/Source/DynamicMaterialEditor/Private/Model/AssetDefinition_DynamicMaterialModel.cpp
+++ b/Engine/Plugins/Experimental/DynamicMaterial/Source/DynamicMaterialEditor/Private/Model/AssetDefinition_DynamicMaterialModel.cpp
@@ -15,7 +15,7 @@ FText UAssetDefinition_DynamicMaterialModel::GetAssetDisplayName() const
 
 FText UAssetDefinition_DynamicMaterialModel::GetAssetDisplayName(const FAssetData& InAssetData) const
 {
-	return FText::FromName(InAssetData.AssetName);
+	return LOCTEXT("MaterialDesignerModel", "Material Designer Model");
 }
 
 TSoftClassPtr<> UAssetDefinition_DynamicMaterialModel::GetAssetClass() const
