Material Designer
- You can no longer create new textures or clear texture values in the Material Designer.

diff --git a/Engine/Plugins/Experimental/DynamicMaterial/Source/DynamicMaterial/Public/Components/MaterialValues/DMMaterialValueTexture.h b/Engine/Plugins/Experimental/DynamicMaterial/Source/DynamicMaterial/Public/Components/MaterialValues/DMMaterialValueTexture.h
--- a/Engine/Plugins/Experimental/DynamicMaterial/Source/DynamicMaterial/Public/Components/MaterialValues/DMMaterialValueTexture.h
+++ b/Engine/Plugins/Experimental/DynamicMaterial/Source/DynamicMaterial/Public/Components/MaterialValues/DMMaterialValueTexture.h
@@ -75,8 +75,8 @@ public:
 protected:
 	UDMMaterialValueTexture();
  
-	UPROPERTY(EditAnywhere, BlueprintReadWrite, Getter = GetValue, Setter = SetValue, BlueprintSetter = SetValue, Category = "Material Designer", 
-		meta = (DisplayThumbnail = true, DisplayName = "Texture", AllowPrivateAccess = "true", HighPriority, NotKeyframeable))
+	UPROPERTY(EditAnywhere, BlueprintReadWrite, NoClear, Getter = GetValue, Setter = SetValue, BlueprintSetter = SetValue, Category = "Material Designer",
+		meta = (DisplayThumbnail = true, DisplayName = "Texture", AllowPrivateAccess = "true", HighPriority, NotKeyframeable, NoCreate))
 	TObjectPtr<UTexture> Value;
 
 #if WITH_EDITORONLY_DATA
