Material Designer
- Effects on the drop down menu are now sorted alphabetically.

diff --git a/Engine/Plugins/Experimental/DynamicMaterial/Source/DynamicMaterialEditor/Private/Menus/DMMaterialSlotLayerAddEffectMenus.cpp b/Engine/Plugins/Experimental/DynamicMaterial/Source/DynamicMaterialEditor/Private/Menus/DMMaterialSlotLayerAddEffectMenus.cpp
--- a/Engine/Plugins/Experimental/DynamicMaterial/Source/DynamicMaterialEditor/Private/Menus/DMMaterialSlotLayerAddEffectMenus.cpp
+++ b/Engine/Plugins/Experimental/DynamicMaterial/Source/DynamicMaterialEditor/Private/Menus/DMMaterialSlotLayerAddEffectMenus.cpp
@@ -128,16 +128,41 @@ namespace UE::DynamicMaterialEditor::Private
 			return;
 		}
 
-		const TArray<FDMMaterialEffectList>& EffectList = Settings->GetEffectList();
+		TArray<FDMMaterialEffectList> EffectLists = Settings->GetEffectList();
 
-		if (!EffectList.IsValidIndex(InCategoryIndex))
+		if (!EffectLists.IsValidIndex(InCategoryIndex))
 		{
 			return;
 		}
 
+		EffectLists[InCategoryIndex].Effects.StableSort([](const TSoftObjectPtr<UMaterialFunctionInterface>& InA, const TSoftObjectPtr<UMaterialFunctionInterface>& InB)
+			{
+				UMaterialFunctionInterface* MaterialFunctionA = InA.LoadSynchronous();
+				UMaterialFunctionInterface* MaterialFunctionB = InB.LoadSynchronous();
+
+				if (!MaterialFunctionA && !MaterialFunctionA)
+				{
+					return false;
+				}
+
+				if (!MaterialFunctionA)
+				{
+					return false;
+				}
+
+				if (!MaterialFunctionB)
+				{
+					return true;
+				}
+
+				const FString MaterialFunctionNameA = MaterialFunctionA->GetUserExposedCaption();
+				const FString MaterialFunctionNameB = MaterialFunctionB->GetUserExposedCaption();
+				return MaterialFunctionNameA < MaterialFunctionNameB;
+			});
+
 		FToolMenuSection& Section = InMenu->AddSection("EffectList", LOCTEXT("EffectList", "Effect List"));
 
-		for (const TSoftObjectPtr<UMaterialFunctionInterface>& Effect : EffectList[InCategoryIndex].Effects)
+		for (const TSoftObjectPtr<UMaterialFunctionInterface>& Effect : EffectLists[InCategoryIndex].Effects)
 		{
 			UMaterialFunctionInterface* MaterialFunction = Effect.LoadSynchronous();
 
@@ -179,13 +204,13 @@ namespace UE::DynamicMaterialEditor::Private
 
 		FToolMenuSection& Section = InMenu->AddSection("AddEffect", LOCTEXT("AddEffect", "Add Effect"));
 
-		const TArray<FDMMaterialEffectList>& EffectList = Settings->GetEffectList();
+		TArray<FDMMaterialEffectList> EffectLists = Settings->GetEffectList();
 
-		for (int32 CategoryIndex = 0; CategoryIndex < EffectList.Num(); ++CategoryIndex)
+		for (int32 CategoryIndex = 0; CategoryIndex < EffectLists.Num(); ++CategoryIndex)
 		{
 			Section.AddSubMenu(
 				NAME_None,
-				FText::FromString(EffectList[CategoryIndex].Name),
+				FText::FromString(EffectLists[CategoryIndex].Name),
 				FText::GetEmpty(),
 				FNewToolMenuChoice(FNewToolMenuDelegate::CreateStatic(&GenerateAddEffectSubMenu, CategoryIndex))
 			);
