Material Designer
- Horde issue 470700

diff --git a/Engine/Plugins/Experimental/DynamicMaterial/Source/DynamicMaterialEditor/Private/Menus/DMMaterialSlotLayerAddEffectMenus.cpp b/Engine/Plugins/Experimental/DynamicMaterial/Source/DynamicMaterialEditor/Private/Menus/DMMaterialSlotLayerAddEffectMenus.cpp
--- a/Engine/Plugins/Experimental/DynamicMaterial/Source/DynamicMaterialEditor/Private/Menus/DMMaterialSlotLayerAddEffectMenus.cpp
+++ b/Engine/Plugins/Experimental/DynamicMaterial/Source/DynamicMaterialEditor/Private/Menus/DMMaterialSlotLayerAddEffectMenus.cpp
@@ -140,11 +140,6 @@ namespace UE::DynamicMaterialEditor::Private
 				UMaterialFunctionInterface* MaterialFunctionA = InA.LoadSynchronous();
 				UMaterialFunctionInterface* MaterialFunctionB = InB.LoadSynchronous();
 
-				if (!MaterialFunctionA && !MaterialFunctionA)
-				{
-					return false;
-				}
-
 				if (!MaterialFunctionA)
 				{
 					return false;
