Material Designer
- Effects lists are now open by default.
- When selecting a slot, the first layer/stage is automatically selected (top of the stack)
- When added an effect, the effect is automatically selected and the effects list expanded.

diff --git a/Engine/Plugins/Experimental/DynamicMaterial/Source/DynamicMaterialEditor/Private/Menus/DMMaterialSlotLayerAddEffectContext.h b/Engine/Plugins/Experimental/DynamicMaterial/Source/DynamicMaterialEditor/Private/Menus/DMMaterialSlotLayerAddEffectContext.h
--- a/Engine/Plugins/Experimental/DynamicMaterial/Source/DynamicMaterialEditor/Private/Menus/DMMaterialSlotLayerAddEffectContext.h
+++ b/Engine/Plugins/Experimental/DynamicMaterial/Source/DynamicMaterialEditor/Private/Menus/DMMaterialSlotLayerAddEffectContext.h
@@ -8,6 +8,7 @@
 #include "UObject/WeakObjectPtrTemplates.h"
 #include "DMMaterialSlotLayerAddEffectContext.generated.h"
 
+class SDMEditor;
 class UDMMaterialLayerObject;
 class UToolMenu;
 
@@ -19,10 +20,15 @@ class UDMMaterialSlotLayerAddEffectContext : public UObject
 public:
 	UDMMaterialSlotLayerAddEffectContext() = default;
 
+	void SetEditorWidget(const TSharedPtr<SDMEditor>& InEditor) { EditorWidgetWeak = InEditor; }
+
+	TSharedPtr<SDMEditor> GetEditorWidget() const { return EditorWidgetWeak.Pin(); }
+
 	void SetLayer(UDMMaterialLayerObject* InLayer) { LayerWeak = InLayer; }
 
 	UDMMaterialLayerObject* GetLayer() const { return LayerWeak.Get(); }
 
 private:
+	TWeakPtr<SDMEditor> EditorWidgetWeak;
 	TWeakObjectPtr<UDMMaterialLayerObject> LayerWeak;
 };
diff --git a/Engine/Plugins/Experimental/DynamicMaterial/Source/DynamicMaterialEditor/Private/Menus/DMMaterialSlotLayerAddEffectMenus.cpp b/Engine/Plugins/Experimental/DynamicMaterial/Source/DynamicMaterialEditor/Private/Menus/DMMaterialSlotLayerAddEffectMenus.cpp
--- a/Engine/Plugins/Experimental/DynamicMaterial/Source/DynamicMaterialEditor/Private/Menus/DMMaterialSlotLayerAddEffectMenus.cpp
+++ b/Engine/Plugins/Experimental/DynamicMaterial/Source/DynamicMaterialEditor/Private/Menus/DMMaterialSlotLayerAddEffectMenus.cpp
@@ -11,7 +11,11 @@
 #include "Materials/MaterialFunctionInterface.h"
 #include "Menus/DMMaterialSlotLayerAddEffectContext.h"
 #include "Menus/DMMenuContext.h"
+#include "Slate/Layers/SDMSlotLayerView.h"
+#include "Slate/SDMEditor.h"
+#include "Slate/SDMSlot.h"
 #include "ToolMenus.h"
+#include "Slate/Layers/SDMLayerEffectsView.h"
 #include "Utils/DMMaterialEffectStackPresetSubsystem.h"
 #include "Widgets/Input/SEditableTextBox.h"
 #include "Widgets/Notifications/SNotificationList.h"
@@ -73,14 +77,17 @@ namespace UE::DynamicMaterialEditor::Private
 	void AddEffect(const FToolMenuContext& InContext, TSoftObjectPtr<UMaterialFunctionInterface> InMaterialFunctionPtr)
 	{
 		const UDMMaterialLayerObject* Layer = nullptr;
+		TSharedPtr<SDMEditor> EditorWidget;
 
 		if (UDMMenuContext* MenuContext = InContext.FindContext<UDMMenuContext>())
 		{
 			Layer = MenuContext->GetLayer();
+			EditorWidget = MenuContext->GetEditorWidget();
 		}
 		else if (UDMMaterialSlotLayerAddEffectContext* SlotContext = InContext.FindContext<UDMMaterialSlotLayerAddEffectContext>())
 		{
 			Layer = SlotContext->GetLayer();
+			EditorWidget = SlotContext->GetEditorWidget();
 		}
 
 		if (!Layer || !IsValid(Layer))
@@ -112,6 +119,32 @@ namespace UE::DynamicMaterialEditor::Private
 		}
 
 		EffectStack->AddEffect(EffectFunction);
+
+		if (!EditorWidget.IsValid())
+		{
+			return;
+		}
+
+		EditorWidget->SetEditedComponent(EffectFunction);
+
+		if (TSharedPtr<SDMSlot> ActiveSlotWidget = EditorWidget->GetActiveSlotWidget())
+		{
+			if (TSharedPtr<SDMSlotLayerView> LayerView = ActiveSlotWidget->GetLayerView())
+			{
+				if (TSharedPtr<FDMMaterialLayerReference> LayerItem = LayerView->FindLayerItem(Layer))
+				{
+					if (TSharedPtr<SDMSlotLayerItem> LayerWiget = LayerView->WidgetFromLayerItem(LayerItem))
+					{
+						LayerWiget->SetEffectsExpanded(true);
+
+						if (TSharedPtr<SDMLayerEffectsView> EffectsList = LayerWiget->GetEffectsList())
+						{
+							EffectsList->SetSelectedItem(EffectFunction);
+						}
+					}
+				}
+			}
+		}
 	}
 
 	void GenerateAddEffectSubMenu(UToolMenu* InMenu, int32 InCategoryIndex)
@@ -500,13 +533,14 @@ namespace UE::DynamicMaterialEditor::Private
 	}
 }
 
-TSharedRef<SWidget> FDMMaterialSlotLayerAddEffectMenus::OpenAddEffectMenu(UDMMaterialLayerObject* InLayer)
+TSharedRef<SWidget> FDMMaterialSlotLayerAddEffectMenus::OpenAddEffectMenu(const TSharedPtr<SDMEditor>& InEditor, UDMMaterialLayerObject* InLayer)
 {
 	using namespace UE::DynamicMaterialEditor::Private;
 
 	RegisterAddEffectMenu();
 
 	UDMMaterialSlotLayerAddEffectContext* ContextObject = NewObject<UDMMaterialSlotLayerAddEffectContext>();
+	ContextObject->SetEditorWidget(InEditor);
 	ContextObject->SetLayer(InLayer);
 
 	return UToolMenus::Get()->GenerateWidget(AddEffectMenuName, FToolMenuContext(ContextObject));
diff --git a/Engine/Plugins/Experimental/DynamicMaterial/Source/DynamicMaterialEditor/Private/Menus/DMMaterialSlotLayerAddEffectMenus.h b/Engine/Plugins/Experimental/DynamicMaterial/Source/DynamicMaterialEditor/Private/Menus/DMMaterialSlotLayerAddEffectMenus.h
--- a/Engine/Plugins/Experimental/DynamicMaterial/Source/DynamicMaterialEditor/Private/Menus/DMMaterialSlotLayerAddEffectMenus.h
+++ b/Engine/Plugins/Experimental/DynamicMaterial/Source/DynamicMaterialEditor/Private/Menus/DMMaterialSlotLayerAddEffectMenus.h
@@ -5,6 +5,7 @@
 #include "Templates/SharedPointer.h"
 #include "UObject/SoftObjectPtr.h"
 
+class SDMEditor;
 class SWidget;
 class UDMMaterialLayerObject;
 class UMaterialFunctionInterface;
@@ -14,7 +15,7 @@ struct FToolMenuContext;
 class FDMMaterialSlotLayerAddEffectMenus
 {
 public:
-	static TSharedRef<SWidget> OpenAddEffectMenu(UDMMaterialLayerObject* InLayer);
+	static TSharedRef<SWidget> OpenAddEffectMenu(const TSharedPtr<SDMEditor>& InEditor, UDMMaterialLayerObject* InLayer);
 
 	static void AddEffectSubMenu(UToolMenu* InMenu, UDMMaterialLayerObject* InLayer);
 };
diff --git a/Engine/Plugins/Experimental/DynamicMaterial/Source/DynamicMaterialEditor/Private/Slate/Layers/SDMLayerEffectsView.cpp b/Engine/Plugins/Experimental/DynamicMaterial/Source/DynamicMaterialEditor/Private/Slate/Layers/SDMLayerEffectsView.cpp
--- a/Engine/Plugins/Experimental/DynamicMaterial/Source/DynamicMaterialEditor/Private/Slate/Layers/SDMLayerEffectsView.cpp
+++ b/Engine/Plugins/Experimental/DynamicMaterial/Source/DynamicMaterialEditor/Private/Slate/Layers/SDMLayerEffectsView.cpp
@@ -176,6 +176,20 @@ int32 SDMLayerEffectsView::GetLayerItemIndex(const TSharedPtr<FDMEffectsLayerIte
 	return GetItems().Find(TListTypeTraits<TSharedPtr<FDMEffectsLayerItem>>::NullableItemTypeConvertToItemType(Item));
 }
 
+bool SDMLayerEffectsView::SetSelectedItem(UDMMaterialEffect* InEffect)
+{
+	for (const TSharedPtr<FDMEffectsLayerItem>& EffectItem : EffectItems)
+	{
+		if (EffectItem->MaterialEffectWeak.Get() == InEffect)
+		{
+			SetSelection(EffectItem);
+			return true;
+		}
+	}
+
+	return false;
+}
+
 void SDMLayerEffectsView::RebuildList()
 {
 	EffectItems.Empty();
diff --git a/Engine/Plugins/Experimental/DynamicMaterial/Source/DynamicMaterialEditor/Private/Slate/Layers/SDMLayerEffectsView.h b/Engine/Plugins/Experimental/DynamicMaterial/Source/DynamicMaterialEditor/Private/Slate/Layers/SDMLayerEffectsView.h
--- a/Engine/Plugins/Experimental/DynamicMaterial/Source/DynamicMaterialEditor/Private/Slate/Layers/SDMLayerEffectsView.h
+++ b/Engine/Plugins/Experimental/DynamicMaterial/Source/DynamicMaterialEditor/Private/Slate/Layers/SDMLayerEffectsView.h
@@ -38,6 +38,8 @@ public:
 
 	int32 GetLayerItemCount() const { return EffectItems.Num(); }
 
+	bool SetSelectedItem(UDMMaterialEffect* InEffect);
+
 	//~ Begin SListView
 	virtual void RebuildList() override;
 	//~ End SListView
diff --git a/Engine/Plugins/Experimental/DynamicMaterial/Source/DynamicMaterialEditor/Private/Slate/Layers/SDMSlotLayerItem.cpp b/Engine/Plugins/Experimental/DynamicMaterial/Source/DynamicMaterialEditor/Private/Slate/Layers/SDMSlotLayerItem.cpp
--- a/Engine/Plugins/Experimental/DynamicMaterial/Source/DynamicMaterialEditor/Private/Slate/Layers/SDMSlotLayerItem.cpp
+++ b/Engine/Plugins/Experimental/DynamicMaterial/Source/DynamicMaterialEditor/Private/Slate/Layers/SDMSlotLayerItem.cpp
@@ -48,6 +48,8 @@ void SDMSlotLayerItem::Construct(const FArguments& InArgs, const TSharedPtr<SDMS
 	OnStageSelected = InArgs._OnStageSelected;
 	OnLayerLinkToggled = InArgs._OnLayerLinkToggled;
 
+	bDisplayEffectsList = true;
+
 	STableRow<TSharedPtr<FDMMaterialLayerReference>>::Construct(
 		STableRow<TSharedPtr<FDMMaterialLayerReference>>::FArguments()
 		.Padding(2.0f)
@@ -862,20 +864,25 @@ void SDMSlotLayerItem::OnStageClick(const FPointerEvent& InMouseEvent, const TSh
 
 	if (MouseButton == EKeys::LeftMouseButton)
 	{
-		if (TSharedPtr<STableViewBase> TableViewBase = LayerViewWeak.Pin())
-		{
-			TSharedPtr<SListView<TSharedPtr<FDMMaterialLayerReference>>> LayerView = StaticCastSharedPtr<SListView<TSharedPtr<FDMMaterialLayerReference>>>(TableViewBase);
-			LayerView->SetSelection(LayerItem);
-		}
+		SetSelectedStage(InStageWidget);
+	}
+}
 
-		if (SelectedStageWidget.IsValid())
-		{
-			OnStageSelected.ExecuteIfBound(false, SelectedStageWidget.ToSharedRef());
-		}
+void SDMSlotLayerItem::SetSelectedStage(const TSharedRef<SDMStage>& InStageWidget)
+{
+	if (TSharedPtr<STableViewBase> TableViewBase = LayerViewWeak.Pin())
+	{
+		TSharedPtr<SListView<TSharedPtr<FDMMaterialLayerReference>>> LayerView = StaticCastSharedPtr<SListView<TSharedPtr<FDMMaterialLayerReference>>>(TableViewBase);
+		LayerView->SetSelection(LayerItem);
+	}
 
-		SelectedStageWidget = InStageWidget;
-		OnStageSelected.ExecuteIfBound(true, InStageWidget);
+	if (SelectedStageWidget.IsValid())
+	{
+		OnStageSelected.ExecuteIfBound(false, SelectedStageWidget.ToSharedRef());
 	}
+
+	SelectedStageWidget = InStageWidget;
+	OnStageSelected.ExecuteIfBound(true, InStageWidget);
 }
 
 bool SDMSlotLayerItem::OnStageCanAcceptDrop(const FDragDropEvent& InDragDropEvent, const TSharedRef<SDMStage>& InStageWidget) const
@@ -1018,7 +1025,14 @@ TSharedRef<SWidget> SDMSlotLayerItem::CreateNewEffectMenu()
 {
 	if (LayerItem.IsValid())
 	{
-		return FDMMaterialSlotLayerAddEffectMenus::OpenAddEffectMenu(LayerItem->GetLayer());
+		TSharedPtr<SDMEditor> EditorWidget;
+
+		if (TSharedPtr<SDMSlot> SlotWidget = SlotWidgetWeak.Pin())
+		{
+			EditorWidget = SlotWidget->GetEditorWidget();
+		}
+
+		return FDMMaterialSlotLayerAddEffectMenus::OpenAddEffectMenu(EditorWidget, LayerItem->GetLayer());
 	}
 
 	return SNullWidget::NullWidget;
diff --git a/Engine/Plugins/Experimental/DynamicMaterial/Source/DynamicMaterialEditor/Private/Slate/Layers/SDMSlotLayerItem.h b/Engine/Plugins/Experimental/DynamicMaterial/Source/DynamicMaterialEditor/Private/Slate/Layers/SDMSlotLayerItem.h
--- a/Engine/Plugins/Experimental/DynamicMaterial/Source/DynamicMaterialEditor/Private/Slate/Layers/SDMSlotLayerItem.h
+++ b/Engine/Plugins/Experimental/DynamicMaterial/Source/DynamicMaterialEditor/Private/Slate/Layers/SDMSlotLayerItem.h
@@ -59,8 +59,14 @@ public:
 
 	FORCEINLINE bool AreEffectsExpanded() const { return bDisplayEffectsList; }
 
+	void SetEffectsExpanded(bool bInExpanded) { bDisplayEffectsList = bInExpanded; }
+
 	void DeselectAllEffects();
 
+	void SetSelectedStage(const TSharedRef<SDMStage>& InStageWidget);
+
+	TSharedPtr<SDMLayerEffectsView> GetEffectsList() const { return EffectsList; }
+
 protected:
 	TWeakPtr<SDMSlot> SlotWidgetWeak;
 	TWeakPtr<STableViewBase> LayerViewWeak;
diff --git a/Engine/Plugins/Experimental/DynamicMaterial/Source/DynamicMaterialEditor/Private/Slate/Layers/SDMSlotLayerView.cpp b/Engine/Plugins/Experimental/DynamicMaterial/Source/DynamicMaterialEditor/Private/Slate/Layers/SDMSlotLayerView.cpp
--- a/Engine/Plugins/Experimental/DynamicMaterial/Source/DynamicMaterialEditor/Private/Slate/Layers/SDMSlotLayerView.cpp
+++ b/Engine/Plugins/Experimental/DynamicMaterial/Source/DynamicMaterialEditor/Private/Slate/Layers/SDMSlotLayerView.cpp
@@ -102,6 +102,24 @@ TSharedRef<ITableRow> SDMSlotLayerView::OnGenerateLayerItemWidget(TSharedPtr<FDM
 		}
 	}
 
+	if (InItem.IsValid() && SelectedStageWidgets.IsEmpty() && !StageWidgets.IsEmpty())
+	{
+		if (InItem->IsBaseEnabled())
+		{
+			if (TSharedPtr<SDMStage> BaseStage = LayerItem->GetBaseStageWidget())
+			{
+				OnLayerStageSelected(true, BaseStage.ToSharedRef());
+			}
+		}
+		else if (InItem->IsMaskEnabled())
+		{
+			if (TSharedPtr<SDMStage> MaskStage = LayerItem->GetMaskStageWidget())
+			{
+				OnLayerStageSelected(true, MaskStage.ToSharedRef());
+			}
+		}
+	}
+
 	return LayerItem;
 }
 
@@ -735,7 +753,7 @@ void SDMSlotLayerView::PostRedo(bool bSuccess)
 	OnUndo();
 }
 
-TSharedPtr<FDMMaterialLayerReference> SDMSlotLayerView::FindLayerItem(UDMMaterialStage* const InStage) const
+TSharedPtr<FDMMaterialLayerReference> SDMSlotLayerView::FindLayerItem(const UDMMaterialStage* InStage) const
 {
 	if (!IsValid(InStage))
 	{
@@ -759,7 +777,7 @@ TSharedPtr<FDMMaterialLayerReference> SDMSlotLayerView::FindLayerItem(UDMMateria
 	return nullptr;
 }
 
-TSharedPtr<FDMMaterialLayerReference> SDMSlotLayerView::FindLayerItem(UDMMaterialLayerObject* const InLayer) const
+TSharedPtr<FDMMaterialLayerReference> SDMSlotLayerView::FindLayerItem(const UDMMaterialLayerObject* InLayer) const
 {
 	if (!MaterialSlotWeak.IsValid())
 	{
diff --git a/Engine/Plugins/Experimental/DynamicMaterial/Source/DynamicMaterialEditor/Private/Slate/Layers/SDMSlotLayerView.h b/Engine/Plugins/Experimental/DynamicMaterial/Source/DynamicMaterialEditor/Private/Slate/Layers/SDMSlotLayerView.h
--- a/Engine/Plugins/Experimental/DynamicMaterial/Source/DynamicMaterialEditor/Private/Slate/Layers/SDMSlotLayerView.h
+++ b/Engine/Plugins/Experimental/DynamicMaterial/Source/DynamicMaterialEditor/Private/Slate/Layers/SDMSlotLayerView.h
@@ -45,8 +45,8 @@ public:
 
 	TSharedPtr<SDMSlotLayerItem> WidgetFromLayerItem(const TSharedPtr<FDMMaterialLayerReference>& InItem);
 
-	TSharedPtr<FDMMaterialLayerReference> FindLayerItem(UDMMaterialStage* const InStage) const;
-	TSharedPtr<FDMMaterialLayerReference> FindLayerItem(UDMMaterialLayerObject* const InLayer) const;
+	TSharedPtr<FDMMaterialLayerReference> FindLayerItem(const UDMMaterialStage* InStage) const;
+	TSharedPtr<FDMMaterialLayerReference> FindLayerItem(const UDMMaterialLayerObject* InLayer) const;
 
 	void SelectLayerItem(UDMMaterialStage* const InStage, const bool bInSelected, const ESelectInfo::Type InSelectInfo);
 	void SelectLayerItem(UDMMaterialLayerObject* const InLayer, const bool bInMask, const bool bInSelected, const ESelectInfo::Type InSelectInfo);
diff --git a/Engine/Plugins/Experimental/DynamicMaterial/Source/DynamicMaterialEditor/Private/Slate/SDMEditor.cpp b/Engine/Plugins/Experimental/DynamicMaterial/Source/DynamicMaterialEditor/Private/Slate/SDMEditor.cpp
--- a/Engine/Plugins/Experimental/DynamicMaterial/Source/DynamicMaterialEditor/Private/Slate/SDMEditor.cpp
+++ b/Engine/Plugins/Experimental/DynamicMaterial/Source/DynamicMaterialEditor/Private/Slate/SDMEditor.cpp
@@ -986,7 +986,6 @@ void SDMEditor::Tick(const FGeometry& AllottedGeometry, const double InCurrentTi
 void SDMEditor::OnMaterialBuilt(UDynamicMaterialModel* InMaterialModel)
 {
 	RefreshSlotPickerList();
-	RefreshSlotWidget();
 	RefreshComponentEditWidget();
 }
 
diff --git a/Engine/Plugins/Experimental/DynamicMaterial/Source/DynamicMaterialEditor/Private/Slate/SDMSlot.cpp b/Engine/Plugins/Experimental/DynamicMaterial/Source/DynamicMaterialEditor/Private/Slate/SDMSlot.cpp
--- a/Engine/Plugins/Experimental/DynamicMaterial/Source/DynamicMaterialEditor/Private/Slate/SDMSlot.cpp
+++ b/Engine/Plugins/Experimental/DynamicMaterial/Source/DynamicMaterialEditor/Private/Slate/SDMSlot.cpp
@@ -93,27 +93,35 @@ void SDMSlot::Construct(const FArguments& InArgs, const TSharedRef<SDMEditor>& I
 	SlotPreviewSize = InArgs._SlotPreviewSize;
 	LayerPreviewSize = InArgs._LayerPreviewSize;
 
-	if (ensure(IsValid(InSlot)))
+	if (!ensure(IsValid(InSlot)))
 	{
-		InSlot->GetOnPropertiesUpdateDelegate().AddSP(this, &SDMSlot::OnSlotPropertiesUpdated);
-		InSlot->GetOnLayersUpdateDelegate().AddSP(this, &SDMSlot::OnSlotLayersUpdated);
+		return;
+	}
 
-		for (const UDMMaterialLayerObject* Layer : InSlot->GetLayers())
-		{
-			Layer->ForEachValidStage(
-				EDMMaterialLayerStage::All,
-				[](UDMMaterialStage* InStage)
-				{
-					InStage->SetBeingEdited(false);
-				});
-		}
+	InSlot->GetOnPropertiesUpdateDelegate().AddSP(this, &SDMSlot::OnSlotPropertiesUpdated);
+	InSlot->GetOnLayersUpdateDelegate().AddSP(this, &SDMSlot::OnSlotLayersUpdated);
 
-		RefreshMainWidget();
+	for (const UDMMaterialLayerObject* Layer : InSlot->GetLayers())
+	{
+		Layer->ForEachValidStage(
+			EDMMaterialLayerStage::All,
+			[](UDMMaterialStage* InStage)
+			{
+				InStage->SetBeingEdited(false);
+			});
+	}
 
-		bInvalidateMainWidget = false;
-		bInvalidateHeaderWidget = false;
-		bInvalidateSettingsWidget = false;
-		OnEndFrameDelegateHandle = FCoreDelegates::OnEndFrame.AddSP(this, &SDMSlot::HandleEndFrameRefresh);
+	RefreshMainWidget();
+
+	bInvalidateMainWidget = false;
+	bInvalidateHeaderWidget = false;
+	bInvalidateSettingsWidget = false;
+	OnEndFrameDelegateHandle = FCoreDelegates::OnEndFrame.AddSP(this, &SDMSlot::HandleEndFrameRefresh);
+
+	if (!InSlot->GetLayers().IsEmpty())
+	{
+		UDMMaterialLayerObject* Layer = InSlot->GetLayers().Last();
+		SetSelectedLayer(Layer);
 	}
 }
 
@@ -166,7 +174,7 @@ TSharedRef<SWidget> SDMSlot::GetLayerEffectsMenuContent()
 {
 	if (UDMMaterialLayerObject* LayerObject = GetSelectedLayer())
 	{
-		return FDMMaterialSlotLayerAddEffectMenus::OpenAddEffectMenu(LayerObject);
+		return FDMMaterialSlotLayerAddEffectMenus::OpenAddEffectMenu(EditorWidgetWeak.Pin(), LayerObject);
 	}
 
 	return SNullWidget::NullWidget;
@@ -1220,6 +1228,11 @@ UDMMaterialLayerObject* SDMSlot::AddNewLayer(UDMMaterialStage* InNewBaseStage, U
 	return Layer;
 }
 
+TSharedPtr<SDMSlotLayerView> SDMSlot::GetLayerView() const
+{
+	return LayerView;
+}
+
 void SDMSlot::AddNewLayer_NewLocalValue(TSubclassOf<UDMMaterialValue> InValueClass)
 {
 	UDMMaterialSlot* Slot = SlotWeak.Get();
diff --git a/Engine/Plugins/Experimental/DynamicMaterial/Source/DynamicMaterialEditor/Private/Slate/SDMSlot.h b/Engine/Plugins/Experimental/DynamicMaterial/Source/DynamicMaterialEditor/Private/Slate/SDMSlot.h
--- a/Engine/Plugins/Experimental/DynamicMaterial/Source/DynamicMaterialEditor/Private/Slate/SDMSlot.h
+++ b/Engine/Plugins/Experimental/DynamicMaterial/Source/DynamicMaterialEditor/Private/Slate/SDMSlot.h
@@ -89,6 +89,8 @@ public:
 
 	UDMMaterialLayerObject* AddNewLayer(UDMMaterialStage* InNewBaseStage = nullptr, UDMMaterialStage* InNewMaskStage = nullptr);
 
+	TSharedPtr<SDMSlotLayerView> GetLayerView() const;
+
 protected:
 	TWeakPtr<SDMEditor> EditorWidgetWeak;
 
