Wrap in WITH_EDITOR as it is only ever included inside WITH_EDITOR blocks and the use of <concepts> prevents compilation on some platforms when doing -includeheaders
#rnx

diff --git a/Engine/Plugins/Experimental/DynamicMaterial/Source/DynamicMaterial/Public/Model/IDMMaterialBuildUtilsInterface.h b/Engine/Plugins/Experimental/DynamicMaterial/Source/DynamicMaterial/Public/Model/IDMMaterialBuildUtilsInterface.h
--- a/Engine/Plugins/Experimental/DynamicMaterial/Source/DynamicMaterial/Public/Model/IDMMaterialBuildUtilsInterface.h
+++ b/Engine/Plugins/Experimental/DynamicMaterial/Source/DynamicMaterial/Public/Model/IDMMaterialBuildUtilsInterface.h
@@ -2,6 +2,8 @@
 
 #pragma once
 
+#if WITH_EDITOR
+
 #include "Containers/ContainersFwd.h"
 #include "Model/IDMMaterialBuildStateInterface.h"
 #include "Templates/SubclassOf.h"
@@ -63,3 +65,4 @@ struct IDMMaterialBuildUtilsInterface
 	virtual void UpdatePreviewMaterial(UMaterialExpression* InLastExpression, int32 InOutputIndex, int32 InOutputChannel, 
 		int32 InSize) const = 0;
 };
+#endif
\ No newline at end of file
