Material Designer
- Changed the alignment of the layer info to the left instead of center.

diff --git a/Engine/Plugins/Experimental/DynamicMaterial/Source/DynamicMaterialEditor/Private/Slate/Layers/SDMSlotLayerItem.cpp b/Engine/Plugins/Experimental/DynamicMaterial/Source/DynamicMaterialEditor/Private/Slate/Layers/SDMSlotLayerItem.cpp
--- a/Engine/Plugins/Experimental/DynamicMaterial/Source/DynamicMaterialEditor/Private/Slate/Layers/SDMSlotLayerItem.cpp
+++ b/Engine/Plugins/Experimental/DynamicMaterial/Source/DynamicMaterialEditor/Private/Slate/Layers/SDMSlotLayerItem.cpp
@@ -217,7 +217,7 @@ TSharedRef<SWidget> SDMSlotLayerItem::CreateHeaderRowContent()
 		]
 		+ SHorizontalBox::Slot()
 		.FillWidth(1.0f)
-		.HAlign(HAlign_Center)
+		.HAlign(HAlign_Left)
 		.VAlign(VAlign_Fill)
 		.Padding(2.0f, 0.0f, 10.0f, 0.0f)
 		[
