Material Designer
- Fonts are now uniform.
- Custom details panel entries uses details panel fonts.

diff --git a/Engine/Plugins/Experimental/DynamicMaterial/Source/DynamicMaterialEditor/Private/DetailsPanel/Slate/SDMDetailsPanelMaterialInterfaceWidget.cpp b/Engine/Plugins/Experimental/DynamicMaterial/Source/DynamicMaterialEditor/Private/DetailsPanel/Slate/SDMDetailsPanelMaterialInterfaceWidget.cpp
--- a/Engine/Plugins/Experimental/DynamicMaterial/Source/DynamicMaterialEditor/Private/DetailsPanel/Slate/SDMDetailsPanelMaterialInterfaceWidget.cpp
+++ b/Engine/Plugins/Experimental/DynamicMaterial/Source/DynamicMaterialEditor/Private/DetailsPanel/Slate/SDMDetailsPanelMaterialInterfaceWidget.cpp
@@ -1,16 +1,14 @@
 // Copyright Epic Games, Inc. All Rights Reserved.
 
 #include "DetailsPanel/Slate/SDMDetailsPanelMaterialInterfaceWidget.h"
-#include "AssetRegistry/AssetRegistryModule.h"
 #include "AssetToolsModule.h"
-#include "DetailLayoutBuilder.h"
+#include "DynamicMaterialEditorStyle.h"
 #include "IAssetTools.h"
 #include "Material/DynamicMaterialInstance.h"
 #include "Material/DynamicMaterialInstanceFactory.h"
 #include "PropertyCustomizationHelpers.h"
 #include "PropertyHandle.h"
 #include "Widgets/Input/SButton.h"
-#include "Widgets/Layout/SWrapBox.h"
 #include "Widgets/Text/STextBlock.h"
 
 #define LOCTEXT_NAMESPACE "SDMDetailsPanelMaterialInterfaceWidget"
@@ -71,8 +69,8 @@ void SDMDetailsPanelMaterialInterfaceWidget::Construct(const FArguments& InArgs,
 			.Content()
 			[
 				SNew(STextBlock)
+				.TextStyle(FDynamicMaterialEditorStyle::Get(), "RegularFont")
 				.Text(this, &SDMDetailsPanelMaterialInterfaceWidget::GetButtonText)
-				.Font(IDetailLayoutBuilder::GetDetailFontBold())
 			]
 		]
 	];
diff --git a/Engine/Plugins/Experimental/DynamicMaterial/Source/DynamicMaterialEditor/Private/DetailsPanel/Slate/SDMDetailsPanelTabSpawner.cpp b/Engine/Plugins/Experimental/DynamicMaterial/Source/DynamicMaterialEditor/Private/DetailsPanel/Slate/SDMDetailsPanelTabSpawner.cpp
--- a/Engine/Plugins/Experimental/DynamicMaterial/Source/DynamicMaterialEditor/Private/DetailsPanel/Slate/SDMDetailsPanelTabSpawner.cpp
+++ b/Engine/Plugins/Experimental/DynamicMaterial/Source/DynamicMaterialEditor/Private/DetailsPanel/Slate/SDMDetailsPanelTabSpawner.cpp
@@ -3,14 +3,13 @@
 #include "SDMDetailsPanelTabSpawner.h"
 #include "AssetRegistry/AssetRegistryModule.h"
 #include "AssetToolsModule.h"
-#include "DetailLayoutBuilder.h"
+#include "DynamicMaterialEditorStyle.h"
 #include "IAssetTools.h"
 #include "Model/DynamicMaterialModel.h"
 #include "Model/DynamicMaterialModelFactory.h"
 #include "PropertyCustomizationHelpers.h"
 #include "PropertyHandle.h"
 #include "Widgets/Input/SButton.h"
-#include "Widgets/Layout/SWrapBox.h"
 #include "Widgets/Text/STextBlock.h"
 
 #define LOCTEXT_NAMESPACE "SDMDetailsPanelTabSpawner"
@@ -64,8 +63,8 @@ void SDMDetailsPanelTabSpawner::Construct(const FArguments& InArgs, const TShare
 			.Content()
 			[
 				SNew(STextBlock)
+				.TextStyle(FDynamicMaterialEditorStyle::Get(), "RegularFont")
 				.Text(this, &SDMDetailsPanelTabSpawner::GetButtonText)
-				.Font(IDetailLayoutBuilder::GetDetailFontBold())
 			]
 		]
 	];
diff --git a/Engine/Plugins/Experimental/DynamicMaterial/Source/DynamicMaterialEditor/Private/DetailsPanel/Slate/SDMMaterialListExtensionWidget.cpp b/Engine/Plugins/Experimental/DynamicMaterial/Source/DynamicMaterialEditor/Private/DetailsPanel/Slate/SDMMaterialListExtensionWidget.cpp
--- a/Engine/Plugins/Experimental/DynamicMaterial/Source/DynamicMaterialEditor/Private/DetailsPanel/Slate/SDMMaterialListExtensionWidget.cpp
+++ b/Engine/Plugins/Experimental/DynamicMaterial/Source/DynamicMaterialEditor/Private/DetailsPanel/Slate/SDMMaterialListExtensionWidget.cpp
@@ -5,7 +5,7 @@
 #include "Components/PrimitiveComponent.h"
 #include "DMObjectMaterialProperty.h"
 #include "DMWorldSubsystem.h"
-#include "DetailLayoutBuilder.h"
+#include "DynamicMaterialEditorStyle.h"
 #include "Engine/World.h"
 #include "IAssetTools.h"
 #include "Material/DynamicMaterialInstance.h"
@@ -39,8 +39,8 @@ void SDMMaterialListExtensionWidget::Construct(const FArguments& InArgs, const T
 			.Content()
 			[
 				SNew(STextBlock)
+				.TextStyle(FDynamicMaterialEditorStyle::Get(), "RegularFont")
 				.Text(this, &SDMMaterialListExtensionWidget::GetButtonText)
-				.Font(IDetailLayoutBuilder::GetDetailFontBold())
 			]
 		]
 	];
diff --git a/Engine/Plugins/Experimental/DynamicMaterial/Source/DynamicMaterialEditor/Private/DynamicMaterialEditorStyle.cpp b/Engine/Plugins/Experimental/DynamicMaterial/Source/DynamicMaterialEditor/Private/DynamicMaterialEditorStyle.cpp
--- a/Engine/Plugins/Experimental/DynamicMaterial/Source/DynamicMaterialEditor/Private/DynamicMaterialEditorStyle.cpp
+++ b/Engine/Plugins/Experimental/DynamicMaterial/Source/DynamicMaterialEditor/Private/DynamicMaterialEditorStyle.cpp
@@ -505,9 +505,19 @@ void FDynamicMaterialEditorStyle::SetupTextStyles(const TSharedRef<FSlateStyleSe
 	HandleFontOutline.OutlineColor = LayerViewItemTextShadowColor;
 	HandleFontOutline.OutlineSize = 1;
 
-	FTextBlockStyle ActorNameTextStyle(NormalTextStyle);
-	ActorNameTextStyle.SetFont(DEFAULT_FONT("Regular", 10));
-	Style->Set("ActorName", ActorNameTextStyle);
+	FTextBlockStyle SmallTextStyle(NormalTextStyle);
+	SmallTextStyle.SetFont(DEFAULT_FONT("Regular", 8));
+	Style->Set("SmallFont", SmallTextStyle);
+
+	FTextBlockStyle RegularTextStyle(NormalTextStyle);
+	RegularTextStyle.SetFont(DEFAULT_FONT("Regular", 10));
+	Style->Set("RegularFont", RegularTextStyle);
+
+	FTextBlockStyle BoldTextStyle(NormalTextStyle);
+	BoldTextStyle.SetFont(DEFAULT_FONT("Bold", 10));
+	Style->Set("BoldFont", BoldTextStyle);
+
+	Style->Set("ActorName", RegularTextStyle);
 
 	FTextBlockStyle ActorNameBigTextStyle(NormalTextStyle);
 	ActorNameBigTextStyle.SetFont(DEFAULT_FONT("Regular", 14));
@@ -525,9 +535,7 @@ void FDynamicMaterialEditorStyle::SetupTextStyles(const TSharedRef<FSlateStyleSe
 	LayerViewItemFont.OutlineSettings = HandleFontOutline;
 	Style->Set("LayerView.Row.Font", LayerViewItemFont);
 
-	FSlateFontInfo LayerViewItemHandleSmallFont = DEFAULT_FONT("Regular", 10);
-	LayerViewItemHandleSmallFont.OutlineSettings = HandleFontOutline;
-	Style->Set("LayerView.Row.HandleFont", LayerViewItemHandleSmallFont);
+	Style->Set("LayerView.Row.HandleFont", RegularTextStyle);
 
 	FTextBlockStyle LayerViewItemTextStyle(NormalTextStyle);
 	LayerViewItemTextStyle.SetShadowOffset(FVector2D(1.0f, 1.0f));
@@ -541,15 +549,13 @@ void FDynamicMaterialEditorStyle::SetupTextStyles(const TSharedRef<FSlateStyleSe
 	Style->Set("LayerView.Row.HeaderText.Small",
 		FTextBlockStyle(LayerViewItemTextStyle)
 		.SetColorAndOpacity(FStyleColors::Foreground)
-		.SetFont(LayerViewItemHandleSmallFont));
+		.SetFont(RegularTextStyle.Font));
 
 	FTextBlockStyle StagePropertyDetailsTextStyle(NormalTextStyle);
 	StagePropertyDetailsTextStyle.SetFont(DEFAULT_FONT("Regular", 12));
 	Style->Set("Font.Stage.Details", StagePropertyDetailsTextStyle);
 
-	FTextBlockStyle StagePropertyDetailsBoldTextStyle(NormalTextStyle);
-	StagePropertyDetailsBoldTextStyle.SetFont(DEFAULT_FONT("Bold", 12));
-	Style->Set("Font.Stage.Details.Bold", StagePropertyDetailsBoldTextStyle);
+	Style->Set("Font.Stage.Details.Bold", BoldTextStyle);
 
 	FTextBlockStyle StagePropertyDetailsSmallTextStyle(NormalTextStyle);
 	StagePropertyDetailsSmallTextStyle.SetFont(IDetailLayoutBuilder::GetDetailFont());
diff --git a/Engine/Plugins/Experimental/DynamicMaterial/Source/DynamicMaterialEditor/Private/Slate/Layers/SDMLayerEffectsItem.cpp b/Engine/Plugins/Experimental/DynamicMaterial/Source/DynamicMaterialEditor/Private/Slate/Layers/SDMLayerEffectsItem.cpp
--- a/Engine/Plugins/Experimental/DynamicMaterial/Source/DynamicMaterialEditor/Private/Slate/Layers/SDMLayerEffectsItem.cpp
+++ b/Engine/Plugins/Experimental/DynamicMaterial/Source/DynamicMaterialEditor/Private/Slate/Layers/SDMLayerEffectsItem.cpp
@@ -1,6 +1,7 @@
 // Copyright Epic Games, Inc. All Rights Reserved.
 
 #include "Slate/Layers/SDMLayerEffectsItem.h"
+
 #include "Components/DMMaterialEffect.h"
 #include "Components/DMMaterialEffectStack.h"
 #include "Components/DMMaterialStage.h"
@@ -70,6 +71,7 @@ TSharedRef<SWidget> SDMLayerEffectsItem::CreateMainContent()
 			.Padding(10.0f, 0.0f, 0.0f, 0.0f)
 			[
 				SNew(STextBlock)
+				.TextStyle(FDynamicMaterialEditorStyle::Get(), "RegularFont")
 				.Text(this, &SDMLayerEffectsItem::GetLayerHeaderText)
 			]
 			+ SHorizontalBox::Slot()
diff --git a/Engine/Plugins/Experimental/DynamicMaterial/Source/DynamicMaterialEditor/Private/Slate/Layers/SDMSlotLayerItem.cpp b/Engine/Plugins/Experimental/DynamicMaterial/Source/DynamicMaterialEditor/Private/Slate/Layers/SDMSlotLayerItem.cpp
--- a/Engine/Plugins/Experimental/DynamicMaterial/Source/DynamicMaterialEditor/Private/Slate/Layers/SDMSlotLayerItem.cpp
+++ b/Engine/Plugins/Experimental/DynamicMaterial/Source/DynamicMaterialEditor/Private/Slate/Layers/SDMSlotLayerItem.cpp
@@ -11,6 +11,7 @@
 #include "Components/DMMaterialStageThroughput.h"
 #include "Components/MaterialStageInputs/DMMSITextureUV.h"
 #include "ContentBrowserDataDragDropOp.h"
+#include "DetailLayoutBuilder.h"
 #include "DragDrop/DMSlotLayerDragDropOperation.h"
 #include "DynamicMaterialEditorSettings.h"
 #include "DynamicMaterialEditorStyle.h"
@@ -240,7 +241,7 @@ TSharedRef<SWidget> SDMSlotLayerItem::CreateHeaderRowContent()
 			.Padding(0.0f, 0.0f, 0.0f, 0.0f)
 			[
 				SNew(STextBlock)
-				.Font(FAppStyle::GetFontStyle("PropertyWindow.NormalFont"))
+				.TextStyle(FDynamicMaterialEditorStyle::Get(), "SmallFont")
 				.Text(this, &SDMSlotLayerItem::GetBlendModeText)
 			]
 			+ SVerticalBox::Slot()
@@ -250,7 +251,7 @@ TSharedRef<SWidget> SDMSlotLayerItem::CreateHeaderRowContent()
 			.Padding(0.0f, 0.0f, 0.0f, 2.0f)
 			[
 				SNew(STextBlock)
-				.Font(FAppStyle::GetFontStyle("PropertyWindow.NormalFont"))
+				.TextStyle(FDynamicMaterialEditorStyle::Get(), "SmallFont")
 				.Text(this, &SDMSlotLayerItem::GetStageDescription)
 			]
 		]
@@ -343,8 +344,8 @@ TSharedRef<SWidget> SDMSlotLayerItem::CreateHandleWidget()
 				.VAlign(VAlign_Center)
 				[
 					SNew(STextBlock)
-					.TextStyle(&FDynamicMaterialEditorStyle::Get(), "LayerView.Row.HeaderText.Small")
 					.Justification(ETextJustify::Center)
+					.TextStyle(&FDynamicMaterialEditorStyle::Get(), "LayerView.Row.HeaderText.Small")
 					.Text(this, &SDMSlotLayerItem::GetLayerIndexText)
 				]
 			]
@@ -1235,9 +1236,9 @@ FVector2D SDMSlotLayerItem::GetStagePreviewSize() const
 TSharedRef<SWidget> SDMSlotLayerItem::CreateLayerHeaderText() const
 {
 	TSharedRef<SWidget> TextBlock = SNew(STextBlock)
-		.Font(FAppStyle::GetFontStyle("PropertyWindow.NormalFont"))
-		.Text(this, &SDMSlotLayerItem::GetLayerHeaderText)
-		.ColorAndOpacity(FSlateColor(EStyleColor::PrimaryHover));
+		.ColorAndOpacity(FSlateColor(EStyleColor::PrimaryHover))
+		.TextStyle(FDynamicMaterialEditorStyle::Get(), "SmallFont")
+		.Text(this, &SDMSlotLayerItem::GetLayerHeaderText);
 
 	TextBlock->SetOnMouseButtonDown(FPointerEventHandler::CreateSPLambda(
 		this,
@@ -1273,7 +1274,7 @@ TSharedRef<SWidget> SDMSlotLayerItem::CreateLayerHeaderEditableText() const
 	}
 
 	return SNew(SEditableTextBox)
-		.Font(FAppStyle::GetFontStyle("PropertyWindow.NormalFont"))
+		.Font(IDetailLayoutBuilder::GetDetailFont())
 		.HintText(LOCTEXT("LayerName", "Layer Name"))
 		.IsEnabled(true)
 		.Text(LayerName)
diff --git a/Engine/Plugins/Experimental/DynamicMaterial/Source/DynamicMaterialEditor/Private/Slate/Properties/SDMDetailsGrid.cpp b/Engine/Plugins/Experimental/DynamicMaterial/Source/DynamicMaterialEditor/Private/Slate/Properties/SDMDetailsGrid.cpp
--- a/Engine/Plugins/Experimental/DynamicMaterial/Source/DynamicMaterialEditor/Private/Slate/Properties/SDMDetailsGrid.cpp
+++ b/Engine/Plugins/Experimental/DynamicMaterial/Source/DynamicMaterialEditor/Private/Slate/Properties/SDMDetailsGrid.cpp
@@ -1,7 +1,7 @@
 // Copyright Epic Games, Inc. All Rights Reserved.
 
 #include "Slate/Properties/SDMDetailsGrid.h"
-#include "Components/DMMaterialValue.h"
+
 #include "DMWorldSubsystem.h"
 #include "DynamicMaterialEditorStyle.h"
 #include "Engine/World.h"
@@ -231,6 +231,7 @@ TSharedRef<SWidget> SDMDetailsGrid::CreateDefaultLabel(const TAttribute<FText>&
 {
 	return
 		SNew(STextBlock)
+		.TextStyle(FDynamicMaterialEditorStyle::Get(), "RegularFont")
 		.Text(InText);
 }
 
diff --git a/Engine/Plugins/Experimental/DynamicMaterial/Source/DynamicMaterialEditor/Private/Slate/Properties/SDMDetailsGridRow.cpp b/Engine/Plugins/Experimental/DynamicMaterial/Source/DynamicMaterialEditor/Private/Slate/Properties/SDMDetailsGridRow.cpp
--- a/Engine/Plugins/Experimental/DynamicMaterial/Source/DynamicMaterialEditor/Private/Slate/Properties/SDMDetailsGridRow.cpp
+++ b/Engine/Plugins/Experimental/DynamicMaterial/Source/DynamicMaterialEditor/Private/Slate/Properties/SDMDetailsGridRow.cpp
@@ -1,8 +1,8 @@
 // Copyright Epic Games, Inc. All Rights Reserved.
 
 #include "Slate/Properties/SDMDetailsGridRow.h"
-#include "Components/DMMaterialValue.h"
-#include "PropertyCustomizationHelpers.h"
+
+#include "DynamicMaterialEditorStyle.h"
 #include "Slate/Properties/SDMPropertyEdit.h"
 #include "Widgets/Colors/SColorBlock.h"
 #include "Widgets/Layout/SSplitter.h"
@@ -115,6 +115,7 @@ void SDMDetailsGridRow::AddPropertyRow(const FText& InText, const TSharedRef<SWi
 		.VAlign(VAlign_Center)
 		[
 			SNew(STextBlock)
+			.TextStyle(FDynamicMaterialEditorStyle::Get(), "RegularFont")
 			.Text(InText)
 		];
 
diff --git a/Engine/Plugins/Experimental/DynamicMaterial/Source/DynamicMaterialEditor/Private/Slate/Properties/SDMMaterialParameters.cpp b/Engine/Plugins/Experimental/DynamicMaterial/Source/DynamicMaterialEditor/Private/Slate/Properties/SDMMaterialParameters.cpp
--- a/Engine/Plugins/Experimental/DynamicMaterial/Source/DynamicMaterialEditor/Private/Slate/Properties/SDMMaterialParameters.cpp
+++ b/Engine/Plugins/Experimental/DynamicMaterial/Source/DynamicMaterialEditor/Private/Slate/Properties/SDMMaterialParameters.cpp
@@ -7,7 +7,6 @@
 #include "DMValueDefinition.h"
 #include "DynamicMaterialEditorModule.h"
 #include "DynamicMaterialEditorStyle.h"
-#include "DynamicMaterialModule.h"
 #include "Model/DynamicMaterialModel.h"
 #include "Slate/Properties/SDMDetailsGrid.h"
 #include "SlateOptMacros.h"
@@ -73,8 +72,8 @@ void SDMMaterialParameters::Construct(const FArguments& InArgs, TWeakObjectPtr<U
 			.VAlign(EVerticalAlignment::VAlign_Center)
 			[
 				SNew(STextBlock)
+				.TextStyle(FDynamicMaterialEditorStyle::Get(), "BoldFont")
 				.Text(LOCTEXT("GlobalValuesHeader", "Global Values"))
-				.Font(IDetailLayoutBuilder::GetDetailFontBold())
 			]
 			+ CreateAddValueButton(EDMValueType::VT_Float1)
 			+ CreateAddValueButton(EDMValueType::VT_Float2)
@@ -110,8 +109,8 @@ TSharedRef<SWidget> SDMMaterialParameters::CreateValuesWidget()
 			.Padding(10.f, 5.f, 10.f, 5.f)
 			[
 				SNew(STextBlock)
+				.TextStyle(FDynamicMaterialEditorStyle::Get(), "RegularFont")
 				.Text(LOCTEXT("GlobalValuesContent", "Global Values Content"))
-				.Font(IDetailLayoutBuilder::GetDetailFont())
 			];
 		// @formatter:on
 	}
@@ -128,8 +127,8 @@ TSharedRef<SWidget> SDMMaterialParameters::CreateValuesWidget()
 			.Padding(10.f, 5.f, 10.f, 5.f)
 			[
 				SNew(STextBlock)
+				.TextStyle(FDynamicMaterialEditorStyle::Get(), "RegularFont")
 				.Text(LOCTEXT("GlobalValuesNoContent", "Click a button to add a value!"))
-				.Font(IDetailLayoutBuilder::GetDetailFont())
 			];
 		// @formatter:on
 	}
diff --git a/Engine/Plugins/Experimental/DynamicMaterial/Source/DynamicMaterialEditor/Private/Slate/Properties/SDMMaterialProperty.cpp b/Engine/Plugins/Experimental/DynamicMaterial/Source/DynamicMaterialEditor/Private/Slate/Properties/SDMMaterialProperty.cpp
--- a/Engine/Plugins/Experimental/DynamicMaterial/Source/DynamicMaterialEditor/Private/Slate/Properties/SDMMaterialProperty.cpp
+++ b/Engine/Plugins/Experimental/DynamicMaterial/Source/DynamicMaterialEditor/Private/Slate/Properties/SDMMaterialProperty.cpp
@@ -2,14 +2,13 @@
 
 #include "SDMMaterialProperty.h"
 #include "Components/DMMaterialProperty.h"
-#include "DetailLayoutBuilder.h"
+#include "DynamicMaterialEditorStyle.h"
 #include "Model/DynamicMaterialModel.h"
 #include "Model/DynamicMaterialModelEditorOnlyData.h"
 #include "Slate/SDMEditor.h"
 #include "Slate/SDMSlot.h"
 #include "Styling/StyleColors.h"
 #include "Widgets/Colors/SColorBlock.h"
-#include "Widgets/SCanvas.h"
 #include "Widgets/SOverlay.h"
 #include "Widgets/Text/STextBlock.h"
 
@@ -36,30 +35,33 @@ void SDMMaterialProperty::Construct(const FArguments& InArgs, const TSharedRef<S
 				if (ensure(MaterialProperty) && ensure(IsValid(*MaterialProperty)))
 				{
 					ChildSlot
+					[
+						SNew(SOverlay)
+
+						+ SOverlay::Slot()
+						[
+							SNew(SColorBlock)
+							.Color(FLinearColor::Black)
+							.CornerRadius(FVector4(6.f, 6.f, 6.f, 6.f))
+						]
+
+						+ SOverlay::Slot()
+						.Padding(1.f)
+						[
+							SNew(SColorBlock)
+							.Color(FStyleColors::Highlight.GetSpecifiedColor())
+							.CornerRadius(FVector4(5.f, 5.f, 5.f, 5.f))
+						]
+
+						+ SOverlay::Slot()
 						[
-							SNew(SOverlay)
-								+ SOverlay::Slot()
-								[
-									SNew(SColorBlock)
-										.Color(FLinearColor::Black)
-										.CornerRadius(FVector4(6.f, 6.f, 6.f, 6.f))
-								]
-								+ SOverlay::Slot()
-								.Padding(1.f)
-								[
-									SNew(SColorBlock)
-										.Color(FStyleColors::Highlight.GetSpecifiedColor())
-										.CornerRadius(FVector4(5.f, 5.f, 5.f, 5.f))
-								]
-								+ SOverlay::Slot()
-								[
-									SNew(STextBlock)
-										.Text((*MaterialProperty)->GetDescription())
-										.Font(IDetailLayoutBuilder::GetDetailFont())
-										.ColorAndOpacity(FLinearColor::White)
-										.Margin(FMargin(8.f, 5.f, 8.f, 4.f))
-								]
-						];
+							SNew(STextBlock)
+							.TextStyle(FDynamicMaterialEditorStyle::Get(), "RegularFont")
+							.Text((*MaterialProperty)->GetDescription())
+							.ColorAndOpacity(FLinearColor::White)
+							.Margin(FMargin(8.f, 5.f, 8.f, 4.f))
+						]
+					];
 				}
 			}
 		}
diff --git a/Engine/Plugins/Experimental/DynamicMaterial/Source/DynamicMaterialEditor/Private/Slate/Properties/SDMPropertyEdit.cpp b/Engine/Plugins/Experimental/DynamicMaterial/Source/DynamicMaterialEditor/Private/Slate/Properties/SDMPropertyEdit.cpp
--- a/Engine/Plugins/Experimental/DynamicMaterial/Source/DynamicMaterialEditor/Private/Slate/Properties/SDMPropertyEdit.cpp
+++ b/Engine/Plugins/Experimental/DynamicMaterial/Source/DynamicMaterialEditor/Private/Slate/Properties/SDMPropertyEdit.cpp
@@ -5,17 +5,13 @@
 #include "DetailLayoutBuilder.h"
 #include "DetailRowMenuContext.h"
 #include "DMWorldSubsystem.h"
-#include "DynamicMaterialEditorModule.h"
-#include "DynamicMaterialEditorSettings.h"
 #include "Engine/Engine.h"
 #include "Engine/World.h"
 #include "Framework/Application/SlateApplication.h"
 #include "IDetailKeyframeHandler.h"
 #include "PropertyCustomizationHelpers.h"
 #include "PropertyHandle.h"
-#include "Slate/SDMComponentEdit.h"
 #include "Slate/SDMEditor.h"
-#include "Slate/SDMSlot.h"
 #include "ToolMenuContext.h"
 #include "ToolMenus.h"
 #include "Utils/DMPrivate.h"
diff --git a/Engine/Plugins/Experimental/DynamicMaterial/Source/DynamicMaterialEditor/Private/Slate/Properties/SDMSourceBlendType.cpp b/Engine/Plugins/Experimental/DynamicMaterial/Source/DynamicMaterialEditor/Private/Slate/Properties/SDMSourceBlendType.cpp
--- a/Engine/Plugins/Experimental/DynamicMaterial/Source/DynamicMaterialEditor/Private/Slate/Properties/SDMSourceBlendType.cpp
+++ b/Engine/Plugins/Experimental/DynamicMaterial/Source/DynamicMaterialEditor/Private/Slate/Properties/SDMSourceBlendType.cpp
@@ -1,6 +1,7 @@
 // Copyright Epic Games, Inc. All Rights Reserved.
 
 #include "Slate/Properties/SDMSourceBlendType.h"
+
 #include "Components/MaterialStageBlends/DMMSBAdd.h"
 #include "Components/MaterialStageBlends/DMMSBColor.h"
 #include "Components/MaterialStageBlends/DMMSBColorBurn.h"
@@ -28,6 +29,7 @@
 #include "Components/MaterialStageBlends/DMMSBSoftLight.h"
 #include "Components/MaterialStageBlends/DMMSBSubtract.h"
 #include "Components/MaterialStageBlends/DMMSBVividLight.h"
+#include "DynamicMaterialEditorStyle.h"
 #include "ToolMenu.h"
 #include "ToolMenuDelegates.h"
 #include "ToolMenus.h"
@@ -155,6 +157,7 @@ void SDMSourceBlendType::Construct(const FArguments& InArgs)
 		.ButtonContent()
 		[
 			SNew(STextBlock)
+			.TextStyle(FDynamicMaterialEditorStyle::Get(), "RegularFont")
 			.Text(this, &SDMSourceBlendType::GetSelectedItemText)
 		]
 	];
@@ -220,6 +223,7 @@ TSharedRef<SWidget> SDMSourceBlendType::OnGenerateWidget(const FName InItem)
 
 	return 
 		SNew(STextBlock)
+		.TextStyle(FDynamicMaterialEditorStyle::Get(), "RegularFont")
 		.Text(BlendMap[InItem].BlendName);
 }
 
diff --git a/Engine/Plugins/Experimental/DynamicMaterial/Source/DynamicMaterialEditor/Private/Slate/Properties/SDMTextureUVVisualizerPopout.cpp b/Engine/Plugins/Experimental/DynamicMaterial/Source/DynamicMaterialEditor/Private/Slate/Properties/SDMTextureUVVisualizerPopout.cpp
--- a/Engine/Plugins/Experimental/DynamicMaterial/Source/DynamicMaterialEditor/Private/Slate/Properties/SDMTextureUVVisualizerPopout.cpp
+++ b/Engine/Plugins/Experimental/DynamicMaterial/Source/DynamicMaterialEditor/Private/Slate/Properties/SDMTextureUVVisualizerPopout.cpp
@@ -8,7 +8,7 @@
 #include "CustomDetailsViewModule.h"
 #include "DetailLayoutBuilder.h"
 #include "DMWorldSubsystem.h"
-#include "DynamicMaterialEditorModule.h"
+#include "DynamicMaterialEditorStyle.h"
 #include "Engine/World.h"
 #include "Generators/DMTextureUVPropertyRowGenerator.h"
 #include "ICustomDetailsView.h"
@@ -242,8 +242,8 @@ TSharedRef<SWidget> SDMTextureUVVisualizerPopout::CreatePropertyWidget(UDMTextur
 			.Content()
 			[
 				SNew(STextBlock)
+				.Font(IDetailLayoutBuilder::GetDetailFont())
 				.Text(LOCTEXT("VisualizerOffset", "Offset"))
-				.Font(IDetailLayoutBuilder::GetDetailFontBold())
 			]
 		]
 		+ SHorizontalBox::Slot()
@@ -259,8 +259,8 @@ TSharedRef<SWidget> SDMTextureUVVisualizerPopout::CreatePropertyWidget(UDMTextur
 			.Content()
 			[
 				SNew(STextBlock)
+				.Font(IDetailLayoutBuilder::GetDetailFont())
 				.Text(LOCTEXT("VisualizerPivot", "Pivot"))
-				.Font(IDetailLayoutBuilder::GetDetailFontBold())
 			]
 		];
 
@@ -319,9 +319,9 @@ TSharedRef<SWidget> SDMTextureUVVisualizerPopout::CreatePropertyWidget(UDMTextur
 			Item->SetOverrideWidget(
 				ECustomDetailsViewWidgetType::Name,
 				SNew(STextBlock)
-				.Font(IDetailLayoutBuilder::GetDetailFont())
-				.Text(EditRow.NameOverride.GetValue())
-				.ToolTipText(EditRow.NameToolTipOverride.Get(FText::GetEmpty()))
+					.Font(IDetailLayoutBuilder::GetDetailFont())
+					.Text(EditRow.NameOverride.GetValue())
+					.ToolTipText(EditRow.NameToolTipOverride.Get(FText::GetEmpty()))
 			);
 		}
 
diff --git a/Engine/Plugins/Experimental/DynamicMaterial/Source/DynamicMaterialEditor/Private/Slate/Properties/SDMTextureUVVisualizerProperty.cpp b/Engine/Plugins/Experimental/DynamicMaterial/Source/DynamicMaterialEditor/Private/Slate/Properties/SDMTextureUVVisualizerProperty.cpp
--- a/Engine/Plugins/Experimental/DynamicMaterial/Source/DynamicMaterialEditor/Private/Slate/Properties/SDMTextureUVVisualizerProperty.cpp
+++ b/Engine/Plugins/Experimental/DynamicMaterial/Source/DynamicMaterialEditor/Private/Slate/Properties/SDMTextureUVVisualizerProperty.cpp
@@ -1,6 +1,7 @@
 // Copyright Epic Games, Inc. All Rights Reserved.
 
 #include "Slate/Properties/SDMTextureUVVisualizerProperty.h"
+
 #include "Components/DMMaterialStage.h"
 #include "Components/DMTextureUV.h"
 #include "DetailLayoutBuilder.h"
@@ -36,8 +37,8 @@ void SDMTextureUVVisualizerProperty::Construct(const FArguments& InArgs, UDMMate
 				.Content()
 				[
 					SNew(STextBlock)
+					.Font(IDetailLayoutBuilder::GetDetailFont())
 					.Text(LOCTEXT("ToggleVisualizer", "Toggle"))
-					.Font(IDetailLayoutBuilder::GetDetailFontBold())
 				]
 			]
 			+ SHorizontalBox::Slot()
@@ -48,8 +49,8 @@ void SDMTextureUVVisualizerProperty::Construct(const FArguments& InArgs, UDMMate
 				.Content()
 				[
 					SNew(STextBlock)
+					.Font(IDetailLayoutBuilder::GetDetailFont())
 					.Text(LOCTEXT("PopoutVisualizer", "Popout"))
-					.Font(IDetailLayoutBuilder::GetDetailFontBold())
 				]
 			]
 		]
@@ -87,8 +88,8 @@ void SDMTextureUVVisualizerProperty::Construct(const FArguments& InArgs, UDMMate
 				.Content()
 				[
 					SNew(STextBlock)
+					.Font(IDetailLayoutBuilder::GetDetailFont())
 					.Text(LOCTEXT("VisualizerOffset", "Offset"))
-					.Font(IDetailLayoutBuilder::GetDetailFontBold())
 				]
 			]
 			+ SHorizontalBox::Slot()
@@ -104,8 +105,8 @@ void SDMTextureUVVisualizerProperty::Construct(const FArguments& InArgs, UDMMate
 				.Content()
 				[
 					SNew(STextBlock)
+					.Font(IDetailLayoutBuilder::GetDetailFont())
 					.Text(LOCTEXT("VisualizerPivot", "Pivot"))
-					.Font(IDetailLayoutBuilder::GetDetailFontBold())
 				]
 			]
 		]
diff --git a/Engine/Plugins/Experimental/DynamicMaterial/Source/DynamicMaterialEditor/Private/Slate/SDMComponentEdit.cpp b/Engine/Plugins/Experimental/DynamicMaterial/Source/DynamicMaterialEditor/Private/Slate/SDMComponentEdit.cpp
--- a/Engine/Plugins/Experimental/DynamicMaterial/Source/DynamicMaterialEditor/Private/Slate/SDMComponentEdit.cpp
+++ b/Engine/Plugins/Experimental/DynamicMaterial/Source/DynamicMaterialEditor/Private/Slate/SDMComponentEdit.cpp
@@ -11,9 +11,7 @@
 #include "Components/DMMaterialStageSource.h"
 #include "Components/DMMaterialStageThroughput.h"
 #include "Components/DMMaterialValue.h"
-#include "Components/DMTextureUV.h"
 #include "Components/MaterialValues/DMMaterialValueFloat1.h"
-#include "Components/MaterialValues/DMMaterialValueFloat2.h"
 #include "CustomDetailsViewArgs.h"
 #include "CustomDetailsViewModule.h"
 #include "CustomDetailsViewSequencer.h"
@@ -421,9 +419,9 @@ TSharedRef<SWidget> SDMComponentEdit::CreateEditWidget()
 			Item->SetOverrideWidget(
 				ECustomDetailsViewWidgetType::Name,
 				SNew(STextBlock)
-				.Font(IDetailLayoutBuilder::GetDetailFont())
-				.Text(EditRow.NameOverride.GetValue())
-				.ToolTipText(EditRow.NameToolTipOverride.Get(FText::GetEmpty()))
+					.Font(IDetailLayoutBuilder::GetDetailFont())
+					.Text(EditRow.NameOverride.GetValue())
+					.ToolTipText(EditRow.NameToolTipOverride.Get(FText::GetEmpty()))
 			);
 		}
 
@@ -922,6 +920,7 @@ TSharedRef<SWidget> SDMComponentEdit::CreateSourceTypeEditWidget()
 				.ButtonContent()
 				[
 					SNew(STextBlock)
+					.Font(IDetailLayoutBuilder::GetDetailFont())
 					.Text(this, &SDMComponentEdit::GetSourceTypeEditWidgetText)
 				]
 		];
diff --git a/Engine/Plugins/Experimental/DynamicMaterial/Source/DynamicMaterialEditor/Private/Slate/SDMEditor.cpp b/Engine/Plugins/Experimental/DynamicMaterial/Source/DynamicMaterialEditor/Private/Slate/SDMEditor.cpp
--- a/Engine/Plugins/Experimental/DynamicMaterial/Source/DynamicMaterialEditor/Private/Slate/SDMEditor.cpp
+++ b/Engine/Plugins/Experimental/DynamicMaterial/Source/DynamicMaterialEditor/Private/Slate/SDMEditor.cpp
@@ -483,6 +483,7 @@ TSharedRef<SWidget> SDMEditor::CreateActorMaterialSlotSelector(const AActor* InA
 			SNew(STextBlock)
 			.Justification(ETextJustify::Center)
 			.AutoWrapText(true)
+			.TextStyle(FDynamicMaterialEditorStyle::Get(), "RegularFont")
 			.Text(LOCTEXT("NoMaterialSlot", "\n\nThe selected actor contains no primitive components with material slots."));
 	}
 
@@ -579,8 +580,8 @@ TSharedRef<SWidget> SDMEditor::CreateActorMaterialSlotSelector(const AActor* InA
 					.Padding(0.f, 5.f, 0.f, 5.f)
 					[
 						SNew(STextBlock)
+						.TextStyle(FDynamicMaterialEditorStyle::Get(), "RegularFont")
 						.Text(MaterialSlot->GetPropertyName(true))
-						.Font(IDetailLayoutBuilder::GetDetailFont())
 					]
 
 					+ SVerticalBox::Slot()
@@ -592,8 +593,8 @@ TSharedRef<SWidget> SDMEditor::CreateActorMaterialSlotSelector(const AActor* InA
 						.OnClicked(this, &SDMEditor::OnCreateMaterialButtonClicked, MaterialSlotWeak)
 						[
 							SNew(STextBlock)
+							.TextStyle(FDynamicMaterialEditorStyle::Get(), "RegularFont")
 							.Text(LOCTEXT("CreateMaterial", "Create Material"))
-							.Font(IDetailLayoutBuilder::GetDetailFont())
 						]
 					]
 				]
@@ -615,6 +616,7 @@ TSharedRef<SWidget> SDMEditor::GetEmptyContent()
 		.Padding(5.0f, 5.0f, 5.0f, 5.0f)
 		[
 			SNew(STextBlock)
+			.TextStyle(FDynamicMaterialEditorStyle::Get(), "RegularFont")
 			.Text(LOCTEXT("NoActiveMaterial", "No active Material Designer Instance."))
 		];
 }
@@ -771,8 +773,8 @@ TSharedRef<SWidget> SDMEditor::CreateParametersArea()
 			.VAlign(EVerticalAlignment::VAlign_Center)
 			[
 				SNew(STextBlock)
+				.Font(IDetailLayoutBuilder::GetDetailFont())
 				.Text(LOCTEXT("MaterialParameters", "Material Parameters"))
-				.Font(IDetailLayoutBuilder::GetDetailFontBold())
 			]
 		]
 		.BodyContent()
@@ -807,6 +809,7 @@ TSharedRef<SWidget> SDMEditor::CreateSlotPickerWidget()
 			.Content()
 			[
 				SNew(STextBlock)
+				.TextStyle(FDynamicMaterialEditorStyle::Get(), "RegularFont")
 				.Text(LOCTEXT("GlobalSettings", "Global"))
 			]
 		];
@@ -863,6 +866,7 @@ TSharedRef<SWidget> SDMEditor::CreateSlotPickerWidget()
 					.Content()
 					[
 						SNew(STextBlock)
+						.TextStyle(FDynamicMaterialEditorStyle::Get(), "RegularFont")
 						.Text(!ShortName.IsEmpty() ? FText::FromString(ShortName) : MaterialPropertyEnum->GetDisplayNameTextByValue(static_cast<int64>(Property.Key)))
 					]
 				];
@@ -1367,8 +1371,8 @@ TSharedRef<SWidget> SDMEditor::CreateSlotWidget()
 			.Padding(10.f, 5.f, 10.f, 5.f)
 			[
 				SNew(STextBlock)
+				.TextStyle(FDynamicMaterialEditorStyle::Get(), "RegularFont")
 				.Text(LOCTEXT("SlotsContent", "Select a Material Property to view the Layer Stack."))
-				.Font(IDetailLayoutBuilder::GetDetailFont())
 			];
 	};
 
diff --git a/Engine/Plugins/Experimental/DynamicMaterial/Source/DynamicMaterialEditor/Private/Slate/SDMMaterialWizard.cpp b/Engine/Plugins/Experimental/DynamicMaterial/Source/DynamicMaterialEditor/Private/Slate/SDMMaterialWizard.cpp
--- a/Engine/Plugins/Experimental/DynamicMaterial/Source/DynamicMaterialEditor/Private/Slate/SDMMaterialWizard.cpp
+++ b/Engine/Plugins/Experimental/DynamicMaterial/Source/DynamicMaterialEditor/Private/Slate/SDMMaterialWizard.cpp
@@ -1,6 +1,7 @@
 // Copyright Epic Games, Inc. All Rights Reserved.
 
 #include "Slate/SDMMaterialWizard.h"
+
 #include "DMDefs.h"
 #include "DynamicMaterialEditorSettings.h"
 #include "DynamicMaterialEditorStyle.h"
@@ -86,6 +87,7 @@ TSharedRef<SWidget> SDMMaterialWizard::CreateLayout()
 			.Padding(0.0f, SeparationDistance, 0.0f, TitleContentDistance)
 			[
 				SNew(STextBlock)
+				.TextStyle(FDynamicMaterialEditorStyle::Get(), "BoldFont")
 				.Text(LOCTEXT("MaterialType", "Material Type"))
 			]
 
@@ -103,6 +105,7 @@ TSharedRef<SWidget> SDMMaterialWizard::CreateLayout()
 			.Padding(0.0f, SeparationDistance, 0.0f, TitleContentDistance)
 			[
 				SNew(STextBlock)
+				.TextStyle(FDynamicMaterialEditorStyle::Get(), "BoldFont")
 				.Text(LOCTEXT("AvailableChannels", "Available Channels"))
 			]
 
@@ -148,6 +151,7 @@ TSharedRef<SWidget> SDMMaterialWizard::CreateChannelPresets()
 					.OnCheckStateChanged(this, &SDMMaterialWizard::Preset_OnChange, Preset.Name)
 					[
 						SNew(STextBlock)
+						.TextStyle(FDynamicMaterialEditorStyle::Get(), "RegularFont")
 						.Text(FText::FromName(Preset.Name))
 					]
 			];
@@ -192,6 +196,7 @@ TSharedRef<SWidget> SDMMaterialWizard::CreateChannelList()
 					.Padding(TextPadding)
 					[
 						SNew(STextBlock)
+						.TextStyle(FDynamicMaterialEditorStyle::Get(), "RegularFont")
 						.Text(!ShortName.IsEmpty() ? FText::FromString(ShortName) : MaterialPropertyEnum->GetDisplayNameTextByValue(static_cast<int64>(Property.Key)))
 					];
 			}
@@ -214,6 +219,7 @@ TSharedRef<SWidget> SDMMaterialWizard::CreateAcceptButton()
 			.OnClicked(this, &SDMMaterialWizard::Accept_OnClick)
 			[
 				SNew(STextBlock)
+				.TextStyle(FDynamicMaterialEditorStyle::Get(), "RegularFont")
 				.Text(LOCTEXT("Continue", "Continue"))
 			]
 		];
diff --git a/Engine/Plugins/Experimental/DynamicMaterial/Source/DynamicMaterialEditor/Private/Slate/SDMSlot.cpp b/Engine/Plugins/Experimental/DynamicMaterial/Source/DynamicMaterialEditor/Private/Slate/SDMSlot.cpp
--- a/Engine/Plugins/Experimental/DynamicMaterial/Source/DynamicMaterialEditor/Private/Slate/SDMSlot.cpp
+++ b/Engine/Plugins/Experimental/DynamicMaterial/Source/DynamicMaterialEditor/Private/Slate/SDMSlot.cpp
@@ -26,7 +26,6 @@
 #include "Components/MaterialValues/DMMaterialValueFloat3RGB.h"
 #include "Components/MaterialValues/DMMaterialValueFloat4.h"
 #include "Components/MaterialValues/DMMaterialValueTexture.h"
-#include "DetailLayoutBuilder.h"
 #include "DynamicMaterialEditorModule.h"
 #include "DynamicMaterialEditorSettings.h"
 #include "DynamicMaterialEditorStyle.h"
@@ -530,10 +529,10 @@ TSharedRef<SWidget> SDMSlot::CreateHeaderPropertyListWidget()
 				.HAlign(EHorizontalAlignment::HAlign_Fill)
 				[
 					SNew(STextBlock)
-					.Text(Pair.Key->GetDescription())
-					.Font(IDetailLayoutBuilder::GetDetailFont())
 					.ColorAndOpacity(FLinearColor::White)
 					.Margin(FMargin(16.f, 4.f, 4.f, 4.f))
+					.TextStyle(FDynamicMaterialEditorStyle::Get(), "RegularFont")
+					.Text(Pair.Key->GetDescription())
 				]
 			];
 	}
@@ -713,6 +712,7 @@ TSharedRef<SWidget> SDMSlot::CreateSlotSettingsRow()
 				.Padding(0.f, 0.f, 5.f, 0.f)
 				[
 					SNew(STextBlock)
+					.TextStyle(FDynamicMaterialEditorStyle::Get(), "RegularFont")
 					.Text(LOCTEXT("MaterialDesignerInstanceBlendMode", "Blend"))
 				]
 			]
@@ -730,11 +730,13 @@ TSharedRef<SWidget> SDMSlot::CreateSlotSettingsRow()
 					+ SWidgetSwitcher::Slot()
 					[
 						SNew(STextBlock)
+						.TextStyle(FDynamicMaterialEditorStyle::Get(), "RegularFont")
 						.Text(INVTEXT("-"))
 					]
 					+ SWidgetSwitcher::Slot()
 					[
 						SNew(STextBlock)
+						.TextStyle(FDynamicMaterialEditorStyle::Get(), "RegularFont")
 						.Text(INVTEXT("-"))
 					]
 					+ SWidgetSwitcher::Slot()
@@ -763,6 +765,7 @@ TSharedRef<SWidget> SDMSlot::CreateSlotSettingsRow()
 				.Padding(0.f, 0.f, 5.f, 0.f)
 				[
 					SNew(STextBlock)
+					.TextStyle(FDynamicMaterialEditorStyle::Get(), "RegularFont")
 					.Text(LOCTEXT("StageOpacity", "Opacity"))
 				]
 			]
@@ -777,7 +780,11 @@ TSharedRef<SWidget> SDMSlot::CreateSlotSettingsRow()
 				[
 					OpacityPropertyEditWidget.IsValid()
 						? StaticCastSharedRef<SWidget>(OpacityPropertyEditWidget.ToSharedRef())
-						: StaticCastSharedRef<SWidget>(SNew(STextBlock).Text(INVTEXT("-")))
+						: StaticCastSharedRef<SWidget>(
+							SNew(STextBlock)
+							.TextStyle(FDynamicMaterialEditorStyle::Get(), "RegularFont")
+							.Text(INVTEXT("-"))
+						)
 				]
 			]
 			+ SHorizontalBox::Slot()
diff --git a/Engine/Plugins/Experimental/DynamicMaterial/Source/DynamicMaterialEditor/Private/Slate/SDMStage.cpp b/Engine/Plugins/Experimental/DynamicMaterial/Source/DynamicMaterialEditor/Private/Slate/SDMStage.cpp
--- a/Engine/Plugins/Experimental/DynamicMaterial/Source/DynamicMaterialEditor/Private/Slate/SDMStage.cpp
+++ b/Engine/Plugins/Experimental/DynamicMaterial/Source/DynamicMaterialEditor/Private/Slate/SDMStage.cpp
@@ -1,6 +1,7 @@
 // Copyright Epic Games, Inc. All Rights Reserved.
 
 #include "SDMStage.h"
+
 #include "Components/DMMaterialLayer.h"
 #include "Components/DMMaterialSlot.h"
 #include "Components/DMMaterialStage.h"
@@ -23,10 +24,8 @@
 #include "SlateOptMacros.h"
 #include "Styling/CoreStyle.h"
 #include "Utils/DMPrivate.h"
-#include "Widgets/Colors/SColorBlock.h"
 #include "Widgets/Layout/SBox.h"
 #include "Widgets/Layout/SScaleBox.h"
-#include "Widgets/Layout/SScissorRectBox.h"
 #include "Widgets/SOverlay.h"
 #include "Widgets/SToolTip.h"
 #include "Widgets/Text/STextBlock.h"
@@ -161,8 +160,8 @@ TSharedRef<SWidget> SDMStage::CreateTextBlockBackground()
 					SNew(STextBlock)
 					.Margin(3.0f)
 					.ColorAndOpacity(FSlateColor(FLinearColor::White))
+					.TextStyle(FDynamicMaterialEditorStyle::Get(), "RegularFont")
 					.Text(this, &SDMStage::GetStageName)
-					.Font(FCoreStyle::GetDefaultFontStyle("Regular", DefaultLabelFontSize))
 				]
 			]
 		];
diff --git a/Engine/Plugins/Experimental/DynamicMaterial/Source/DynamicMaterialEditor/Private/Slate/SDMToolBar.cpp b/Engine/Plugins/Experimental/DynamicMaterial/Source/DynamicMaterialEditor/Private/Slate/SDMToolBar.cpp
--- a/Engine/Plugins/Experimental/DynamicMaterial/Source/DynamicMaterialEditor/Private/Slate/SDMToolBar.cpp
+++ b/Engine/Plugins/Experimental/DynamicMaterial/Source/DynamicMaterialEditor/Private/Slate/SDMToolBar.cpp
@@ -1,16 +1,12 @@
 // Copyright Epic Games, Inc. All Rights Reserved.
 
 #include "Slate/SDMToolBar.h"
-#include "AssetToolsModule.h"
 #include "ContentBrowserModule.h"
-#include "DMEDefs.h"
 #include "DynamicMaterialEditorModule.h"
 #include "DynamicMaterialEditorSettings.h"
 #include "DynamicMaterialEditorStyle.h"
-#include "DynamicMaterialModule.h"
 #include "EngineAnalytics.h"
 #include "GameFramework/Actor.h"
-#include "IAssetTools.h"
 #include "IContentBrowserSingleton.h"
 #include "Material/DynamicMaterialInstance.h"
 #include "Model/DynamicMaterialModel.h"
@@ -18,7 +14,6 @@
 #include "Styling/AppStyle.h"
 #include "Styling/StyleColors.h"
 #include "Utils/DMBlueprintFunctionLibrary.h"
-#include "Widgets/Colors/SColorBlock.h"
 #include "Widgets/Images/SImage.h"
 #include "Widgets/Input/SButton.h"
 #include "Widgets/Input/SComboBox.h"
@@ -92,6 +87,7 @@ TSharedRef<SWidget> SDMToolBar::CreateToolBarEntries()
 				.Padding(0.0f, 0.0f, 10.0f, 0.0f)
 				[
 					SNew(STextBlock)
+					.TextStyle(FDynamicMaterialEditorStyle::Get(), "RegularFont")
 					.Text(LOCTEXT("MaterialDesignerInstanceActorLabel", "Actor"))
 				]
 				+ SHorizontalBox::Slot()
@@ -118,6 +114,7 @@ TSharedRef<SWidget> SDMToolBar::CreateToolBarEntries()
 				.Padding(0.0f, 0.0f, 10.0f, 0.0f)
 				[
 					SNew(STextBlock)
+					.TextStyle(FDynamicMaterialEditorStyle::Get(), "RegularFont")
 					.Text(LOCTEXT("MaterialDesignerInstanceActorSlotLabel", "Property"))
 				]
 				+ SHorizontalBox::Slot()
@@ -224,6 +221,7 @@ TSharedRef<SWidget> SDMToolBar::CreateSlotsComboBoxWidget()
 		[
 			SNew(STextBlock)
 			.MinDesiredWidth(100.0f)
+			.TextStyle(FDynamicMaterialEditorStyle::Get(), "RegularFont")
 			.Text(this, &SDMToolBar::GetSelectedMaterialSlotName)
 		];
 }
@@ -234,6 +232,7 @@ TSharedRef<SWidget> SDMToolBar::GenerateSelectedMaterialSlotRow(TSharedPtr<FDMOb
 	{
 		return SNew(STextBlock)
 			.MinDesiredWidth(100.f)
+			.TextStyle(FDynamicMaterialEditorStyle::Get(), "RegularFont")
 			.Text(this, &SDMToolBar::GetSlotDisplayName, InSelectedSlot);
 	}
 	return SNullWidget::NullWidget;
