Material Designer
- Improved how the material designer detects and responds to invalid or incorrect asset displays.

diff --git a/Engine/Plugins/Experimental/DynamicMaterial/Source/DynamicMaterialEditor/Private/Slate/SDMEditor.cpp b/Engine/Plugins/Experimental/DynamicMaterial/Source/DynamicMaterialEditor/Private/Slate/SDMEditor.cpp
--- a/Engine/Plugins/Experimental/DynamicMaterial/Source/DynamicMaterialEditor/Private/Slate/SDMEditor.cpp
+++ b/Engine/Plugins/Experimental/DynamicMaterial/Source/DynamicMaterialEditor/Private/Slate/SDMEditor.cpp
@@ -186,6 +186,41 @@ namespace UE::DynamicMaterialEditor::Private
 			}
 		}
 	}
+
+	bool CheckMaterialModelValidity(UDynamicMaterialModel* InMaterialModel)
+		{
+			if (UWorld* World = InMaterialModel->GetWorld())
+			{
+				if (UDMWorldSubsystem* WorldSubsystem = World->GetSubsystem<UDMWorldSubsystem>())
+				{
+					if (WorldSubsystem->GetIsValidDelegate().IsBound()
+						&& WorldSubsystem->GetIsValidDelegate().Execute(InMaterialModel) == false)
+					{
+						return false;
+					}
+				}
+			}
+
+			UActorComponent* ComponentOuter = InMaterialModel->GetTypedOuter<UActorComponent>();
+			if (ComponentOuter && !IsValid(ComponentOuter))
+			{
+				return false;
+			}
+
+			AActor* ActorOuter = InMaterialModel->GetTypedOuter<AActor>();
+			if (ActorOuter && !IsValid(ActorOuter))
+			{
+				return false;
+			}
+
+			UPackage* PackageOuter = InMaterialModel->GetPackage();
+			if (PackageOuter && !IsValid(PackageOuter))
+			{
+				return false;
+			}
+
+			return true;
+		};
 }
 
 TSharedPtr<FAssetThumbnailPool> SDMEditor::ThumbnailPool = nullptr;
@@ -392,6 +427,7 @@ void SDMEditor::SetMaterialModel(UDynamicMaterialModel* InMaterialModel)
 	ActiveSlotWidget.Reset();
 	ComponentEditContainer.Reset();
 	SplitterContainer.Reset();
+	bInvalidateComponentEditWidget = false;
 
 	Toolbar->SetMaterialModel(InMaterialModel);
 	SetEditedComponent(nullptr);
@@ -399,9 +435,12 @@ void SDMEditor::SetMaterialModel(UDynamicMaterialModel* InMaterialModel)
 	if (!InMaterialModel)
 	{
 		Container->SetContent(SDMEditor::GetEmptyContent());
+		bHasActiveLayout = false;
 		return;
 	}
 
+	bHasActiveLayout = true;
+
 	UDynamicMaterialModelEditorOnlyData* EditorOnlyData = UDynamicMaterialModelEditorOnlyData::Get(InMaterialModel);
 
 	if (EditorOnlyData && EditorOnlyData->NeedsWizard())
@@ -910,74 +949,65 @@ void SDMEditor::Tick(const FGeometry& AllottedGeometry, const double InCurrentTi
 {
 	SCompoundWidget::Tick(AllottedGeometry, InCurrentTime, InDeltaTime);
 
-	if (!MaterialModelWeak.IsValid())
+	if (!FDynamicMaterialModule::AreUObjectsSafe())
 	{
 		return;
 	}
 
-	bool bValidModel = false;
+	UDynamicMaterialModel* MaterialModel = MaterialModelWeak.Get();
 
-	if (ObjectProperty.IsValid())
+	if (!IsValid(MaterialModel))
 	{
-		bValidModel = IsValid(ObjectProperty.GetMaterialModel());
+		if (bHasActiveLayout)
+		{
+			ClearEditor();
+		}
+
+		return;
 	}
-	else if (UDynamicMaterialModel* MaterialModel = MaterialModelWeak.Get())
+
+	if (ObjectProperty.IsValid())
 	{
-		auto CheckValidity = [MaterialModel]()
+		UDynamicMaterialModel* MaterialModelProperty = ObjectProperty.GetMaterialModel();
+
+		if (!IsValid(MaterialModelProperty))
 		{
-			if (UWorld* World = MaterialModel->GetWorld())
-			{
-				if (UDMWorldSubsystem* WorldSubsystem = World->GetSubsystem<UDMWorldSubsystem>())
-				{
-					// ExecuteIfBound doesn't work with return values
-					if (WorldSubsystem->GetIsValidDelegate().IsBound()
-						&& WorldSubsystem->GetIsValidDelegate().Execute(MaterialModel) == false)
-					{
-						return false;
-					}
-				}
-			}
+			MaterialModelProperty = nullptr;
+		}
 
-			UActorComponent* ComponentOuter = MaterialModel->GetTypedOuter<UActorComponent>();
-			if (ComponentOuter && !IsValid(ComponentOuter))
+		if (MaterialModel != MaterialModelProperty)
+		{
+			if (MaterialModelProperty)
 			{
-				return false;
+				// Re-set the property
+				SetMaterialObjectProperty(ObjectProperty);
+				return;
 			}
-
-			AActor* ActorOuter = MaterialModel->GetTypedOuter<AActor>();
-			if (ActorOuter && !IsValid(ActorOuter))
+			else
 			{
-				return false;
+				ClearEditor();
+				return;
 			}
-
-			UPackage* PackageOuter = MaterialModel->GetPackage();
-			if (PackageOuter && !IsValid(PackageOuter))
-			{
-				return false;
-			}
-
-			return true;
-		};
-
-		bValidModel = CheckValidity();
+		}
+	}
+	else 
+	{
+		if (!UE::DynamicMaterialEditor::Private::CheckMaterialModelValidity(MaterialModel))
+		{
+			ClearEditor();
+			return;
+		}
 	}
 
-	if (!bValidModel)
+	if (Toolbar.IsValid() && Toolbar->GetMaterialModel() != MaterialModel)
 	{
-		ClearEditor();
+		Toolbar->SetMaterialModel(MaterialModel);
 	}
-	else if (FDynamicMaterialModule::AreUObjectsSafe())
+	else if (bHasActiveLayout && ActiveSlotWidget.IsValid())
 	{
-		if (Toolbar.IsValid() && Toolbar->GetMaterialModel() != MaterialModelWeak)
+		if (!ActiveSlotWidget->CheckValidity())
 		{
-			Toolbar->SetMaterialModel(MaterialModelWeak.Get());
-		}
-		else if (ActiveSlotWidget.IsValid())
-		{
-			if (!ActiveSlotWidget->CheckValidity())
-			{
-				RefreshSlotWidget();
-			}
+			RefreshSlotWidget();
 		}
 	}
 
diff --git a/Engine/Plugins/Experimental/DynamicMaterial/Source/DynamicMaterialEditor/Private/Slate/SDMEditor.h b/Engine/Plugins/Experimental/DynamicMaterial/Source/DynamicMaterialEditor/Private/Slate/SDMEditor.h
--- a/Engine/Plugins/Experimental/DynamicMaterial/Source/DynamicMaterialEditor/Private/Slate/SDMEditor.h
+++ b/Engine/Plugins/Experimental/DynamicMaterial/Source/DynamicMaterialEditor/Private/Slate/SDMEditor.h
@@ -133,6 +133,7 @@ protected:
 	SSplitter::FSlot* LayerViewSplitterSlot = nullptr;
 	SSplitter::FSlot* ExtraSpaceSplitterSlot = nullptr;
 	TSharedPtr<SDMSlot> ActiveSlotWidget;
+	bool bHasActiveLayout = false;
 
 	TSharedPtr<FUICommandList> CommandList;
 	int32 ActiveSlotIndex;
