Material Designer
- Fixing the 5.4 to 5.5 upgrade - fixed variable names and added redirectors and added newly added values to old materials

diff --git a/Engine/Plugins/Experimental/DynamicMaterial/Config/BaseDynamicMaterial.ini b/Engine/Plugins/Experimental/DynamicMaterial/Config/BaseDynamicMaterial.ini
--- a/Engine/Plugins/Experimental/DynamicMaterial/Config/BaseDynamicMaterial.ini
+++ b/Engine/Plugins/Experimental/DynamicMaterial/Config/BaseDynamicMaterial.ini
@@ -16,3 +16,17 @@
 +PackageRedirects=(OldName="/DynamicMaterial/MaterialFunctions/Effects/MF_DM_Effect_Mask_Level.MF_DM_Effect_Mask_Level",NewName="/DynamicMaterial/MaterialFunctions/Effects/Alpha/MF_DM_Effect_Mask_Level.MF_DM_Effect_Mask_Level")
 +PackageRedirects=(OldName="/DynamicMaterial/MaterialFunctions/Effects/MF_DM_UV_Distort.MF_DM_UV_Distort",NewName="/DynamicMaterial/MaterialFunctions/Effects/UV/MF_DM_UV_Distort.MF_DM_UV_Distort")
 +PackageRedirects=(OldName="/DynamicMaterial/MaterialFunctions/Effects/MF_DM_Effect_UV_Pixelate.MF_DM_Effect_UV_Pixelate",NewName="/DynamicMaterial/MaterialFunctions/Effects/UV/MF_DM_Effect_UV_Pixelate.MF_DM_Effect_UV_Pixelate")
+
++PropertyRedirects=(OldName="/Script/DynamicMaterial.DMRenderTargetTextRenderer.TextScale",NewName="/Script/DynamicMaterial.DMRenderTargetTextRenderer.RenderScale")
++PackageRedirects=(OldName="/Script/DynamicMaterial.DynamicMaterialModel.GlobalOpacity",NewName="/Script/DynamicMaterial.DynamicMaterialModel.GlobalOpacityValue")
++PackageRedirects=(OldName="/Script/DynamicMaterial.DynamicMaterialModel.GlobalTiling",NewName="/Script/DynamicMaterial.DynamicMaterialModel.GlobalTilingValue")
++PackageRedirects=(OldName="/Script/DynamicMaterial.DynamicMaterialModel.GlobalMetallic",NewName="/Script/DynamicMaterial.DynamicMaterialModel.GlobalMetallicParameter")
++PackageRedirects=(OldName="/Script/DynamicMaterial.DynamicMaterialModel.GlobalSpecular",NewName="/Script/DynamicMaterial.DynamicMaterialModel.GlobalSpecularParameter")
++PackageRedirects=(OldName="/Script/DynamicMaterial.DynamicMaterialModel.GlobalRoughness",NewName="/Script/DynamicMaterial.DynamicMaterialModel.GlobalRoughnessParameter")
++PackageRedirects=(OldName="/Script/DynamicMaterial.DynamicMaterialModel.GlobalAnisotropy",NewName="/Script/DynamicMaterial.DynamicMaterialModel.GlobalAnisotropyParameter")
++PackageRedirects=(OldName="/Script/DynamicMaterial.DynamicMaterialModel.GlobalWorldPositionOffset",NewName="/Script/DynamicMaterial.DynamicMaterialModel.GlobalWorldPositionOffsetParameter")
++PackageRedirects=(OldName="/Script/DynamicMaterial.DynamicMaterialModel.GlobalAmbientOcclusion",NewName="/Script/DynamicMaterial.DynamicMaterialModel.GlobalAmbientOcclusionParameter")
++PackageRedirects=(OldName="/Script/DynamicMaterial.DynamicMaterialModel.GlobalRefraction",NewName="/Script/DynamicMaterial.DynamicMaterialModel.GlobalRefractionParameter")
++PackageRedirects=(OldName="/Script/DynamicMaterial.DynamicMaterialModel.GlobalPixelDepthOffset",NewName="/Script/DynamicMaterial.DynamicMaterialModel.GlobalPixelDepthOffsetParameter")
++PackageRedirects=(OldName="/Script/DynamicMaterial.DynamicMaterialModel.GlobalOffset",NewName="/Script/DynamicMaterial.DynamicMaterialModel.GlobalOffsetParameter")
++PackageRedirects=(OldName="/Script/DynamicMaterial.DynamicMaterialModel.GlobalRotation",NewName="/Script/DynamicMaterial.DynamicMaterialModel.GlobalRotationParameter")
diff --git a/Engine/Plugins/Experimental/DynamicMaterial/Source/DynamicMaterial/Private/Model/DynamicMaterialModel.cpp b/Engine/Plugins/Experimental/DynamicMaterial/Source/DynamicMaterial/Private/Model/DynamicMaterialModel.cpp
--- a/Engine/Plugins/Experimental/DynamicMaterial/Source/DynamicMaterial/Private/Model/DynamicMaterialModel.cpp
+++ b/Engine/Plugins/Experimental/DynamicMaterial/Source/DynamicMaterial/Private/Model/DynamicMaterialModel.cpp
@@ -67,7 +67,7 @@ UDynamicMaterialModel::UDynamicMaterialModel()
 			Property->ApplyDefaultValue();
 #endif
 
-			UDMMaterialParameter* Parameter = CreateDefaultSubobject<UDMMaterialParameter>(InParameterName);
+			UDMMaterialParameter* Parameter = CreateDefaultSubobject<UDMMaterialParameter>(FName(*(InParameterName.GetPlainNameString() + TEXT("Parameter"))));
 			Parameter->ParameterName = InParameterName;
 			Property->Parameter = Parameter;
 
@@ -84,7 +84,7 @@ UDynamicMaterialModel::UDynamicMaterialModel()
 			Property->ApplyDefaultValue();
 #endif
 
-			UDMMaterialParameter* Parameter = CreateDefaultSubobject<UDMMaterialParameter>(InParameterName);
+			UDMMaterialParameter* Parameter = CreateDefaultSubobject<UDMMaterialParameter>(FName(*(InParameterName.GetPlainNameString() + TEXT("Parameter"))));
 			Parameter->ParameterName = InParameterName;
 			Property->Parameter = Parameter;
 
diff --git a/Engine/Plugins/Experimental/DynamicMaterial/Source/DynamicMaterialEditor/Private/Model/DynamicMaterialModelEditorOnlyData.cpp b/Engine/Plugins/Experimental/DynamicMaterial/Source/DynamicMaterialEditor/Private/Model/DynamicMaterialModelEditorOnlyData.cpp
--- a/Engine/Plugins/Experimental/DynamicMaterial/Source/DynamicMaterialEditor/Private/Model/DynamicMaterialModelEditorOnlyData.cpp
+++ b/Engine/Plugins/Experimental/DynamicMaterial/Source/DynamicMaterialEditor/Private/Model/DynamicMaterialModelEditorOnlyData.cpp
@@ -38,6 +38,7 @@
 #include "Misc/Guid.h"
 #include "Model/DMMaterialBuildState.h"
 #include "Model/DynamicMaterialModel.h"
+#include "Model/DynamicMaterialModelEditorOnlyDataVersion.h"
 #include "UObject/Package.h"
 
 #define LOCTEXT_NAMESPACE "MaterialDesignerModel"
@@ -1347,7 +1348,7 @@ void UDynamicMaterialModelEditorOnlyData::NotifyPostChange(const FPropertyChange
 
 void UDynamicMaterialModelEditorOnlyData::PostLoad()
 {
-	Super::PostLoad();
+ 	Super::PostLoad();
 
 	// Backwards compatibility change - materials were originally parented to this object instead of the model.
 	if (IsValid(MaterialModel))
@@ -1425,6 +1426,33 @@ void UDynamicMaterialModelEditorOnlyData::PostEditChangeChainProperty(FPropertyC
 	}
 }
 
+void UDynamicMaterialModelEditorOnlyData::Serialize(FArchive& Ar)
+{
+	Ar.UsingCustomVersion(FDynamicMaterialModelEditorOnlyDataVersion::GUID);
+
+	Super::Serialize(Ar);
+
+	const int32 Version = Ar.CustomVer(FDynamicMaterialModelEditorOnlyDataVersion::GUID);
+
+	if (Version < FDynamicMaterialModelEditorOnlyDataVersion::GlobalValueRename)
+	{
+		if (ChannelListPreset == NAME_None)
+		{
+			// The default blend mode was translucent and now it's opaque.
+			if (BlendMode == BLEND_Opaque)
+			{
+				BlendMode = BLEND_Translucent;
+				ChannelListPreset = "Translucent";
+			}
+			// If the default blend mode was changed to anything else, set to All preset and let the user figure it out.
+			else
+			{
+				ChannelListPreset = "All";
+			}
+		}
+	}
+}
+
 void UDynamicMaterialModelEditorOnlyData::OnValueUpdated(UDMMaterialValue* InValue, EDMUpdateType InUpdateType)
 {
 	check(InValue);
diff --git a/Engine/Plugins/Experimental/DynamicMaterial/Source/DynamicMaterialEditor/Private/Model/DynamicMaterialModelEditorOnlyDataVersion.cpp b/Engine/Plugins/Experimental/DynamicMaterial/Source/DynamicMaterialEditor/Private/Model/DynamicMaterialModelEditorOnlyDataVersion.cpp
new file mode 100644
--- /dev/null
+++ b/Engine/Plugins/Experimental/DynamicMaterial/Source/DynamicMaterialEditor/Private/Model/DynamicMaterialModelEditorOnlyDataVersion.cpp
@@ -0,0 +1,7 @@
+// Copyright Epic Games, Inc. All Rights Reserved.
+
+#include "Model/DynamicMaterialModelEditorOnlyDataVersion.h"
+#include "Serialization/CustomVersion.h"
+
+const FGuid FDynamicMaterialModelEditorOnlyDataVersion::GUID(0xFCF57AFC, 0x50764285, 0xB9A9E660, 0xFFA02D34);
+FCustomVersionRegistration GRegisterDMTextureUVVersion(FDynamicMaterialModelEditorOnlyDataVersion::GUID, static_cast<int32>(FDynamicMaterialModelEditorOnlyDataVersion::Type::LatestVersion), TEXT("DynamicMaterialModelEditorOnlyData"));
diff --git a/Engine/Plugins/Experimental/DynamicMaterial/Source/DynamicMaterialEditor/Private/Model/DynamicMaterialModelEditorOnlyDataVersion.h b/Engine/Plugins/Experimental/DynamicMaterial/Source/DynamicMaterialEditor/Private/Model/DynamicMaterialModelEditorOnlyDataVersion.h
new file mode 100644
--- /dev/null
+++ b/Engine/Plugins/Experimental/DynamicMaterial/Source/DynamicMaterialEditor/Private/Model/DynamicMaterialModelEditorOnlyDataVersion.h
@@ -0,0 +1,26 @@
+// Copyright Epic Games, Inc. All Rights Reserved.
+
+#pragma once
+
+#include "HAL/Platform.h"
+#include "Misc/Guid.h"
+
+struct FDynamicMaterialModelEditorOnlyDataVersion
+{
+private:
+	FDynamicMaterialModelEditorOnlyDataVersion() = delete;
+
+public:
+	enum Type : uint8
+	{
+		PreVersioning = 0,
+
+		/** When the global variables were added and a couple were renamed. */
+		GlobalValueRename,
+
+		VersionPlusOne,
+		LatestVersion = VersionPlusOne - 1
+	};
+
+	static const FGuid GUID;
+};
diff --git a/Engine/Plugins/Experimental/DynamicMaterial/Source/DynamicMaterialEditor/Public/Model/DynamicMaterialModelEditorOnlyData.h b/Engine/Plugins/Experimental/DynamicMaterial/Source/DynamicMaterialEditor/Public/Model/DynamicMaterialModelEditorOnlyData.h
--- a/Engine/Plugins/Experimental/DynamicMaterial/Source/DynamicMaterialEditor/Public/Model/DynamicMaterialModelEditorOnlyData.h
+++ b/Engine/Plugins/Experimental/DynamicMaterial/Source/DynamicMaterialEditor/Public/Model/DynamicMaterialModelEditorOnlyData.h
@@ -195,6 +195,7 @@ public:
 	virtual void PostEditImport() override;
 	virtual void PostDuplicate(bool bDuplicateForPIE) override;
 	virtual void PostEditChangeChainProperty(FPropertyChangedChainEvent& PropertyChangedEvent) override;
+	virtual void Serialize(FArchive& Ar) override;
 	//~ End UObject
 
 	void SaveEditor();
