Material Designer
- Fix for 5.4 to 5.5 materials part 2

diff --git a/Engine/Plugins/Experimental/DynamicMaterial/Source/DynamicMaterialEditor/Private/Model/DynamicMaterialModelEditorOnlyData.cpp b/Engine/Plugins/Experimental/DynamicMaterial/Source/DynamicMaterialEditor/Private/Model/DynamicMaterialModelEditorOnlyData.cpp
--- a/Engine/Plugins/Experimental/DynamicMaterial/Source/DynamicMaterialEditor/Private/Model/DynamicMaterialModelEditorOnlyData.cpp
+++ b/Engine/Plugins/Experimental/DynamicMaterial/Source/DynamicMaterialEditor/Private/Model/DynamicMaterialModelEditorOnlyData.cpp
@@ -1438,16 +1438,23 @@ void UDynamicMaterialModelEditorOnlyData::Serialize(FArchive& Ar)
 	{
 		if (ChannelListPreset == NAME_None)
 		{
-			// The default blend mode was translucent and now it's opaque.
-			if (BlendMode == BLEND_Opaque)
+			// The default blend mode was changed from translucent and to opaque. If we have a
+			// transparent channel, let's set it back to translucent.
+			if (GetSlotForMaterialProperty(EDMMaterialPropertyType::Opacity))
 			{
+				ChannelListPreset = "Translucent";
 				BlendMode = BLEND_Translucent;
+			}
+			else if (GetSlotForMaterialProperty(EDMMaterialPropertyType::OpacityMask))
+			{
 				ChannelListPreset = "Translucent";
+				BlendMode = BLEND_Masked;
 			}
-			// If the default blend mode was changed to anything else, set to All preset and let the user figure it out.
+			// Else let's stay on opaque
 			else
 			{
-				ChannelListPreset = "All";
+				BlendMode = BLEND_Opaque;
+				ChannelListPreset = "Opaque";
 			}
 		}
 	}
