Material Designer
- Changed the asset opening flow.
- Added "Use selected asset in content browser" button for in-actor materials.
- Added "Browse to asset in content browser" button for assets.
- Fixed bug where Material Designer assets were reparented when assigned to Motion Design dynamic meshes.

diff --git a/Engine/Plugins/Experimental/Avalanche/Source/AvalancheEditor/Private/DetailView/Widgets/SAvaDynamicMaterialWidget.cpp b/Engine/Plugins/Experimental/Avalanche/Source/AvalancheEditor/Private/DetailView/Widgets/SAvaDynamicMaterialWidget.cpp
--- a/Engine/Plugins/Experimental/Avalanche/Source/AvalancheEditor/Private/DetailView/Widgets/SAvaDynamicMaterialWidget.cpp
+++ b/Engine/Plugins/Experimental/Avalanche/Source/AvalancheEditor/Private/DetailView/Widgets/SAvaDynamicMaterialWidget.cpp
@@ -1,17 +1,17 @@
 // Copyright Epic Games, Inc. All Rights Reserved.
 
 #include "SAvaDynamicMaterialWidget.h"
-#include "AssetRegistry/AssetRegistryModule.h"
-#include "AssetToolsModule.h"
+#include "Components/DynamicMeshComponent.h"
 #include "DetailLayoutBuilder.h"
-#include "IAssetTools.h"
+#include "DMObjectMaterialProperty.h"
+#include "DynamicMeshes/AvaShapeDynMeshBase.h"
+#include "IDynamicMaterialEditorModule.h"
 #include "Material/DynamicMaterialInstance.h"
 #include "Material/DynamicMaterialInstanceFactory.h"
 #include "PropertyCustomizationHelpers.h"
 #include "PropertyHandle.h"
 #include "ThumbnailRendering/ThumbnailManager.h"
 #include "Widgets/Input/SButton.h"
-#include "Widgets/Layout/SWrapBox.h"
 #include "Widgets/Text/STextBlock.h"
 
 #define LOCTEXT_NAMESPACE "SAvaDynamicMaterialWidget"
@@ -170,7 +170,7 @@ FReply SAvaDynamicMaterialWidget::CreateDynamicMaterialInstance()
 		UDynamicMaterialInstance::StaticClass(),
 		OuterObjects[0],
 		"DynamicMaterialInstance",
-		RF_NoFlags,
+		RF_Transactional,
 		nullptr,
 		GWarn
 	));
@@ -205,9 +205,47 @@ FReply SAvaDynamicMaterialWidget::OpenDynamicMaterialInstanceTab()
 		return FReply::Unhandled();
 	}
 
-	IAssetTools& AssetTools = FModuleManager::LoadModuleChecked<FAssetToolsModule>("AssetTools").Get();
-	AssetTools.OpenEditorForAssets({Instance});
+	const IDynamicMaterialEditorModule& MaterialDesignerModule = IDynamicMaterialEditorModule::Get();
+	constexpr bool bInvokeTab = true;
 
+	TArray<UObject*> Outers;
+	PropertyHandle->GetOuterObjects(Outers);
+
+	if (Outers.Num() == 0)
+	{
+		MaterialDesignerModule.OpenMaterialModel(Instance->GetMaterialModel(), nullptr, bInvokeTab);
+		return FReply::Handled();
+	}
+
+	UWorld* OuterWorld = Outers[0]->GetWorld();
+	UPrimitiveComponent* PrimitiveComponent = Cast<UPrimitiveComponent>(Outers[0]);
+
+	if (!PrimitiveComponent)
+	{
+		if (UAvaShapeDynamicMeshBase* ShapeMesh = Cast<UAvaShapeDynamicMeshBase>(Outers[0]))
+		{
+			PrimitiveComponent = ShapeMesh->GetShapeMeshComponent();
+		}
+
+		if (!PrimitiveComponent)
+		{
+			MaterialDesignerModule.OpenMaterialModel(Instance->GetMaterialModel(), nullptr, bInvokeTab);
+			return FReply::Handled();
+		}
+	}
+
+	const int32 NumMaterials = PrimitiveComponent->GetNumMaterials();
+
+	for (int32 Index = 0; Index < NumMaterials; ++Index)
+	{
+		if (PrimitiveComponent->GetMaterial(Index) == Instance)
+		{
+			MaterialDesignerModule.OpenMaterialObjectProperty({PrimitiveComponent, Index}, OuterWorld, bInvokeTab);
+			return FReply::Handled();
+		}
+	}
+
+	MaterialDesignerModule.OpenMaterialModel(Instance->GetMaterialModel(), OuterWorld, bInvokeTab);
 	return FReply::Handled();
 }
 
diff --git a/Engine/Plugins/Experimental/Avalanche/Source/AvalancheEditor/Private/MaterialDesigner/AvaMaterialDesignerExtension.cpp b/Engine/Plugins/Experimental/Avalanche/Source/AvalancheEditor/Private/MaterialDesigner/AvaMaterialDesignerExtension.cpp
--- a/Engine/Plugins/Experimental/Avalanche/Source/AvalancheEditor/Private/MaterialDesigner/AvaMaterialDesignerExtension.cpp
+++ b/Engine/Plugins/Experimental/Avalanche/Source/AvalancheEditor/Private/MaterialDesigner/AvaMaterialDesignerExtension.cpp
@@ -88,7 +88,11 @@ bool FAvaMaterialDesignerExtension::SetDynamicMaterialValue(const FDMObjectMater
 			UAvaShapeDynamicMeshBase* const DynamicMesh = ShapeActor->GetDynamicMesh();
 			if (DynamicMesh && DynamicMesh->GetMeshesIndexes().Contains(InObjectMaterialProperty.Index))
 			{
-				InMaterial->Rename(nullptr, DynamicMesh);
+				if (!InMaterial->IsAsset())
+				{
+					InMaterial->Rename(nullptr, DynamicMesh);
+				}
+
 				DynamicMesh->SetMaterial(InObjectMaterialProperty.Index, InMaterial);
 				return true;
 			}
diff --git a/Engine/Plugins/Experimental/DynamicMaterial/Source/DynamicMaterialEditor/Private/DetailsPanel/DMPropertyTypeCustomizer.h b/Engine/Plugins/Experimental/DynamicMaterial/Source/DynamicMaterialEditor/Private/DetailsPanel/DMPropertyTypeCustomizer.h
--- a/Engine/Plugins/Experimental/DynamicMaterial/Source/DynamicMaterialEditor/Private/DetailsPanel/DMPropertyTypeCustomizer.h
+++ b/Engine/Plugins/Experimental/DynamicMaterial/Source/DynamicMaterialEditor/Private/DetailsPanel/DMPropertyTypeCustomizer.h
@@ -12,19 +12,16 @@ class IPropertyTypeCustomizationUtils;
 
 class FDMPropertyTypeCustomizer : public IPropertyTypeCustomization
 {
-
 public:
-
 	static TSharedRef<IPropertyTypeCustomization> MakeInstance();
 
 	FDMPropertyTypeCustomizer();
 
-	/** BEGIN IPropertyTypeCustomization interface */
+	/** Begin IPropertyTypeCustomization interface */
 	virtual void CustomizeHeader(TSharedRef<IPropertyHandle> EditorPropertyHandle, class FDetailWidgetRow& HeaderRow,
 		IPropertyTypeCustomizationUtils& CustomizationUtils) override;
 
 	virtual void CustomizeChildren(TSharedRef<IPropertyHandle> EditorPropertyHandle, class IDetailChildrenBuilder& ChildBuilder,
 		IPropertyTypeCustomizationUtils& CustomizationUtils) override;
-	/** END IPropertyTypeCustomization interface */
-
+	/** End IPropertyTypeCustomization interface */
 };
diff --git a/Engine/Plugins/Experimental/DynamicMaterial/Source/DynamicMaterialEditor/Private/DetailsPanel/Slate/SDMDetailsPanelMaterialInterfaceWidget.cpp b/Engine/Plugins/Experimental/DynamicMaterial/Source/DynamicMaterialEditor/Private/DetailsPanel/Slate/SDMDetailsPanelMaterialInterfaceWidget.cpp
--- a/Engine/Plugins/Experimental/DynamicMaterial/Source/DynamicMaterialEditor/Private/DetailsPanel/Slate/SDMDetailsPanelMaterialInterfaceWidget.cpp
+++ b/Engine/Plugins/Experimental/DynamicMaterial/Source/DynamicMaterialEditor/Private/DetailsPanel/Slate/SDMDetailsPanelMaterialInterfaceWidget.cpp
@@ -1,9 +1,10 @@
 // Copyright Epic Games, Inc. All Rights Reserved.
 
 #include "DetailsPanel/Slate/SDMDetailsPanelMaterialInterfaceWidget.h"
-#include "AssetToolsModule.h"
+#include "Components/PrimitiveComponent.h"
+#include "DMObjectMaterialProperty.h"
 #include "DynamicMaterialEditorStyle.h"
-#include "IAssetTools.h"
+#include "IDynamicMaterialEditorModule.h"
 #include "Material/DynamicMaterialInstance.h"
 #include "Material/DynamicMaterialInstanceFactory.h"
 #include "PropertyCustomizationHelpers.h"
@@ -196,9 +197,39 @@ FReply SDMDetailsPanelMaterialInterfaceWidget::OpenDynamicMaterialInstanceTab()
 		return FReply::Unhandled();
 	}
 
-	IAssetTools& AssetTools = FModuleManager::LoadModuleChecked<FAssetToolsModule>("AssetTools").Get();
-	AssetTools.OpenEditorForAssets({Instance});
+	const IDynamicMaterialEditorModule& MaterialDesignerModule = IDynamicMaterialEditorModule::Get();
+	constexpr bool bInvokeTab = true;
 
+	TArray<UObject*> Outers;
+	PropertyHandle->GetOuterObjects(Outers);
+
+	if (Outers.Num() == 0)
+	{
+		MaterialDesignerModule.OpenMaterialModel(Instance->GetMaterialModel(), nullptr, bInvokeTab);
+		return FReply::Handled();
+	}
+
+	UWorld* OuterWorld = Outers[0]->GetWorld();
+	UPrimitiveComponent* PrimitiveComponent = Cast<UPrimitiveComponent>(Outers[0]);
+
+	if (!PrimitiveComponent)
+	{
+		MaterialDesignerModule.OpenMaterialModel(Instance->GetMaterialModel(), OuterWorld, bInvokeTab);
+		return FReply::Handled();
+	}
+
+	const int32 NumMaterials = PrimitiveComponent->GetNumMaterials();
+
+	for (int32 Index = 0; Index < NumMaterials; ++Index)
+	{
+		if (PrimitiveComponent->GetMaterial(Index) == Instance)
+		{
+			MaterialDesignerModule.OpenMaterialObjectProperty({PrimitiveComponent, Index}, OuterWorld, bInvokeTab);
+			return FReply::Handled();
+		}
+	}
+
+	MaterialDesignerModule.OpenMaterialModel(Instance->GetMaterialModel(), OuterWorld, bInvokeTab);
 	return FReply::Handled();
 }
 
diff --git a/Engine/Plugins/Experimental/DynamicMaterial/Source/DynamicMaterialEditor/Private/DetailsPanel/Slate/SDMMaterialListExtensionWidget.cpp b/Engine/Plugins/Experimental/DynamicMaterial/Source/DynamicMaterialEditor/Private/DetailsPanel/Slate/SDMMaterialListExtensionWidget.cpp
--- a/Engine/Plugins/Experimental/DynamicMaterial/Source/DynamicMaterialEditor/Private/DetailsPanel/Slate/SDMMaterialListExtensionWidget.cpp
+++ b/Engine/Plugins/Experimental/DynamicMaterial/Source/DynamicMaterialEditor/Private/DetailsPanel/Slate/SDMMaterialListExtensionWidget.cpp
@@ -1,13 +1,13 @@
 // Copyright Epic Games, Inc. All Rights Reserved.
 
 #include "SDMMaterialListExtensionWidget.h"
-#include "AssetToolsModule.h"
 #include "Components/PrimitiveComponent.h"
 #include "DMObjectMaterialProperty.h"
 #include "DMWorldSubsystem.h"
 #include "DynamicMaterialEditorStyle.h"
+#include "Editor.h"
 #include "Engine/World.h"
-#include "IAssetTools.h"
+#include "IDynamicMaterialEditorModule.h"
 #include "Material/DynamicMaterialInstance.h"
 #include "Material/DynamicMaterialInstanceFactory.h"
 #include "MaterialList.h"
@@ -193,16 +193,20 @@ FReply SDMMaterialListExtensionWidget::ClearDynamicMaterialInstance()
 
 FReply SDMMaterialListExtensionWidget::OpenDynamicMaterialInstanceTab()
 {
-	UDynamicMaterialInstance* Instance = GetDynamicMaterialInstance();
+	UPrimitiveComponent* CurrentComponent = CurrentComponentWeak.Get();
+	if (!ensure(CurrentComponent))
+	{
+		return FReply::Handled();
+	}
 
-	// We don't have an instance, so we can't open it
-	if (!Instance)
+	const TSharedPtr<FMaterialItemView> MaterialItemView = MaterialItemViewWeak.Pin();
+	if (!ensure(MaterialItemView.IsValid()))
 	{
-		return FReply::Unhandled();
+		return FReply::Handled();
 	}
 
-	IAssetTools& AssetTools = FModuleManager::LoadModuleChecked<FAssetToolsModule>("AssetTools").Get();
-	AssetTools.OpenEditorForAssets({Instance});
+	IDynamicMaterialEditorModule::Get().OpenMaterialObjectProperty({CurrentComponent, MaterialItemView->GetMaterialListItem().SlotIndex},
+		CurrentComponent->GetWorld(), /* Invoke Tab */ true);
 
 	return FReply::Handled();
 }
diff --git a/Engine/Plugins/Experimental/DynamicMaterial/Source/DynamicMaterialEditor/Private/DynamicMaterialEditorModule.cpp b/Engine/Plugins/Experimental/DynamicMaterial/Source/DynamicMaterialEditor/Private/DynamicMaterialEditorModule.cpp
--- a/Engine/Plugins/Experimental/DynamicMaterial/Source/DynamicMaterialEditor/Private/DynamicMaterialEditorModule.cpp
+++ b/Engine/Plugins/Experimental/DynamicMaterial/Source/DynamicMaterialEditor/Private/DynamicMaterialEditorModule.cpp
@@ -341,7 +341,7 @@ void FDynamicMaterialEditorModule::ShutdownModule()
 	FDMValueDetailsRowExtensions::Get().UnregisterRowExtensions();
 }
 
-void FDynamicMaterialEditorModule::SetDynamicMaterialModel(UDynamicMaterialModel* InMaterialModel, UWorld* InWorld, bool bInInvokeTab)
+void FDynamicMaterialEditorModule::OpenMaterialModel(UDynamicMaterialModel* InMaterialModel, UWorld* InWorld, bool bInInvokeTab) const
 {
 	using namespace UE::DynamicMaterialEditor::Private;
 
@@ -368,8 +368,8 @@ void FDynamicMaterialEditorModule::SetDynamicMaterialModel(UDynamicMaterialModel
 	}
 }
 
-void FDynamicMaterialEditorModule::SetDynamicMaterialObjectProperty(const FDMObjectMaterialProperty& InObjectProperty,
-	UWorld* InWorld, bool bInInvokeTab)
+void FDynamicMaterialEditorModule::OpenMaterialObjectProperty(const FDMObjectMaterialProperty& InObjectProperty,
+	UWorld* InWorld, bool bInInvokeTab) const
 {
 	using namespace UE::DynamicMaterialEditor::Private;
 
@@ -396,19 +396,19 @@ void FDynamicMaterialEditorModule::SetDynamicMaterialObjectProperty(const FDMObj
 	}
 }
 
-void FDynamicMaterialEditorModule::SetDynamicMaterialInstance(UDynamicMaterialInstance* InInstance, UWorld* InWorld, 
-	bool bInInvokeTab)
+void FDynamicMaterialEditorModule::OpenMaterialInstance(UDynamicMaterialInstance* InInstance, UWorld* InWorld, 
+	bool bInInvokeTab) const
 {
 	if (IsValid(InInstance))
 	{
 		if (UDynamicMaterialModel* InstanceModel = InInstance->GetMaterialModel())
 		{
-			SetDynamicMaterialModel(InstanceModel, InWorld, bInInvokeTab);
+			OpenMaterialModel(InstanceModel, InWorld, bInInvokeTab);
 		}
 	}
 }
 
-void FDynamicMaterialEditorModule::SetDynamicMaterialActor(AActor* InActor, UWorld* InWorld, bool bInInvokeTab)
+void FDynamicMaterialEditorModule::OnActorSelected(AActor* InActor, UWorld* InWorld, bool bInInvokeTab) const
 {
 	using namespace UE::DynamicMaterialEditor::Private;
 
@@ -419,7 +419,7 @@ void FDynamicMaterialEditorModule::SetDynamicMaterialActor(AActor* InActor, UWor
 			FDMLevelEditorIntegration::InvokeTabForWorld(InWorld);
 		}
 
-		Editor->SetMaterialActor(InActor);
+		Editor->OnActorSelected(InActor);
 	}
 	else if (IsValid(InWorld))
 	{
@@ -435,9 +435,9 @@ void FDynamicMaterialEditorModule::SetDynamicMaterialActor(AActor* InActor, UWor
 	}
 }
 
-void FDynamicMaterialEditorModule::ClearDynamicMaterialModel(UWorld* InWorld)
+void FDynamicMaterialEditorModule::ClearDynamicMaterialModel(UWorld* InWorld) const
 {
-	SetDynamicMaterialModel(nullptr, InWorld, /* Invoke tab */ false);
+	OpenMaterialModel(nullptr, InWorld, /* Invoke tab */ false);
 }
 
 TSharedRef<SWidget> FDynamicMaterialEditorModule::CreateEditor(UDynamicMaterialModel* InMaterialModel, UWorld* InAssetEditorWorld)
@@ -452,7 +452,7 @@ TSharedRef<SWidget> FDynamicMaterialEditorModule::CreateEditor(UDynamicMaterialM
 		{
 			WorldSubsystem->GetSetCustomEditorModelDelegate().BindSP(NewEditor, &SDMEditor::SetMaterialModel);
 			WorldSubsystem->GetCustomObjectPropertyEditorDelegate().BindSP(NewEditor, &SDMEditor::SetMaterialObjectProperty);
-			WorldSubsystem->GetSetCustomEditorActorDelegate().BindSP(NewEditor, &SDMEditor::SetMaterialActor);
+			WorldSubsystem->GetSetCustomEditorActorDelegate().BindSP(NewEditor, &SDMEditor::OnActorSelected);
 		}
 	}
 
@@ -539,7 +539,7 @@ void FDynamicMaterialEditorModule::RemoveBuildRequestForOuter(UObject* InOuter)
 	}
 }
 
-void FDynamicMaterialEditorModule::OpenEditor(UWorld* InWorld)
+void FDynamicMaterialEditorModule::OpenEditor(UWorld* InWorld) const
 {
 	if (!IsValid(InWorld))
 	{
diff --git a/Engine/Plugins/Experimental/DynamicMaterial/Source/DynamicMaterialEditor/Private/DynamicMaterialEditorModule.h b/Engine/Plugins/Experimental/DynamicMaterial/Source/DynamicMaterialEditor/Private/DynamicMaterialEditorModule.h
--- a/Engine/Plugins/Experimental/DynamicMaterial/Source/DynamicMaterialEditor/Private/DynamicMaterialEditorModule.h
+++ b/Engine/Plugins/Experimental/DynamicMaterial/Source/DynamicMaterialEditor/Private/DynamicMaterialEditorModule.h
@@ -88,10 +88,15 @@ public:
 	FDynamicMaterialEditorModule();
 
 	//~ Begin IDynamicMaterialEditorModule
-	virtual void OpenEditor(UWorld* InWorld) override;
 	virtual void RegisterCustomMaterialPropertyGenerator(UClass* InClass, FDMGetObjectMaterialPropertiesDelegate InGenerator) override;
 	virtual void RegisterMaterialModelCreatedCallback(const TSharedRef<IDMOnWizardCompleteCallback> InCallback)  override;
 	virtual void UnregisterMaterialModelCreatedCallback(const TSharedRef<IDMOnWizardCompleteCallback> InCallback) override;
+	virtual void OpenEditor(UWorld* InWorld) const override;
+	virtual void OpenMaterialModel(UDynamicMaterialModel* InMaterialModel, UWorld* InWorld, bool bInInvokeTab) const override;
+	virtual void OpenMaterialObjectProperty(const FDMObjectMaterialProperty& InObjectProperty, UWorld* InWorld, bool bInInvokeTab) const override;
+	virtual void OpenMaterialInstance(UDynamicMaterialInstance* InInstance, UWorld* InWorld, bool bInInvokeTab) const override;
+	virtual void OnActorSelected(AActor* InActor, UWorld* InWorld, bool bInInvokeTab) const override;
+	virtual void ClearDynamicMaterialModel(UWorld* InWorld) const override;
 	//~ End IDynamicMaterialEditorModule
 
 	//~ Begin IModuleInterface
@@ -99,12 +104,6 @@ public:
 	virtual void ShutdownModule() override;
 	//~ End IModuleInterface
 
-	void SetDynamicMaterialModel(UDynamicMaterialModel* InMaterialModel, UWorld* InWorld, bool bInInvokeTab);
-	void SetDynamicMaterialObjectProperty(const FDMObjectMaterialProperty& InObjectProperty, UWorld* InWorld, bool bInInvokeTab);
-	void SetDynamicMaterialInstance(UDynamicMaterialInstance* InInstance, UWorld* InWorld, bool bInInvokeTab);
-	void SetDynamicMaterialActor(AActor* InActor, UWorld* InWorld, bool bInInvokeTab);
-	void ClearDynamicMaterialModel(UWorld* InWorld);
-
 	//~ Begin FTickableEditorObject
 	virtual void Tick(float DeltaTime) override;
 	virtual ETickableTickType GetTickableTickType() const override { return ETickableTickType::Always; }
diff --git a/Engine/Plugins/Experimental/DynamicMaterial/Source/DynamicMaterialEditor/Private/LevelEditor/DMLevelEditorIntegration.cpp b/Engine/Plugins/Experimental/DynamicMaterial/Source/DynamicMaterialEditor/Private/LevelEditor/DMLevelEditorIntegration.cpp
--- a/Engine/Plugins/Experimental/DynamicMaterial/Source/DynamicMaterialEditor/Private/LevelEditor/DMLevelEditorIntegration.cpp
+++ b/Engine/Plugins/Experimental/DynamicMaterial/Source/DynamicMaterialEditor/Private/LevelEditor/DMLevelEditorIntegration.cpp
@@ -56,7 +56,7 @@ namespace UE::DynamicMaterialEditor::Private
 			return;
 		}
 
-		Editor->SetMaterialModel(nullptr);
+		Editor->ResetEditor();
 
 		const FString WorldPath = InWorld->GetPathName();
 		const int WorldPathLength = WorldPath.Len();
diff --git a/Engine/Plugins/Experimental/DynamicMaterial/Source/DynamicMaterialEditor/Private/LevelEditor/DMLevelEditorIntegrationInstance.cpp b/Engine/Plugins/Experimental/DynamicMaterial/Source/DynamicMaterialEditor/Private/LevelEditor/DMLevelEditorIntegrationInstance.cpp
--- a/Engine/Plugins/Experimental/DynamicMaterial/Source/DynamicMaterialEditor/Private/LevelEditor/DMLevelEditorIntegrationInstance.cpp
+++ b/Engine/Plugins/Experimental/DynamicMaterial/Source/DynamicMaterialEditor/Private/LevelEditor/DMLevelEditorIntegrationInstance.cpp
@@ -287,36 +287,7 @@ void FDMLevelEditorIntegrationInstance::OnActorSelected(AActor* InActor)
 		return;
 	}
 
-	UDynamicMaterialInstance* Instance = nullptr;
-
-	InActor->ForEachComponent<UPrimitiveComponent>(false, [&Instance](const UPrimitiveComponent* InPrimComp)
-		{
-			// Can't break this, so just skip every component
-			if (Instance)
-			{
-				return;
-			}
-
-			for (int32 MaterialIdx = 0; MaterialIdx < InPrimComp->GetNumMaterials(); ++MaterialIdx)
-			{
-				if (UDynamicMaterialInstance* MDI = Cast<UDynamicMaterialInstance>(InPrimComp->GetMaterial(MaterialIdx)))
-				{
-					Instance = MDI;
-					return;
-				}
-			}
-		});
-
-	UDynamicMaterialModel* MaterialModel = Instance ? Instance->GetMaterialModel() : nullptr;
-
-	if (!MaterialModel)
-	{
-		Editor->SetMaterialActor(InActor);
-	}
-	else
-	{
-		OnMaterialModelSelected(MaterialModel);
-	}
+	Editor->OnActorSelected(InActor);
 }
 
 void FDMLevelEditorIntegrationInstance::OnObjectSelectionChanged(const UTypedElementSelectionSet* InSelectionSet)
@@ -336,7 +307,7 @@ void FDMLevelEditorIntegrationInstance::OnObjectSelectionChanged(const UTypedEle
 
 	if (NewSelectedMaterialModel)
 	{
-		OnMaterialModelSelected(NewSelectedMaterialModel);
+		Editor->OnMaterialModelSelected(NewSelectedMaterialModel);
 		return;
 	}
 
@@ -355,16 +326,7 @@ void FDMLevelEditorIntegrationInstance::OnObjectSelectionChanged(const UTypedEle
 
 	if (NewSelectedMaterialInstance)
 	{
-		OnMaterialModelSelected(NewSelectedMaterialInstance->GetMaterialModel());
-		return;
-	}
-}
-
-void FDMLevelEditorIntegrationInstance::OnMaterialModelSelected(UDynamicMaterialModel* InMaterialModel)
-{
-	if (InMaterialModel)
-	{
-		Editor->SetMaterialModel(InMaterialModel);
+		Editor->OnMaterialModelSelected(NewSelectedMaterialInstance->GetMaterialModel());
 	}
 }
 
diff --git a/Engine/Plugins/Experimental/DynamicMaterial/Source/DynamicMaterialEditor/Private/LevelEditor/DMLevelEditorIntegrationInstance.h b/Engine/Plugins/Experimental/DynamicMaterial/Source/DynamicMaterialEditor/Private/LevelEditor/DMLevelEditorIntegrationInstance.h
--- a/Engine/Plugins/Experimental/DynamicMaterial/Source/DynamicMaterialEditor/Private/LevelEditor/DMLevelEditorIntegrationInstance.h
+++ b/Engine/Plugins/Experimental/DynamicMaterial/Source/DynamicMaterialEditor/Private/LevelEditor/DMLevelEditorIntegrationInstance.h
@@ -63,6 +63,4 @@ private:
 	void OnActorSelected(AActor* InActor);
 
 	void OnObjectSelectionChanged(const UTypedElementSelectionSet* InSelectionSet);
-
-	void OnMaterialModelSelected(UDynamicMaterialModel* InMaterialModel);
 };
diff --git a/Engine/Plugins/Experimental/DynamicMaterial/Source/DynamicMaterialEditor/Private/Model/AssetDefinition_DynamicMaterialModel.cpp b/Engine/Plugins/Experimental/DynamicMaterial/Source/DynamicMaterialEditor/Private/Model/AssetDefinition_DynamicMaterialModel.cpp
--- a/Engine/Plugins/Experimental/DynamicMaterial/Source/DynamicMaterialEditor/Private/Model/AssetDefinition_DynamicMaterialModel.cpp
+++ b/Engine/Plugins/Experimental/DynamicMaterial/Source/DynamicMaterialEditor/Private/Model/AssetDefinition_DynamicMaterialModel.cpp
@@ -52,7 +52,7 @@ EAssetCommandResult UAssetDefinition_DynamicMaterialModel::OpenAssets(const FAss
 			World = InOpenArgs.ToolkitHost->GetWorld();
 		}
 
-		FDynamicMaterialEditorModule::Get().SetDynamicMaterialModel(
+		FDynamicMaterialEditorModule::Get().OpenMaterialModel(
 			MaterialModel,
 			World,
 			/* Invoke Tab */ true
diff --git a/Engine/Plugins/Experimental/DynamicMaterial/Source/DynamicMaterialEditor/Private/Slate/Properties/SDMTextureUVVisualizer.cpp b/Engine/Plugins/Experimental/DynamicMaterial/Source/DynamicMaterialEditor/Private/Slate/Properties/SDMTextureUVVisualizer.cpp
--- a/Engine/Plugins/Experimental/DynamicMaterial/Source/DynamicMaterialEditor/Private/Slate/Properties/SDMTextureUVVisualizer.cpp
+++ b/Engine/Plugins/Experimental/DynamicMaterial/Source/DynamicMaterialEditor/Private/Slate/Properties/SDMTextureUVVisualizer.cpp
@@ -6,7 +6,6 @@
 #include "Components/DMMaterialStage.h"
 #include "Components/DMTextureUV.h"
 #include "DynamicMaterialEditorModule.h"
-#include "DynamicMaterialEditorSettings.h"
 #include "Editor.h"
 #include "Framework/Application/SlateApplication.h"
 #include "Materials/MaterialInterface.h"
diff --git a/Engine/Plugins/Experimental/DynamicMaterial/Source/DynamicMaterialEditor/Private/Slate/SDMEditor.cpp b/Engine/Plugins/Experimental/DynamicMaterial/Source/DynamicMaterialEditor/Private/Slate/SDMEditor.cpp
--- a/Engine/Plugins/Experimental/DynamicMaterial/Source/DynamicMaterialEditor/Private/Slate/SDMEditor.cpp
+++ b/Engine/Plugins/Experimental/DynamicMaterial/Source/DynamicMaterialEditor/Private/Slate/SDMEditor.cpp
@@ -33,6 +33,7 @@
 #include "Slate/SDMSlot.h"
 #include "Slate/SDMToolBar.h"
 #include "SlateOptMacros.h"
+#include "Material/DynamicMaterialInstance.h"
 #include "Styling/StyleColors.h"
 #include "ThumbnailRendering/ThumbnailManager.h"
 #include "Utils/DMBlueprintFunctionLibrary.h"
@@ -264,6 +265,8 @@ void SDMEditor::Construct(const FArguments& InArgs, TWeakObjectPtr<UDynamicMater
 		]
 	];
 
+	SetEmptyLayout();
+
 	SetMaterialModel(InModelWeak.Get());
 
 	UDynamicMaterialEditorSettings::Get()->OnSettingsChanged.AddSP(this, &SDMEditor::OnSettingsChanged);
@@ -305,7 +308,22 @@ SDMEditor::~SDMEditor()
 
 void SDMEditor::ClearEditor()
 {
-	SetMaterialModel(nullptr);
+	MaterialModelWeak.Reset();
+	ObjectProperty.Reset();
+
+	SlotPickerContainer.Reset();
+	SlotContainer.Reset();
+	ActiveSlotWidget.Reset();
+	ComponentEditContainer.Reset();
+	SplitterContainer.Reset();
+
+	bInvalidateComponentEditWidget = false;
+}
+
+void SDMEditor::ResetEditor()
+{
+	ClearEditor();
+	SetEmptyLayout();
 }
 
 void SDMEditor::SetActiveSlotIndex(int InSlotIndex)
@@ -420,22 +438,27 @@ FDMPropertyHandle SDMEditor::CreatePropertyHandle(const void* InOwningWidget, UO
 
 void SDMEditor::SetMaterialModel(UDynamicMaterialModel* InMaterialModel)
 {
-	MaterialModelWeak = InMaterialModel;
+	if (InMaterialModel && !UE::DynamicMaterialEditor::Private::CheckMaterialModelValidity(InMaterialModel))
+	{
+		InMaterialModel = nullptr;
+	}
 
-	SlotPickerContainer.Reset();
-	SlotContainer.Reset();
-	ActiveSlotWidget.Reset();
-	ComponentEditContainer.Reset();
-	SplitterContainer.Reset();
-	bInvalidateComponentEditWidget = false;
+	if (MaterialModelWeak.Get() == InMaterialModel)
+	{
+		return;
+	}
+
+	ClearEditor();
+
+	MaterialModelWeak = InMaterialModel;
 
 	Toolbar->SetMaterialModel(InMaterialModel);
+
 	SetEditedComponent(nullptr);
 
 	if (!InMaterialModel)
 	{
-		Container->SetContent(SDMEditor::GetEmptyContent());
-		bHasActiveLayout = false;
+		SetEmptyLayout();
 		return;
 	}
 
@@ -462,24 +485,110 @@ void SDMEditor::SetMaterialModel(UDynamicMaterialModel* InMaterialModel)
 
 void SDMEditor::SetMaterialObjectProperty(const FDMObjectMaterialProperty& InObjectProperty)
 {
+	if (!InObjectProperty.IsValid())
+	{
+		ResetEditor();
+		return;
+	}
+
 	UDynamicMaterialModel* MaterialModel = InObjectProperty.GetMaterialModel();
 
-	if (IsValid(MaterialModel))
+	if (!IsValid(MaterialModel))
 	{
-		ObjectProperty = InObjectProperty;
-		SetMaterialModel(MaterialModel);
+		ResetEditor();
+		return;
 	}
-	else
+
+	SetMaterialModel(MaterialModel);
+	ObjectProperty = InObjectProperty;
+
+	if (AActor* Actor = ObjectProperty.GetTypedOuter<AActor>())
+	{
+		SetMaterialActor(Actor);
+	}
+}
+
+AActor* SDMEditor::GetMaterialActor() const
+{
+	if (Toolbar.IsValid())
 	{
-		ObjectProperty.Reset();
+		return Toolbar->GetMaterialActor();
 	}
+
+	return nullptr;
 }
 
 void SDMEditor::SetMaterialActor(AActor* InActor)
 {
-	if (UDynamicMaterialEditorSettings* Settings = UDynamicMaterialEditorSettings::Get())
+	if (!IsValid(InActor))
+	{
+		return;
+	}
+
+	if (!GetMaterialModel())
+	{
+		Container->SetContent(
+			SNew(SBox)
+			.HAlign(HAlign_Center)
+			.Padding(5.0f, 5.0f, 5.0f, 5.0f)
+			[
+				CreateActorMaterialSlotSelector(InActor)
+			]
+		);
+	}
+
+
+	if (Toolbar->GetMaterialActor() != InActor)
+	{
+		Toolbar->SetMaterialActor(InActor);
+	}
+}
+
+void SDMEditor::SetEmptyLayout()
+{
+	bHasActiveLayout = false;
+	Container->SetContent(SDMEditor::GetEmptyContent());
+}
+
+void SDMEditor::OnMaterialModelSelected(UDynamicMaterialModel* InMaterialModel)
+{
+	if (!IsValid(InMaterialModel))
+	{
+		return;
+	}
+
+	// Only check this setting if we have an active material model.
+	if (GetMaterialModel())
+	{
+		UDynamicMaterialEditorSettings* Settings = UDynamicMaterialEditorSettings::Get();
+
+		if (!Settings || !Settings->bFollowSelection)
+		{
+			return;
+		}
+	}
+
+	SetMaterialModel(InMaterialModel);
+}
+
+void SDMEditor::OnMaterialInstanceSelected(UDynamicMaterialInstance* InMaterialInstance)
+{
+	if (!IsValid(InMaterialInstance))
 	{
-		if (!Settings->bFollowSelection)
+		return;
+	}
+
+	OnMaterialModelSelected(InMaterialInstance->GetMaterialModel());
+}
+
+void SDMEditor::OnActorSelected(AActor* InActor)
+{
+	// Only check this setting if we have an active material model.
+	if (GetMaterialModel())
+	{
+		UDynamicMaterialEditorSettings* Settings = UDynamicMaterialEditorSettings::Get();
+
+		if (!Settings || !Settings->bFollowSelection)
 		{
 			return;
 		}
@@ -490,22 +599,36 @@ void SDMEditor::SetMaterialActor(AActor* InActor)
 		return;
 	}
 
+	ClearEditor();
+
 	if (!IsValid(InActor))
 	{
-		Container->SetContent(SDMEditor::GetEmptyContent());
+		SetEmptyLayout();
 		return;
 	}
 
-	Toolbar->SetMaterialActor(InActor);
+	TArray<FDMObjectMaterialProperty> ActorProperties = UDMBlueprintFunctionLibrary::GetActorMaterialProperties(InActor);
 
-	Container->SetContent(
-		SNew(SBox)
-		.HAlign(HAlign_Center)
-		.Padding(5.0f, 5.0f, 5.0f, 5.0f)
-		[
-			CreateActorMaterialSlotSelector(InActor)
-		]
-	);
+	for (int32 MaterialPropertyIdx = 0; MaterialPropertyIdx < ActorProperties.Num(); ++MaterialPropertyIdx)
+	{
+		const FDMObjectMaterialProperty& MaterialProperty = ActorProperties[MaterialPropertyIdx];
+		UDynamicMaterialModel* Model = MaterialProperty.GetMaterialModel();
+
+		if (UDynamicMaterialModel* MaterialModel = Model)
+		{
+			SetMaterialObjectProperty(MaterialProperty);
+		}
+	}
+
+	if (!ActorProperties.IsEmpty() && !GetMaterialModel())
+	{
+		SetMaterialObjectProperty(ActorProperties[0]);
+	}
+
+	if (GetMaterialActor() != InActor)
+	{
+		SetMaterialActor(InActor);
+	}
 }
 
 TSharedPtr<SDMSlot> SDMEditor::GetActiveSlotWidget() const
@@ -718,7 +841,7 @@ TSharedRef<SWidget> SDMEditor::CreateMainLayout()
 		.HAlign(HAlign_Fill)
 		.VAlign(VAlign_Top)
 		[
-			SAssignNew(Toolbar, SDMToolBar)
+			SAssignNew(Toolbar, SDMToolBar, SharedThis(this))
 			.MaterialModel(MaterialModelWeak.Get())
 			.OnSlotChanged(this, &SDMEditor::OnToolBarPropertyChanged)
 			.OnGetSettingsMenu(this, &SDMEditor::MakeToolBarSettingsMenu)
@@ -960,12 +1083,17 @@ void SDMEditor::Tick(const FGeometry& AllottedGeometry, const double InCurrentTi
 	{
 		if (bHasActiveLayout)
 		{
-			ClearEditor();
+			ResetEditor();
 		}
 
 		return;
 	}
 
+	if (!bHasActiveLayout)
+	{
+		return;
+	}
+
 	if (ObjectProperty.IsValid())
 	{
 		UDynamicMaterialModel* MaterialModelProperty = ObjectProperty.GetMaterialModel();
@@ -985,32 +1113,44 @@ void SDMEditor::Tick(const FGeometry& AllottedGeometry, const double InCurrentTi
 			}
 			else
 			{
-				ClearEditor();
+				ResetEditor();
 				return;
 			}
 		}
 	}
-	else 
+	else if (!UE::DynamicMaterialEditor::Private::CheckMaterialModelValidity(MaterialModel))
 	{
-		if (!UE::DynamicMaterialEditor::Private::CheckMaterialModelValidity(MaterialModel))
-		{
-			ClearEditor();
-			return;
-		}
+		ResetEditor();
+		return;
 	}
 
-	if (Toolbar.IsValid() && Toolbar->GetMaterialModel() != MaterialModel)
-	{
-		Toolbar->SetMaterialModel(MaterialModel);
-	}
-	else if (bHasActiveLayout && ActiveSlotWidget.IsValid())
+	if (Toolbar.IsValid())
 	{
-		if (!ActiveSlotWidget->CheckValidity())
+		if (Toolbar->GetMaterialModel() != MaterialModel)
+		{
+			Toolbar->SetMaterialModel(MaterialModel);
+		}
+
+		if (ObjectProperty.IsValid())
+		{
+			AActor* Actor = ObjectProperty.GetTypedOuter<AActor>();
+
+			if (Toolbar->GetMaterialActor() != Actor)
+			{
+				Toolbar->SetMaterialActor(Actor);
+			}
+		}
+		else if (Toolbar->GetMaterialActor())
 		{
-			RefreshSlotWidget();
+			Toolbar->SetMaterialActor(nullptr);
 		}
 	}
 
+	if (ActiveSlotWidget.IsValid() && !ActiveSlotWidget->CheckValidity())
+	{
+		RefreshSlotWidget();
+	}
+
 	if (bInvalidateComponentEditWidget)
 	{
 		RefreshComponentEditWidget();
@@ -1175,6 +1315,14 @@ FReply SDMEditor::OnCreateMaterialButtonClicked(TWeakPtr<FDMObjectMaterialProper
 			{
 				Toolbar->SetMaterialModel(NewModel);
 			}
+
+			if (AActor* Actor = MaterialProperty->GetTypedOuter<AActor>())
+			{
+				if (Toolbar->GetMaterialActor() != Actor)
+				{
+					Toolbar->SetMaterialActor(Actor);
+				}
+			}
 		}
 	}
 
@@ -1196,6 +1344,7 @@ TSharedRef<SWidget> SDMEditor::MakeToolBarSettingsMenu()
 
 void SDMEditor::OnSettingsChanged(const FPropertyChangedEvent& InPropertyChangedEvent)
 {
+	ClearEditor();
 	SetMaterialModel(MaterialModelWeak.Get());
 }
 
diff --git a/Engine/Plugins/Experimental/DynamicMaterial/Source/DynamicMaterialEditor/Private/Slate/SDMEditor.h b/Engine/Plugins/Experimental/DynamicMaterial/Source/DynamicMaterialEditor/Private/Slate/SDMEditor.h
--- a/Engine/Plugins/Experimental/DynamicMaterial/Source/DynamicMaterialEditor/Private/Slate/SDMEditor.h
+++ b/Engine/Plugins/Experimental/DynamicMaterial/Source/DynamicMaterialEditor/Private/Slate/SDMEditor.h
@@ -51,6 +51,7 @@ public:
 	static void ClearPropertyHandles(const SWidget* InOwningWidget);
 
 	void Construct(const FArguments& InArgs, TWeakObjectPtr<UDynamicMaterialModel> InModelWeak);
+
 	virtual ~SDMEditor() override;
 
 	UDynamicMaterialModel* GetMaterialModel() const { return MaterialModelWeak.Get(); }
@@ -59,8 +60,15 @@ public:
 	const FDMObjectMaterialProperty& GetMaterialObjectProperty() const { return ObjectProperty; }
 	void SetMaterialObjectProperty(const FDMObjectMaterialProperty& InObjectProperty);
 
+	AActor* GetMaterialActor() const;
 	void SetMaterialActor(AActor* InActor);
 
+	void OnMaterialModelSelected(UDynamicMaterialModel* InMaterialModel);
+
+	void OnMaterialInstanceSelected(UDynamicMaterialInstance* InMaterialInstance);
+
+	void OnActorSelected(AActor* InActor);
+
 	TSharedPtr<SDMSlot> GetActiveSlotWidget() const;
 
 	void RefreshParametersList();
@@ -81,6 +89,8 @@ public:
 
 	void ClearEditor();
 
+	void ResetEditor();
+
 	int32 GetActiveSlotIndex() const { return ActiveSlotIndex; }
 	void SetActiveSlotIndex(int InSlotIndex);
 
@@ -175,4 +185,6 @@ protected:
 	void OnUndo();
 
 	UDMMaterialSlot* GetSlotForMaterialProperty(EDMMaterialPropertyType InProperty) const;
+
+	void SetEmptyLayout();
 };
diff --git a/Engine/Plugins/Experimental/DynamicMaterial/Source/DynamicMaterialEditor/Private/Slate/SDMMaterialWizard.cpp b/Engine/Plugins/Experimental/DynamicMaterial/Source/DynamicMaterialEditor/Private/Slate/SDMMaterialWizard.cpp
--- a/Engine/Plugins/Experimental/DynamicMaterial/Source/DynamicMaterialEditor/Private/Slate/SDMMaterialWizard.cpp
+++ b/Engine/Plugins/Experimental/DynamicMaterial/Source/DynamicMaterialEditor/Private/Slate/SDMMaterialWizard.cpp
@@ -253,9 +253,21 @@ FReply SDMMaterialWizard::Accept_OnClick()
 		{
 			if (UDynamicMaterialModelEditorOnlyData* EditorOnlyData = UDynamicMaterialModelEditorOnlyData::Get(MaterialModel))
 			{
+				FDMObjectMaterialProperty MaterialProperty = Editor->GetMaterialObjectProperty();
+				Editor->ClearEditor();
+
 				EditorOnlyData->SetChannelListPreset(CurrentPreset);
 				EditorOnlyData->OnWizardComplete();
-				Editor->SetMaterialModel(MaterialModel); // Refresh display
+
+				// Refresh display
+				if (MaterialProperty.IsValid())
+				{
+					Editor->SetMaterialObjectProperty(MaterialProperty);
+				}
+				else
+				{
+					Editor->SetMaterialModel(MaterialModel);
+				}
 			}
 		}
 	}
diff --git a/Engine/Plugins/Experimental/DynamicMaterial/Source/DynamicMaterialEditor/Private/Slate/SDMToolBar.cpp b/Engine/Plugins/Experimental/DynamicMaterial/Source/DynamicMaterialEditor/Private/Slate/SDMToolBar.cpp
--- a/Engine/Plugins/Experimental/DynamicMaterial/Source/DynamicMaterialEditor/Private/Slate/SDMToolBar.cpp
+++ b/Engine/Plugins/Experimental/DynamicMaterial/Source/DynamicMaterialEditor/Private/Slate/SDMToolBar.cpp
@@ -2,12 +2,16 @@
 
 #include "Slate/SDMToolBar.h"
 #include "ContentBrowserModule.h"
-#include "DynamicMaterialEditorModule.h"
+#include "DMWorldSubsystem.h"
 #include "DynamicMaterialEditorSettings.h"
 #include "DynamicMaterialEditorStyle.h"
+#include "Editor.h"
+#include "Engine/World.h"
 #include "EngineAnalytics.h"
 #include "GameFramework/Actor.h"
 #include "IContentBrowserSingleton.h"
+#include "SDMEditor.h"
+#include "Selection.h"
 #include "Material/DynamicMaterialInstance.h"
 #include "Model/DynamicMaterialModel.h"
 #include "SlateOptMacros.h"
@@ -26,24 +30,16 @@
 
 #define LOCTEXT_NAMESPACE "SDMToolBar"
 
-void SDMToolBar::Construct(const FArguments& InArgs)
+void SDMToolBar::Construct(const FArguments& InArgs, const TSharedRef<SDMEditor>& InEditor)
 {
 	SetCanTick(true);
 
+	EditorWeak = InEditor;
 	MaterialActorWeak = InArgs._MaterialActor;
 	MaterialModelWeak = InArgs._MaterialModel;
 	OnSlotChanged = InArgs._OnSlotChanged;
 	OnGetSettingsMenu = InArgs._OnGetSettingsMenu;
 	
-	if (MaterialActorWeak.IsValid())
-	{
-		SetMaterialActor(MaterialActorWeak.Get());
-	}
-	else
-	{
-		SetMaterialModel(MaterialModelWeak.Get());
-	}
-	
 	ChildSlot
 	.HAlign(HAlign_Fill)
 	.VAlign(VAlign_Center)
@@ -57,6 +53,9 @@ void SDMToolBar::Construct(const FArguments& InArgs)
 			CreateToolBarEntries()
 		]
 	];
+
+	SetMaterialModel(MaterialModelWeak.Get());
+	SetMaterialActor(MaterialActorWeak.Get());
 }
 
 TSharedRef<SWidget> SDMToolBar::CreateToolBarEntries()
@@ -127,6 +126,44 @@ TSharedRef<SWidget> SDMToolBar::CreateToolBarEntries()
 			]
 		]
 		
+		+ SHorizontalBox::Slot()
+		.AutoWidth()
+		.HAlign(HAlign_Right)
+		.VAlign(VAlign_Top)
+		.Padding(5.0f, 0.0f, 0.0f, 0.0f)
+		[
+			SAssignNew(BrowseButton, SButton)
+			.Visibility(EVisibility::Collapsed)
+			.ContentPadding(GetLargeIconToolBarButtonContentPadding())
+			.ButtonStyle(FDynamicMaterialEditorStyle::Get(), "HoverHintOnly")
+			.ToolTipText(LOCTEXT("MaterialDesignerBrowseTooltip", "Browse to the selected asset in the content browser."))
+			.OnClicked(this, &SDMToolBar::OnBrowseClicked)
+			[
+				SNew(SImage)
+				.Image(FAppStyle::GetBrush(TEXT("Icons.BrowseContent")))
+				.DesiredSizeOverride(GetLargeIconToolBarButtonSize())
+			]
+		]
+
+		+ SHorizontalBox::Slot()
+		.AutoWidth()
+		.HAlign(HAlign_Right)
+		.VAlign(VAlign_Top)
+		.Padding(5.0f, 0.0f, 0.0f, 0.0f)
+		[
+			SAssignNew(UseButton, SButton)
+			.Visibility(EVisibility::Collapsed)
+			.ContentPadding(GetLargeIconToolBarButtonContentPadding())
+			.ButtonStyle(FDynamicMaterialEditorStyle::Get(), "HoverHintOnly")
+			.ToolTipText(LOCTEXT("MaterialDesignerUseTooltip", "Replace the material in this slot with the one selected in the content browser."))
+			.OnClicked(this, &SDMToolBar::OnUseClicked)
+			[
+				SNew(SImage)
+				.Image(FAppStyle::GetBrush(TEXT("Icons.Use")))
+				.DesiredSizeOverride(GetLargeIconToolBarButtonSize())
+			]
+		]
+		
 		+ SHorizontalBox::Slot()
 		.AutoWidth()
 		.HAlign(HAlign_Right)
@@ -153,7 +190,7 @@ TSharedRef<SWidget> SDMToolBar::CreateToolBarEntries()
 		.Padding(5.0f, 0.0f, 0.0f, 0.0f)
 		[
 			SNew(SButton)
-			.ContentPadding(GetDefaultToolBarButtonContentPadding())
+			.ContentPadding(GetLargeIconToolBarButtonContentPadding())
 			.ButtonStyle(FDynamicMaterialEditorStyle::Get(), "HoverHintOnly")
 			.ToolTipText(LOCTEXT("ExportMaterialInstance", "Export Material Designer Instance"))
 			.Visibility(this, &SDMToolBar::GetExportMaterialInstanceButtonVisibility)
@@ -161,7 +198,7 @@ TSharedRef<SWidget> SDMToolBar::CreateToolBarEntries()
 			[
 				SNew(SImage)
 				.Image(FAppStyle::Get().GetBrush(TEXT("Icons.Toolbar.Export")))
-				.DesiredSizeOverride(GetDefaultToolBarButtonSize())
+				.DesiredSizeOverride(GetLargeIconToolBarButtonSize())
 			]
 		]
 		
@@ -260,6 +297,7 @@ void SDMToolBar::OnMaterialSlotChanged(TSharedPtr<FDMObjectMaterialProperty> InS
 	}
 
 	UDynamicMaterialModel* SelectedMaterialModel = InSelectedSlot->GetMaterialModel();
+
 	if (IsValid(SelectedMaterialModel))
 	{
 		SetMaterialModel(SelectedMaterialModel);
@@ -269,6 +307,13 @@ void SDMToolBar::OnMaterialSlotChanged(TSharedPtr<FDMObjectMaterialProperty> InS
 		UDMBlueprintFunctionLibrary::CreateDynamicMaterialInObject(*InSelectedSlot.Get());
 	}
 
+	AActor* SlotActor = InSelectedSlot->GetTypedOuter<AActor>();
+
+	if (GetMaterialActor() != SlotActor)
+	{
+		SetMaterialActor(SlotActor);
+	}
+
 	OnSlotChanged.ExecuteIfBound(InSelectedSlot);
 }
 
@@ -298,36 +343,27 @@ void SDMToolBar::SetMaterialModel(UDynamicMaterialModel* InModel)
 
 	MaterialModelWeak = InModel;
 
-	if (!MaterialActorWeak.IsValid())
-	{
-		MaterialActorWeak = MaterialModelWeak->GetTypedOuter<AActor>();
-	}
-
-	ActorMaterialProperties.Empty(0);
-	SelectedMaterialSlotIndex = INDEX_NONE;
-
-	if (!MaterialActorWeak.IsValid())
-	{
-		return;
-	}
+	bool bIsAsset = false;
 
-	const TArray<FDMObjectMaterialProperty> ActorDynamicMaterialSlots = UDMBlueprintFunctionLibrary::GetActorMaterialProperties(MaterialActorWeak.Get());
-	ActorMaterialProperties.Reserve(ActorDynamicMaterialSlots.Num());
-
-	for (const FDMObjectMaterialProperty& ActorDynamicMaterialSlot : ActorDynamicMaterialSlots)
+	if (InModel->IsAsset())
 	{
-		ActorMaterialProperties.Add(MakeShared<FDMObjectMaterialProperty>(ActorDynamicMaterialSlot));
+		bIsAsset = true;
 	}
-
-	for (int32 SlotIndex = 0; SlotIndex < ActorMaterialProperties.Num(); ++SlotIndex)
+	else if (UDynamicMaterialInstance* MaterialInstance = InModel->GetDynamicMaterialInstance())
 	{
-		FDMObjectMaterialProperty* ActorMaterialSlot = ActorMaterialProperties[SlotIndex].Get();
-		if (ActorMaterialSlot && ActorMaterialSlot->GetMaterialModel() == MaterialModelWeak)
+		if (MaterialInstance->IsAsset())
 		{
-			SelectedMaterialSlotIndex = SlotIndex;
-			break;
+			bIsAsset = true;
 		}
 	}
+
+	BrowseButton->SetVisibility(bIsAsset ? EVisibility::Visible : EVisibility::Collapsed);
+
+	// If this is a valid actor material slot, this will be updated by a SetMaterialActor call.
+	UseButton->SetVisibility(EVisibility::Collapsed);
+
+	ActorMaterialProperties.Empty(0);
+	SelectedMaterialSlotIndex = INDEX_NONE;
 }
 
 void SDMToolBar::SetMaterialActor(AActor* InActor, const int32 InActiveSlotIndex)
@@ -336,6 +372,13 @@ void SDMToolBar::SetMaterialActor(AActor* InActor, const int32 InActiveSlotIndex
 	ActorMaterialProperties.Empty();
 	SelectedMaterialSlotIndex = 0;
 
+	UseButton->SetVisibility(MaterialActorWeak.IsValid() ? EVisibility::Visible : EVisibility::Collapsed);
+
+	if (!IsValid(InActor))
+	{
+		return;
+	}
+
 	TArray<FDMObjectMaterialProperty> ActorProperties = UDMBlueprintFunctionLibrary::GetActorMaterialProperties(InActor);
 	UDynamicMaterialModel* MaterialModel = MaterialModelWeak.Get();
 
@@ -374,16 +417,15 @@ FSlateColor SDMToolBar::GetFollowSelectionColor() const
 	static FSlateColor EnabledColor = FSlateColor(EStyleColor::AccentGray);
 	static FSlateColor DisabledColor = FSlateColor(EStyleColor::Primary);
 
-	bool bFollowingSelection = true;
-
 	if (UDynamicMaterialEditorSettings* Settings = UDynamicMaterialEditorSettings::Get())
 	{
-		bFollowingSelection = Settings->bFollowSelection;
+		if (Settings->bFollowSelection)
+		{
+			return EnabledColor;
+		}
 	}
 
-	return bFollowingSelection
-		? EnabledColor
-		: DisabledColor;
+	return DisabledColor;
 }
 
 FReply SDMToolBar::OnFollowSelectionButtonClicked()
@@ -454,4 +496,129 @@ FReply SDMToolBar::OnExportMaterialInstanceButtonClicked()
 	return FReply::Handled();
 }
 
+FReply SDMToolBar::OnBrowseClicked()
+{
+	UDynamicMaterialModel* MaterialModel = MaterialModelWeak.Get();
+
+	if (!IsValid(MaterialModel))
+	{
+		return FReply::Handled();
+	}
+
+	UObject* Asset = nullptr;
+
+	if (MaterialModelWeak->IsAsset())
+	{
+		Asset = MaterialModel;
+	}
+	else if (UDynamicMaterialInstance* MaterialInstance = MaterialModel->GetDynamicMaterialInstance())
+	{
+		if (MaterialInstance->IsAsset())
+		{
+			Asset = MaterialInstance;
+		}
+	}
+
+	if (!Asset)
+	{
+		return FReply::Handled();
+	}
+
+	TArray<FAssetData> AssetDataList;
+	AssetDataList.Add(Asset);
+	GEditor->SyncBrowserToObjects(AssetDataList);
+
+	return FReply::Handled();
+}
+
+FReply SDMToolBar::OnUseClicked()
+{
+	if (!ActorMaterialProperties.IsValidIndex(SelectedMaterialSlotIndex))
+	{
+		return FReply::Handled();
+	}
+
+	UDynamicMaterialModel* CurrentModel = MaterialModelWeak.Get();
+	UDMWorldSubsystem* DMSubsystem = nullptr;
+
+	if (CurrentModel)
+	{
+		AActor* Actor = MaterialActorWeak.Get();
+
+		if (!IsValid(Actor))
+		{
+			return FReply::Handled();
+		}
+
+		UWorld* World = Actor->GetWorld();
+
+		if (!IsValid(World))
+		{
+			return FReply::Handled();
+		}
+
+		DMSubsystem = World->GetSubsystem<UDMWorldSubsystem>();
+
+		if (!DMSubsystem)
+		{
+			return FReply::Handled();
+		}
+	}
+
+	USelection* Selection = GEditor->GetSelectedObjects();
+
+	if (!Selection)
+	{
+		return FReply::Handled();
+	}
+
+	FEditorDelegates::LoadSelectedAssetsIfNeeded.Broadcast();
+
+	UDynamicMaterialInstance* SelectedInstance = nullptr;
+
+	TArray<UDynamicMaterialInstance*> SelectedInstances;
+	Selection->GetSelectedObjects(SelectedInstances);
+
+	for (UDynamicMaterialInstance* SelectedInstanceIter : SelectedInstances)
+	{
+		if (!IsValid(SelectedInstanceIter) || !SelectedInstanceIter->IsAsset())
+		{
+			continue;
+		}
+
+		SelectedInstance = SelectedInstanceIter;
+		break;
+	}
+
+	if (!SelectedInstance)
+	{
+		return FReply::Handled();
+	}
+
+	TSharedPtr<FDMObjectMaterialProperty> CurrentActorProperty = ActorMaterialProperties[SelectedMaterialSlotIndex];
+
+	if (DMSubsystem && DMSubsystem->GetMaterialValueSetterDelegate().IsBound())
+	{
+		if (IsValid(CurrentModel)
+			&& DMSubsystem->GetIsValidDelegate().IsBound()
+			&& !DMSubsystem->GetIsValidDelegate().Execute(CurrentModel))
+		{
+			return FReply::Handled();
+		}
+
+		DMSubsystem->GetMaterialValueSetterDelegate().Execute(*CurrentActorProperty, SelectedInstance);
+	}
+	else
+	{
+		ActorMaterialProperties[SelectedMaterialSlotIndex]->SetMaterial(SelectedInstance);
+	}
+
+	if (TSharedPtr<SDMEditor> Editor = EditorWeak.Pin())
+	{
+		Editor->SetMaterialObjectProperty(*CurrentActorProperty);
+	}
+
+	return FReply::Handled();
+}
+
 #undef LOCTEXT_NAMESPACE
diff --git a/Engine/Plugins/Experimental/DynamicMaterial/Source/DynamicMaterialEditor/Private/Slate/SDMToolBar.h b/Engine/Plugins/Experimental/DynamicMaterial/Source/DynamicMaterialEditor/Private/Slate/SDMToolBar.h
--- a/Engine/Plugins/Experimental/DynamicMaterial/Source/DynamicMaterialEditor/Private/Slate/SDMToolBar.h
+++ b/Engine/Plugins/Experimental/DynamicMaterial/Source/DynamicMaterialEditor/Private/Slate/SDMToolBar.h
@@ -3,9 +3,11 @@
 #pragma once
 
 #include "Framework/SlateDelegates.h"
+#include "Model/DynamicMaterialModel.h"
 #include "Widgets/SCompoundWidget.h"
 
 class AActor;
+class SDMEditor;
 class SWidget;
 class UDMMaterialStageExpression;
 class UDynamicMaterialModel;
@@ -31,12 +33,12 @@ public:
 
 	virtual ~SDMToolBar() {}
 
-	void Construct(const FArguments& InArgs);
+	void Construct(const FArguments& InArgs, const TSharedRef<SDMEditor>& InEditor);
 
 	const TArray<TSharedPtr<FDMObjectMaterialProperty>>& GetMaterialProperties() const { return ActorMaterialProperties; }
 	void SetMaterialProperties(const TArray<TSharedPtr<FDMObjectMaterialProperty>>& InActorMaterialProperties);
 
-	const AActor* GetMaterialActor() const { return MaterialActorWeak.Get(); }
+	AActor* GetMaterialActor() const { return MaterialActorWeak.Get(); }
 	void SetMaterialActor(AActor* InActor, const int32 InActiveSlotIndex = 0);
 
 	UDynamicMaterialModel* GetMaterialModel() const { return MaterialModelWeak.Get(); }
@@ -45,6 +47,7 @@ public:
 	FText GetSlotActorDisplayName() const;
 
 protected:
+	TWeakPtr<SDMEditor> EditorWeak;
 	TWeakObjectPtr<AActor> MaterialActorWeak;
 	TWeakObjectPtr<UDynamicMaterialModel> MaterialModelWeak;
 	FDMOnActorMaterailSlotChanged OnSlotChanged;
@@ -53,6 +56,9 @@ protected:
 	TArray<TSharedPtr<FDMObjectMaterialProperty>> ActorMaterialProperties;
 	int32 SelectedMaterialSlotIndex;
 
+	TSharedPtr<SWidget> BrowseButton;
+	TSharedPtr<SWidget> UseButton;
+
 	TSharedRef<SWidget> CreateToolBarEntries();
 
 	TSharedRef<SWidget> CreateToolBarButton(TAttribute<const FSlateBrush*> InImageBrush, const TAttribute<FText>& InTooltipText, FOnClicked InOnClicked);
@@ -66,6 +72,9 @@ protected:
 	const FMargin GetDefaultToolBarButtonContentPadding() const { return FMargin(2.0f); }
 	const FVector2D GetDefaultToolBarButtonSize() const { return FVector2D(20.0f); }
 
+	const FMargin GetLargeIconToolBarButtonContentPadding() const { return FMargin(4.0f); }
+	const FVector2D GetLargeIconToolBarButtonSize() const { return FVector2D(16.0f); }
+
 	EVisibility GetSlotsComboBoxWidgetVisibiltiy() const;
 
 	const FSlateBrush* GetFollowSelectionBrush() const;
@@ -74,4 +83,8 @@ protected:
 
 	EVisibility GetExportMaterialInstanceButtonVisibility() const;
 	FReply OnExportMaterialInstanceButtonClicked();
+
+	FReply OnBrowseClicked();
+
+	FReply OnUseClicked();
 };
diff --git a/Engine/Plugins/Experimental/DynamicMaterial/Source/DynamicMaterialEditor/Public/DMObjectMaterialProperty.h b/Engine/Plugins/Experimental/DynamicMaterial/Source/DynamicMaterialEditor/Public/DMObjectMaterialProperty.h
--- a/Engine/Plugins/Experimental/DynamicMaterial/Source/DynamicMaterialEditor/Public/DMObjectMaterialProperty.h
+++ b/Engine/Plugins/Experimental/DynamicMaterial/Source/DynamicMaterialEditor/Public/DMObjectMaterialProperty.h
@@ -54,4 +54,20 @@ struct DYNAMICMATERIALEDITOR_API FDMObjectMaterialProperty
 	FText GetPropertyName(bool bInIgnoreNewStatus) const;
 
 	void Reset();
+
+	template<typename InClass>
+	InClass* GetTypedOuter() const
+	{
+		if (UObject* Outer = OuterWeak.Get())
+		{
+			if (InClass* CastOuter = Cast<InClass>(Outer))
+			{
+				return CastOuter;
+			}
+
+			return Outer->GetTypedOuter<InClass>();
+		}
+
+		return nullptr;
+	}
 };
diff --git a/Engine/Plugins/Experimental/DynamicMaterial/Source/DynamicMaterialEditor/Public/IDynamicMaterialEditorModule.h b/Engine/Plugins/Experimental/DynamicMaterial/Source/DynamicMaterialEditor/Public/IDynamicMaterialEditorModule.h
--- a/Engine/Plugins/Experimental/DynamicMaterial/Source/DynamicMaterialEditor/Public/IDynamicMaterialEditorModule.h
+++ b/Engine/Plugins/Experimental/DynamicMaterial/Source/DynamicMaterialEditor/Public/IDynamicMaterialEditorModule.h
@@ -5,8 +5,11 @@
 #include "Containers/ContainersFwd.h"
 #include "Delegates/Delegate.h"
 #include "Modules/ModuleManager.h"
+#include "Templates/SharedPointerFwd.h"
 
+class AActor;
 class UClass;
+class UDynamicMaterialInstance;
 class UDynamicMaterialModel;
 class UObject;
 struct FDMObjectMaterialProperty;
@@ -31,8 +34,6 @@ public:
 		return FModuleManager::LoadModuleChecked<IDynamicMaterialEditorModule>(ModuleName);
 	}
 
-	virtual void OpenEditor(UWorld* InWorld) = 0;
-
 	virtual void RegisterCustomMaterialPropertyGenerator(UClass* InClass, FDMGetObjectMaterialPropertiesDelegate InGenerator) = 0;
 
 	virtual void RegisterMaterialModelCreatedCallback(const TSharedRef<IDMOnWizardCompleteCallback> InCallback) = 0;
@@ -47,4 +48,16 @@ public:
 		RegisterMaterialModelCreatedCallback(NewCallback);
 		return NewCallback;
 	}
+
+	virtual void OpenEditor(UWorld* InWorld) const = 0;
+
+	virtual void OpenMaterialModel(UDynamicMaterialModel* InMaterialModel, UWorld* InWorld, bool bInInvokeTab) const = 0;
+
+	virtual void OpenMaterialObjectProperty(const FDMObjectMaterialProperty& InObjectProperty, UWorld* InWorld, bool bInInvokeTab) const = 0;
+
+	virtual void OpenMaterialInstance(UDynamicMaterialInstance* InInstance, UWorld* InWorld, bool bInInvokeTab) const = 0;
+
+	virtual void OnActorSelected(AActor* InActor, UWorld* InWorld, bool bInInvokeTab) const = 0;
+
+	virtual void ClearDynamicMaterialModel(UWorld* InWorld) const = 0;
 };
