Material Designer
- Material 5.4 to 5.5 update.

diff --git a/Engine/Plugins/Experimental/DynamicMaterial/Source/DynamicMaterialEditor/Private/Model/DynamicMaterialModelEditorOnlyData.cpp b/Engine/Plugins/Experimental/DynamicMaterial/Source/DynamicMaterialEditor/Private/Model/DynamicMaterialModelEditorOnlyData.cpp
--- a/Engine/Plugins/Experimental/DynamicMaterial/Source/DynamicMaterialEditor/Private/Model/DynamicMaterialModelEditorOnlyData.cpp
+++ b/Engine/Plugins/Experimental/DynamicMaterial/Source/DynamicMaterialEditor/Private/Model/DynamicMaterialModelEditorOnlyData.cpp
@@ -1444,16 +1444,19 @@ void UDynamicMaterialModelEditorOnlyData::Serialize(FArchive& Ar)
 			{
 				ChannelListPreset = "Translucent";
 				BlendMode = BLEND_Translucent;
+				ShadingModel = EDMMaterialShadingModel::Unlit;
 			}
 			else if (GetSlotForMaterialProperty(EDMMaterialPropertyType::OpacityMask))
 			{
 				ChannelListPreset = "Translucent";
 				BlendMode = BLEND_Masked;
+				ShadingModel = EDMMaterialShadingModel::Unlit;
 			}
 			// Else let's stay on opaque
 			else
 			{
 				BlendMode = BLEND_Opaque;
+				ShadingModel = EDMMaterialShadingModel::DefaultLit;
 				ChannelListPreset = "Opaque";
 			}
 		}
