Material Designer
- Swapped Material Type and Material Settings order.
- Added all category-less items to a categroy called General.

diff --git a/Engine/Plugins/Experimental/DynamicMaterial/Source/DynamicMaterialEditor/Private/Slate/SDMComponentEdit.cpp b/Engine/Plugins/Experimental/DynamicMaterial/Source/DynamicMaterialEditor/Private/Slate/SDMComponentEdit.cpp
--- a/Engine/Plugins/Experimental/DynamicMaterial/Source/DynamicMaterialEditor/Private/Slate/SDMComponentEdit.cpp
+++ b/Engine/Plugins/Experimental/DynamicMaterial/Source/DynamicMaterialEditor/Private/Slate/SDMComponentEdit.cpp
@@ -296,6 +296,19 @@ TSharedRef<SWidget> SDMComponentEdit::CreateEditWidget()
 	TSharedRef<ICustomDetailsView> DetailsView = ICustomDetailsViewModule::Get().CreateCustomDetailsView(Args);
 	FCustomDetailsViewItemId RootId = DetailsView->GetRootItem()->GetItemId();
 
+	TSharedPtr<ICustomDetailsViewItem> DefaultCategoryItem;
+
+	auto GetDefaultCategory = [&DefaultCategoryItem, &DetailsView, &RootId]()
+		{
+			if (!DefaultCategoryItem.IsValid())
+			{
+				DefaultCategoryItem = DetailsView->CreateCustomCategoryItem("General", LOCTEXT("General", "General"))->AsItem();
+				DetailsView->ExtendTree(RootId, ECustomDetailsTreeInsertPosition::Child, DefaultCategoryItem.ToSharedRef());
+			}
+
+			return DefaultCategoryItem;
+		};
+
 	if (UDMMaterialStage* Stage = Cast<UDMMaterialStage>(ComponentWeak.Get()))
 	{
 		TSharedPtr<ICustomDetailsViewCustomItem> TypeSelectorItem = DetailsView->CreateCustomItem(
@@ -309,7 +322,7 @@ TSharedRef<SWidget> SDMComponentEdit::CreateEditWidget()
 			TypeSelectorItem->SetValueWidget(CreateSourceTypeEditWidget());
 
 			DetailsView->ExtendTree(
-				RootId,
+				GetDefaultCategory()->GetItemId(),
 				ECustomDetailsTreeInsertPosition::FirstChild,
 				TypeSelectorItem->AsItem()
 			);
@@ -364,7 +377,11 @@ TSharedRef<SWidget> SDMComponentEdit::CreateEditWidget()
 
 		TSharedPtr<ICustomDetailsViewItem> CategoryItem;
 
-		if (CategoryName != NAME_None)
+		if (CategoryName == NAME_None)
+		{
+			CategoryItem = GetDefaultCategory();
+		}
+		else
 		{
 			CategoryItem = DetailsView->FindCustomItem(CategoryName);
 
@@ -823,25 +840,6 @@ void SDMComponentEdit::OnUndo()
 void SDMComponentEdit::GenerateMaterialModelPropertyRows(const TSharedRef<SDMEditor> InEditorWidget, UDynamicMaterialModel* InMaterialModel,
 	TArray<FDMPropertyHandle>& InOutPropertyRows, TSet<UDMMaterialComponent*>& InOutProcessedObjects)
 {
-	const FName MaterialTypeCategory = FName("Material Type");
-	auto AddVariable = [InEditorWidget, &InOutPropertyRows, &MaterialTypeCategory](UObject* InObject, FName InPropertyName)
-		{
-			FDMPropertyHandle& ValueHandle = InOutPropertyRows.Add_GetRef(InEditorWidget->GetPropertyHandle(&*InEditorWidget,
-				InObject, InPropertyName));
-
-			ValueHandle.CategoryOverrideName = MaterialTypeCategory;
-		};
-
-	if (UDynamicMaterialModelEditorOnlyData* EditorOnlyData = UDynamicMaterialModelEditorOnlyData::Get(InMaterialModel))
-	{
-		AddVariable(EditorOnlyData, GET_MEMBER_NAME_CHECKED(UDynamicMaterialModelEditorOnlyData, ChannelListPreset));
-		AddVariable(EditorOnlyData, GET_MEMBER_NAME_CHECKED(UDynamicMaterialModelEditorOnlyData, Domain));
-		AddVariable(EditorOnlyData, GET_MEMBER_NAME_CHECKED(UDynamicMaterialModelEditorOnlyData, BlendMode));
-		AddVariable(EditorOnlyData, GET_MEMBER_NAME_CHECKED(UDynamicMaterialModelEditorOnlyData, ShadingModel));
-		AddVariable(EditorOnlyData, GET_MEMBER_NAME_CHECKED(UDynamicMaterialModelEditorOnlyData, bPixelAnimationFlag));
-		AddVariable(EditorOnlyData, GET_MEMBER_NAME_CHECKED(UDynamicMaterialModelEditorOnlyData, bTwoSidedFlag));
-	}
-
 	const FName MaterialSettingsCategory = FName("Material Settings");
 
 	auto AddGlobalValue = [InEditorWidget, &InOutPropertyRows, &InOutProcessedObjects, &MaterialSettingsCategory]
@@ -899,6 +897,25 @@ void SDMComponentEdit::GenerateMaterialModelPropertyRows(const TSharedRef<SDMEdi
 			/* Start from */ EDMMaterialPropertyType::Roughness
 		);
 	}
+
+	const FName MaterialTypeCategory = FName("Material Type");
+	auto AddVariable = [InEditorWidget, &InOutPropertyRows, &MaterialTypeCategory](UObject* InObject, FName InPropertyName)
+		{
+			FDMPropertyHandle& ValueHandle = InOutPropertyRows.Add_GetRef(InEditorWidget->GetPropertyHandle(&*InEditorWidget,
+				InObject, InPropertyName));
+
+			ValueHandle.CategoryOverrideName = MaterialTypeCategory;
+		};
+
+	if (UDynamicMaterialModelEditorOnlyData* EditorOnlyData = UDynamicMaterialModelEditorOnlyData::Get(InMaterialModel))
+	{
+		AddVariable(EditorOnlyData, GET_MEMBER_NAME_CHECKED(UDynamicMaterialModelEditorOnlyData, ChannelListPreset));
+		AddVariable(EditorOnlyData, GET_MEMBER_NAME_CHECKED(UDynamicMaterialModelEditorOnlyData, Domain));
+		AddVariable(EditorOnlyData, GET_MEMBER_NAME_CHECKED(UDynamicMaterialModelEditorOnlyData, BlendMode));
+		AddVariable(EditorOnlyData, GET_MEMBER_NAME_CHECKED(UDynamicMaterialModelEditorOnlyData, ShadingModel));
+		AddVariable(EditorOnlyData, GET_MEMBER_NAME_CHECKED(UDynamicMaterialModelEditorOnlyData, bPixelAnimationFlag));
+		AddVariable(EditorOnlyData, GET_MEMBER_NAME_CHECKED(UDynamicMaterialModelEditorOnlyData, bTwoSidedFlag));
+	}
 }
 
 TSharedRef<SWidget> SDMComponentEdit::CreateSourceTypeEditWidget()
