Material Designer
- The base stage now has a channel override selector.

diff --git a/Engine/Plugins/Experimental/DynamicMaterial/Source/DynamicMaterialEditor/Private/Components/DMMaterialStageBlend.cpp b/Engine/Plugins/Experimental/DynamicMaterial/Source/DynamicMaterialEditor/Private/Components/DMMaterialStageBlend.cpp
--- a/Engine/Plugins/Experimental/DynamicMaterial/Source/DynamicMaterialEditor/Private/Components/DMMaterialStageBlend.cpp
+++ b/Engine/Plugins/Experimental/DynamicMaterial/Source/DynamicMaterialEditor/Private/Components/DMMaterialStageBlend.cpp
@@ -7,9 +7,9 @@
 #include "Components/DMMaterialStageInput.h"
 #include "Components/MaterialStageExpressions/DMMSETextureSample.h"
 #include "Components/MaterialStageInputs/DMMSIExpression.h"
-#include "Components/MaterialStageInputs/DMMSITextureUV.h"
 #include "Components/MaterialStageInputs/DMMSIValue.h"
 #include "Components/MaterialValues/DMMaterialValueFloat1.h"
+#include "DMEDefs.h"
 #include "Materials/MaterialExpressionComponentMask.h"
 #include "Model/DMMaterialBuildState.h"
 #include "Model/DMMaterialBuildUtils.h"
@@ -32,6 +32,10 @@ UDMMaterialStageBlend::UDMMaterialStageBlend(const FText& InName)
 {
 	bInputRequired = true;
 	bAllowNestedInputs = true;
+
+	BaseChannelOverride = EAvaColorChannel::None;
+
+	EditableProperties.Add(GET_MEMBER_NAME_CHECKED(UDMMaterialStageBlend, BaseChannelOverride));
  
 	InputConnectors.Add({InputAlpha, LOCTEXT("Opacity", "Opacity"), EDMValueType::VT_Float1});
 	InputConnectors.Add({InputA, LOCTEXT("PreviousStage", "Previous Stage"), EDMValueType::VT_Float3_RGB});
@@ -39,7 +43,197 @@ UDMMaterialStageBlend::UDMMaterialStageBlend(const FText& InName)
  
 	OutputConnectors.Add({0, LOCTEXT("Blend", "Blend"), EDMValueType::VT_Float3_RGB});
 }
- 
+
+bool UDMMaterialStageBlend::CanUseBaseChannelOverride() const
+{
+	return GetDefaultBaseChannelOverrideOutputIndex() != INDEX_NONE;
+}
+
+EAvaColorChannel UDMMaterialStageBlend::GetBaseChannelOverride() const
+{
+	if (CanUseBaseChannelOverride())
+	{
+		PullBaseChannelOverride();
+		return BaseChannelOverride;
+	}
+
+	return EAvaColorChannel::None;
+}
+
+void UDMMaterialStageBlend::SetBaseChannelOverride(EAvaColorChannel InMaskChannel)
+{
+	if (!CanUseBaseChannelOverride())
+	{
+		return;
+	}
+
+	if (GetBaseChannelOverride() == InMaskChannel)
+	{
+		return;
+	}
+
+	BaseChannelOverride = InMaskChannel;
+	PushBaseChannelOverride();
+
+	Update(EDMUpdateType::Structure);
+}
+
+int32 UDMMaterialStageBlend::GetDefaultBaseChannelOverrideOutputIndex() const
+{
+	UDMMaterialStageInput* StageInputB = GetInputB();
+
+	if (!StageInputB)
+	{
+		return INDEX_NONE;
+	}
+
+	const TArray<FDMMaterialStageConnector>& MaskInputOutputConnectors = StageInputB->GetOutputConnectors();
+
+	for (int32 Index = 0; Index < MaskInputOutputConnectors.Num(); ++Index)
+	{
+		if (UDMValueDefinitionLibrary::GetValueDefinition(MaskInputOutputConnectors[Index].Type).GetFloatCount() > 1)
+		{
+			return Index;
+		}
+	}
+
+	return INDEX_NONE;
+}
+
+bool UDMMaterialStageBlend::IsValidBaseChannelOverrideOutputIndex(int32 InIndex) const
+{
+	UDMMaterialStageInput* StageInputB = GetInputB();
+
+	if (!StageInputB)
+	{
+		return false;
+	}
+
+	const TArray<FDMMaterialStageConnector>& MaskInputOutputConnectors = StageInputB->GetOutputConnectors();
+
+	if (!MaskInputOutputConnectors.IsValidIndex(InIndex))
+	{
+		return false;
+	}
+
+	return UDMValueDefinitionLibrary::GetValueDefinition(MaskInputOutputConnectors[InIndex].Type).GetFloatCount() > 1;
+}
+
+void UDMMaterialStageBlend::PullBaseChannelOverride() const
+{
+	BaseChannelOverride = EAvaColorChannel::None;
+
+	if (!CanUseBaseChannelOverride())
+	{
+		return;
+	}
+
+	UDMMaterialStage* Stage = GetStage();
+
+	if (!Stage)
+	{
+		return;
+	}
+
+	const TArray<FDMMaterialStageConnection>& InputMap = Stage->GetInputConnectionMap();
+
+	if (!InputMap.IsValidIndex(UDMMaterialStageBlend::InputB)
+		|| !InputMap[UDMMaterialStageBlend::InputB].Channels.IsValidIndex(0))
+	{
+		return;
+	}
+
+	TArray<UDMMaterialStageInput*> Inputs = Stage->GetInputs();
+	const FDMMaterialStageConnectorChannel& MaskConnectorChannel = InputMap[UDMMaterialStageBlend::InputB].Channels[0];
+
+	switch (MaskConnectorChannel.OutputChannel)
+	{
+		case FDMMaterialStageConnectorChannel::FIRST_CHANNEL:
+			BaseChannelOverride = EAvaColorChannel::Red;
+			break;
+
+		case FDMMaterialStageConnectorChannel::SECOND_CHANNEL:
+			BaseChannelOverride = EAvaColorChannel::Green;
+			break;
+
+		case FDMMaterialStageConnectorChannel::THIRD_CHANNEL:
+			BaseChannelOverride = EAvaColorChannel::Blue;
+			break;
+
+		case FDMMaterialStageConnectorChannel::FOURTH_CHANNEL:
+			BaseChannelOverride = EAvaColorChannel::Alpha;
+			break;
+
+		default:
+			// Do nothing
+			break;
+	}
+}
+
+void UDMMaterialStageBlend::PushBaseChannelOverride()
+{
+	UDMMaterialStage* Stage = GetStage();
+
+	if (!Stage)
+	{
+		return;
+	}
+
+	if (!CanUseBaseChannelOverride())
+	{
+		return;
+	}
+
+	const TArray<FDMMaterialStageConnection>& InputMap = Stage->GetInputConnectionMap();
+
+	if (!InputMap.IsValidIndex(UDMMaterialStageBlend::InputB)
+		|| !InputMap[UDMMaterialStageBlend::InputB].Channels.IsValidIndex(0))
+	{
+		return;
+	}
+
+	TArray<UDMMaterialStageInput*> Inputs = Stage->GetInputs();
+	const FDMMaterialStageConnectorChannel& MaskConnectorChannel = InputMap[UDMMaterialStageBlend::InputB].Channels[0];
+	const int32 MaskInputIdx = MaskConnectorChannel.SourceIndex - FDMMaterialStageConnectorChannel::FIRST_STAGE_INPUT;
+
+	if (!Inputs.IsValidIndex(MaskInputIdx))
+	{
+		return;
+	}
+
+	const int32 OutputIndex = IsValidBaseChannelOverrideOutputIndex(MaskConnectorChannel.OutputIndex)
+		? MaskConnectorChannel.OutputIndex
+		: GetDefaultBaseChannelOverrideOutputIndex();
+
+	int32 OutputChannel;
+
+	switch (BaseChannelOverride)
+	{
+		case EAvaColorChannel::Red:
+			OutputChannel = FDMMaterialStageConnectorChannel::FIRST_CHANNEL;
+			break;
+
+		case EAvaColorChannel::Green:
+			OutputChannel = FDMMaterialStageConnectorChannel::SECOND_CHANNEL;
+			break;
+
+		case EAvaColorChannel::Blue:
+			OutputChannel = FDMMaterialStageConnectorChannel::THIRD_CHANNEL;
+			break;
+
+		case EAvaColorChannel::Alpha:
+			OutputChannel = FDMMaterialStageConnectorChannel::FOURTH_CHANNEL;
+			break;
+
+		default:
+			OutputChannel = FDMMaterialStageConnectorChannel::WHOLE_CHANNEL;
+			break;
+	}
+
+	Stage->UpdateInputMap(UDMMaterialStageBlend::InputB, MaskConnectorChannel.SourceIndex, FDMMaterialStageConnectorChannel::WHOLE_CHANNEL,
+		OutputIndex, OutputChannel, MaskConnectorChannel.MaterialProperty);
+}
+
 UDMMaterialStage* UDMMaterialStageBlend::CreateStage(TSubclassOf<UDMMaterialStageBlend> InMaterialStageBlendClass, UDMMaterialLayerObject* InLayer)
 {
 	check(InMaterialStageBlendClass);
@@ -200,6 +394,27 @@ bool UDMMaterialStageBlend::UpdateStagePreviewMaterial(UDMMaterialStage* InStage
 	return true;
 }
 
+void UDMMaterialStageBlend::NotifyPostChange(const FPropertyChangedEvent& InPropertyChangedEvent, FEditPropertyChain* InPropertyThatChanged)
+{
+	if (!IsComponentValid())
+	{
+		return;
+	}
+
+	static const FName MaskChannelName = GET_MEMBER_NAME_CHECKED(UDMMaterialStageBlend, BaseChannelOverride);
+
+	if (InPropertyChangedEvent.GetPropertyName() == MaskChannelName)
+	{
+		PushBaseChannelOverride();
+	}
+	else
+	{
+		PullBaseChannelOverride();
+	}
+
+	Super::NotifyPostChange(InPropertyChangedEvent, InPropertyThatChanged);
+}
+
 void UDMMaterialStageBlend::AddDefaultInput(int32 InInputIndex) const
 {
 	if (!IsComponentValid())
@@ -543,5 +758,5 @@ UDMMaterialStageInput* UDMMaterialStageBlend::GetInputB() const
 
 	return nullptr;
 }
- 
+
 #undef LOCTEXT_NAMESPACE
diff --git a/Engine/Plugins/Experimental/DynamicMaterial/Source/DynamicMaterialEditor/Private/Components/DMMaterialStageThroughputLayerBlend.cpp b/Engine/Plugins/Experimental/DynamicMaterial/Source/DynamicMaterialEditor/Private/Components/DMMaterialStageThroughputLayerBlend.cpp
--- a/Engine/Plugins/Experimental/DynamicMaterial/Source/DynamicMaterialEditor/Private/Components/DMMaterialStageThroughputLayerBlend.cpp
+++ b/Engine/Plugins/Experimental/DynamicMaterial/Source/DynamicMaterialEditor/Private/Components/DMMaterialStageThroughputLayerBlend.cpp
@@ -11,7 +11,7 @@
 #include "Components/MaterialStageInputs/DMMSIExpression.h"
 #include "Components/MaterialStageInputs/DMMSIValue.h"
 #include "Components/MaterialValues/DMMaterialValueTexture.h"
-#include "DMDefs.h"
+#include "DMEDefs.h"
 #include "DynamicMaterialEditorSettings.h"
 #include "Materials/MaterialExpressionComponentMask.h"
 #include "Materials/MaterialExpressionConstant3Vector.h"
@@ -1290,8 +1290,9 @@ void UDMMaterialStageThroughputLayerBlend::PushMaskChannelOverride()
 			break;
 	}
 
-	Stage->UpdateInputMap(InputMaskSource, MaskConnectorChannel.SourceIndex, FDMMaterialStageConnectorChannel::WHOLE_CHANNEL,
-		OutputIndex, OutputChannel, MaskConnectorChannel.MaterialProperty);
+	Stage->UpdateInputMap(UDMMaterialStageThroughputLayerBlend::InputMaskSource, MaskConnectorChannel.SourceIndex,
+		FDMMaterialStageConnectorChannel::WHOLE_CHANNEL, OutputIndex, OutputChannel, 
+		MaskConnectorChannel.MaterialProperty);
 }
 
 #undef LOCTEXT_NAMESPACE
diff --git a/Engine/Plugins/Experimental/DynamicMaterial/Source/DynamicMaterialEditor/Public/Components/DMMaterialStageBlend.h b/Engine/Plugins/Experimental/DynamicMaterial/Source/DynamicMaterialEditor/Public/Components/DMMaterialStageBlend.h
--- a/Engine/Plugins/Experimental/DynamicMaterial/Source/DynamicMaterialEditor/Public/Components/DMMaterialStageBlend.h
+++ b/Engine/Plugins/Experimental/DynamicMaterial/Source/DynamicMaterialEditor/Public/Components/DMMaterialStageBlend.h
@@ -11,6 +11,7 @@ class UDMMaterialLayerObject;
 class UDMMaterialStageInput;
 class UDMMaterialValueFloat1;
 class UMaterial;
+enum class EAvaColorChannel : uint8;
 struct FDMMaterialBuildState;
 
 /**
@@ -33,6 +34,12 @@ public:
 	UDMMaterialValueFloat1* GetInputAlpha() const;
 	UDMMaterialStageInput* GetInputB() const;
 
+	UFUNCTION(BlueprintCallable, Category = "Material Designer")
+	EAvaColorChannel GetBaseChannelOverride() const;
+
+	UFUNCTION(BlueprintCallable, Category = "Material Designer")
+	void SetBaseChannelOverride(EAvaColorChannel InMaskChannel);
+
 	//~ Begin UDMMaterialStageThroughput
 	virtual bool CanInputAcceptType(int32 InputIndex, EDMValueType ValueType) const override;
 	virtual void AddDefaultInput(int32 InInputIndex) const override;
@@ -46,20 +53,44 @@ public:
 	//~ Begin UDMMaterialStageSource
 	virtual FText GetStageDescription() const override;
 	virtual bool SupportsLayerMaskTextureUVLink() const override { return true; }
-	virtual FDMExpressionInput GetLayerMaskLinkTextureUVInputExpressions(const TSharedRef<FDMMaterialBuildState>& InBuildState) const;
+	virtual FDMExpressionInput GetLayerMaskLinkTextureUVInputExpressions(const TSharedRef<FDMMaterialBuildState>& InBuildState) const override;
 	virtual void GetMaskAlphaBlendNode(const TSharedRef<FDMMaterialBuildState>& InBuildState, UMaterialExpression*& OutExpression,
 		int32& OutOutputIndex, int32& OutOutputChannel) const override;
 	virtual bool UpdateStagePreviewMaterial(UDMMaterialStage* InStage, UMaterial* InPreviewMaterial, UMaterialExpression*& OutMaterialExpression,
 		int32& OutputIndex) override;
 	//~ End UDMMaterialStageSource
 
+	//~ Begin FNotifyHook
+	virtual void NotifyPostChange(const FPropertyChangedEvent& InPropertyChangedEvent, class FEditPropertyChain* InPropertyThatChanged) override;
+	//~ End FNotifyHook
+
 protected:
 	static TArray<TStrongObjectPtr<UClass>> Blends;
 
+	UPROPERTY(EditInstanceOnly, BlueprintReadOnly, Getter = GetBaseChannelOverride, Setter = SetBaseChannelOverride, BlueprintGetter = GetBaseChannelOverride,
+		Category = "Material Designer", DisplayName = "Channel Mask",
+		meta = (NotKeyframeable, ToolTip = "Changes the output channel of the base input."))
+	mutable EAvaColorChannel BaseChannelOverride;
+
 	static void GenerateBlendList();
 
 	UDMMaterialStageBlend();
 	UDMMaterialStageBlend(const FText& InName);
 
+	/* Returns true if there are any outputs on the base input that have more than 1 channel. */
+	bool CanUseBaseChannelOverride() const;
+
+	/* Returns the first output on the base input that has more than 1 channel. */
+	int32 GetDefaultBaseChannelOverrideOutputIndex() const;
+
+	/* Returns true if the given base output supports more than 1 channel. */
+	bool IsValidBaseChannelOverrideOutputIndex(int32 InIndex) const;
+
+	/* Reads the current output setting from the input map. */
+	void PullBaseChannelOverride() const;
+
+	/* Takes the override setting and applies it to the input map. */
+	void PushBaseChannelOverride();
+
 	virtual void UpdatePreviewMaterial(UMaterial* InPreviewMaterial = nullptr) override;
 };
\ No newline at end of file
diff --git a/Engine/Plugins/Experimental/DynamicMaterial/Source/DynamicMaterialEditor/Public/Components/DMMaterialStageThroughputLayerBlend.h b/Engine/Plugins/Experimental/DynamicMaterial/Source/DynamicMaterialEditor/Public/Components/DMMaterialStageThroughputLayerBlend.h
--- a/Engine/Plugins/Experimental/DynamicMaterial/Source/DynamicMaterialEditor/Public/Components/DMMaterialStageThroughputLayerBlend.h
+++ b/Engine/Plugins/Experimental/DynamicMaterial/Source/DynamicMaterialEditor/Public/Components/DMMaterialStageThroughputLayerBlend.h
@@ -8,16 +8,7 @@
 class UDMMaterialLayerObject;
 class UMaterial;
 class UMaterialExpression;
-
-UENUM(BlueprintType)
-enum class EAvaColorChannel : uint8
-{
-	None =  0,
-	Red =   1,
-	Green = 2,
-	Blue =  3,
-	Alpha = 4
-};
+enum class EAvaColorChannel : uint8;
 
 UCLASS(BlueprintType, ClassGroup = "Material Designer")
 class DYNAMICMATERIALEDITOR_API UDMMaterialStageThroughputLayerBlend : public UDMMaterialStageThroughput
@@ -84,7 +75,7 @@ public:
 
 protected:
 	UPROPERTY(EditInstanceOnly, BlueprintReadOnly, Getter=GetMaskChannelOverride, Setter=SetMaskChannelOverride, BlueprintGetter=GetMaskChannelOverride,
-		Category = "Material Designer", DisplayName = "Channel Override",
+		Category = "Material Designer", DisplayName = "Channel Mask",
 		meta=(NotKeyframeable, ToolTip="Changes the output channel of the mask input."))
 	mutable EAvaColorChannel MaskChannelOverride;
 
diff --git a/Engine/Plugins/Experimental/DynamicMaterial/Source/DynamicMaterialEditor/Public/DMEDefs.h b/Engine/Plugins/Experimental/DynamicMaterial/Source/DynamicMaterialEditor/Public/DMEDefs.h
--- a/Engine/Plugins/Experimental/DynamicMaterial/Source/DynamicMaterialEditor/Public/DMEDefs.h
+++ b/Engine/Plugins/Experimental/DynamicMaterial/Source/DynamicMaterialEditor/Public/DMEDefs.h
@@ -141,3 +141,13 @@ enum class EDMMaterialLayerStage : uint8
 	All  = Base | Mask
 };
 ENUM_CLASS_FLAGS(EDMMaterialLayerStage);
+
+UENUM(BlueprintType)
+enum class EAvaColorChannel : uint8
+{
+	None = 0,
+	Red = 1 << 0,
+	Green = 1 << 1,
+	Blue = 1 << 2,
+	Alpha = 1 << 3
+};
