Material Designer
- Parameter object clean-up on load was causing the map to be dirtied. This is fixed.

diff --git a/Engine/Plugins/Experimental/DynamicMaterial/Source/DynamicMaterial/Private/Model/DynamicMaterialModel.cpp b/Engine/Plugins/Experimental/DynamicMaterial/Source/DynamicMaterial/Private/Model/DynamicMaterialModel.cpp
--- a/Engine/Plugins/Experimental/DynamicMaterial/Source/DynamicMaterial/Private/Model/DynamicMaterialModel.cpp
+++ b/Engine/Plugins/Experimental/DynamicMaterial/Source/DynamicMaterial/Private/Model/DynamicMaterialModel.cpp
@@ -385,11 +385,6 @@ void UDynamicMaterialModel::FreeParameter(UDMMaterialParameter* InParameter)
 
 	InParameter->ParameterName = NAME_None;
 	InParameter->SetComponentState(EDMComponentLifetimeState::Removed);
-
-	if (IDynamicMaterialModelEditorOnlyDataInterface* ModelEditorOnlyData = GetEditorOnlyData())
-	{
-		ModelEditorOnlyData->RequestMaterialBuild();
-	}
 }
 
 bool UDynamicMaterialModel::ConditionalFreeParameter(UDMMaterialParameter* InParameter)
