Material Designer
- Scripted action "texture set" (part 1 Texture Set system)

#virtualized

diff --git a/Engine/Plugins/Experimental/DynamicMaterial/DynamicMaterial.uplugin b/Engine/Plugins/Experimental/DynamicMaterial/DynamicMaterial.uplugin
--- a/Engine/Plugins/Experimental/DynamicMaterial/DynamicMaterial.uplugin
+++ b/Engine/Plugins/Experimental/DynamicMaterial/DynamicMaterial.uplugin
@@ -20,11 +20,21 @@
 			"Type": "RuntimeAndProgram",
 			"LoadingPhase": "Default"
 		},
+		{
+			"Name": "DynamicMaterialTextureSet",
+			"Type": "RuntimeAndProgram",
+			"LoadingPhase": "Default"
+		},
 		{
 			"Name": "DynamicMaterialEditor",
 			"Type": "Editor",
 			"LoadingPhase": "Default"
 		},
+		{
+			"Name": "DynamicMaterialTextureSetEditor",
+			"Type": "Editor",
+			"LoadingPhase": "Default"
+		},
 		{
 			"Name": "DynamicMaterialShadersEditor",
 			"Type": "Editor",
@@ -34,7 +44,11 @@
 	"Plugins": [
 		{
 			"Name": "CustomDetailsView",
-			"Enabled":  true
+			"Enabled": true
+		},
+		{
+			"Name": "StructUtils",
+			"Enabled": true
 		}
 	]
 }
\ No newline at end of file
diff --git a/Engine/Plugins/Experimental/DynamicMaterial/Source/DynamicMaterialEditor/DynamicMaterialEditor.Build.cs b/Engine/Plugins/Experimental/DynamicMaterial/Source/DynamicMaterialEditor/DynamicMaterialEditor.Build.cs
--- a/Engine/Plugins/Experimental/DynamicMaterial/Source/DynamicMaterialEditor/DynamicMaterialEditor.Build.cs
+++ b/Engine/Plugins/Experimental/DynamicMaterial/Source/DynamicMaterialEditor/DynamicMaterialEditor.Build.cs
@@ -29,6 +29,7 @@ public class DynamicMaterialEditor : ModuleRules
 				"CustomDetailsView",
 				"DeveloperSettings",
 				"DynamicMaterialShadersEditor",
+				"DynamicMaterialTextureSet",
 				"EditorSubsystem",
 				"EditorWidgets",
 				"InputCore",
diff --git a/Engine/Plugins/Experimental/DynamicMaterial/Source/DynamicMaterialEditor/Private/DMWorldSubsystem.cpp b/Engine/Plugins/Experimental/DynamicMaterial/Source/DynamicMaterialEditor/Private/DMWorldSubsystem.cpp
--- a/Engine/Plugins/Experimental/DynamicMaterial/Source/DynamicMaterialEditor/Private/DMWorldSubsystem.cpp
+++ b/Engine/Plugins/Experimental/DynamicMaterial/Source/DynamicMaterialEditor/Private/DMWorldSubsystem.cpp
@@ -1,8 +1,10 @@
 // Copyright Epic Games, Inc. All Rights Reserved.
  
 #include "DMWorldSubsystem.h"
-#include "IDetailKeyframeHandler.h"
+
+#include "GameFramework/Actor.h"
 #include "LevelEditor/DMLevelEditorIntegration.h"
+#include "Model/DynamicMaterialModel.h"
 
 UDMWorldSubsystem::UDMWorldSubsystem()
 	: KeyframeHandler(nullptr)
@@ -10,3 +12,67 @@ UDMWorldSubsystem::UDMWorldSubsystem()
 	// Default fallback implementation
 	InvokeTabDelegate.BindWeakLambda(this, [this]() { FDMLevelEditorIntegration::InvokeTabForWorld(GetWorld()); });
 }
+
+UDynamicMaterialModel* UDMWorldSubsystem::ExecuteGetCustomEditorModelDelegate()
+{
+	if (CustomModelEditorGetDelegate.IsBound())
+	{
+		return CustomModelEditorGetDelegate.Execute();
+	}
+
+	// An unbound delegate means there is no model to get.
+	return nullptr;
+}
+
+void UDMWorldSubsystem::ExecuteSetCustomEditorModelDelegate(UDynamicMaterialModel* InMaterialModel)
+{
+	if (!InMaterialModel || IsValid(InMaterialModel))
+	{
+		CustomModelEditorSetDelegate.ExecuteIfBound(InMaterialModel);
+	}
+}
+
+void UDMWorldSubsystem::ExecuteCustomObjectPropertyEditorDelegate(const FDMObjectMaterialProperty& InObjectProperty)
+{
+	CustomObjectPropertyEditorDelegate.ExecuteIfBound(InObjectProperty);
+}
+
+void UDMWorldSubsystem::ExecuteSetCustomEditorActorDelegate(AActor* InActor)
+{
+	if (!InActor || IsValid(InActor))
+	{
+		CustomActorEditorDelegate.ExecuteIfBound(InActor);
+	}
+}
+
+bool UDMWorldSubsystem::ExecuteIsValidDelegate(UDynamicMaterialModel* InMaterialModel)
+{
+	if (!IsValid(InMaterialModel))
+	{
+		return false;
+	}
+
+	if (IsValidDelegate.IsBound())
+	{
+		return IsValidDelegate.Execute(InMaterialModel);
+	}
+
+	// An unbound delegate means it is valid.
+	return true;
+}
+
+bool UDMWorldSubsystem::ExecuteMaterialValueSetterDelegate(const FDMObjectMaterialProperty& InObjectProperty, UDynamicMaterialInstance* InMaterialInstance)
+{
+	if (SetMaterialValueDelegate.IsBound())
+	{
+		return SetMaterialValueDelegate.Execute(InObjectProperty, InMaterialInstance);
+	}
+
+	// An unbound delegate means that no setting took place.
+	return false;
+}
+
+void UDMWorldSubsystem::ExecuteInvokeTabDelegate()
+{
+	InvokeTabDelegate.ExecuteIfBound();
+}
diff --git a/Engine/Plugins/Experimental/DynamicMaterial/Source/DynamicMaterialEditor/Private/DetailsPanel/Slate/SDMMaterialListExtensionWidget.cpp b/Engine/Plugins/Experimental/DynamicMaterial/Source/DynamicMaterialEditor/Private/DetailsPanel/Slate/SDMMaterialListExtensionWidget.cpp
--- a/Engine/Plugins/Experimental/DynamicMaterial/Source/DynamicMaterialEditor/Private/DetailsPanel/Slate/SDMMaterialListExtensionWidget.cpp
+++ b/Engine/Plugins/Experimental/DynamicMaterial/Source/DynamicMaterialEditor/Private/DetailsPanel/Slate/SDMMaterialListExtensionWidget.cpp
@@ -114,7 +114,7 @@ void SDMMaterialListExtensionWidget::SetDynamicMaterialInstance(UDynamicMaterial
 						{
 							const FMaterialListItem& ListItem = MaterialItemViewWeak.Pin()->GetMaterialListItem();
 							const FDMObjectMaterialProperty MaterialProperty(CurrentComponent, ListItem.SlotIndex);
-							WorldSubsystem->GetMaterialValueSetterDelegate().Execute(MaterialProperty, NewInstance);
+							WorldSubsystem->ExecuteMaterialValueSetterDelegate(MaterialProperty, NewInstance);
 						}
 
 						return;
diff --git a/Engine/Plugins/Experimental/DynamicMaterial/Source/DynamicMaterialEditor/Private/DynamicMaterialEditorModule.cpp b/Engine/Plugins/Experimental/DynamicMaterial/Source/DynamicMaterialEditor/Private/DynamicMaterialEditorModule.cpp
--- a/Engine/Plugins/Experimental/DynamicMaterial/Source/DynamicMaterialEditor/Private/DynamicMaterialEditorModule.cpp
+++ b/Engine/Plugins/Experimental/DynamicMaterial/Source/DynamicMaterialEditor/Private/DynamicMaterialEditorModule.cpp
@@ -1,6 +1,7 @@
 // Copyright Epic Games, Inc. All Rights Reserved.
 
 #include "DynamicMaterialEditorModule.h"
+
 #include "Components/ActorComponent.h"
 #include "Components/DMMaterialComponent.h"
 #include "Components/DMMaterialEffectFunction.h"
@@ -360,10 +361,10 @@ void FDynamicMaterialEditorModule::OpenMaterialModel(UDynamicMaterialModel* InMa
 		{
 			if (bInInvokeTab)
 			{
-				DMWorldSubsystem->GetInvokeTabDelegate().ExecuteIfBound();
+				DMWorldSubsystem->ExecuteInvokeTabDelegate();
 			}
 
-			DMWorldSubsystem->GetSetCustomEditorModelDelegate().ExecuteIfBound(InMaterialModel);
+			DMWorldSubsystem->ExecuteSetCustomEditorModelDelegate(InMaterialModel);
 		}
 	}
 }
@@ -388,10 +389,10 @@ void FDynamicMaterialEditorModule::OpenMaterialObjectProperty(const FDMObjectMat
 		{
 			if (bInInvokeTab)
 			{
-				DMWorldSubsystem->GetInvokeTabDelegate().ExecuteIfBound();
+				DMWorldSubsystem->ExecuteInvokeTabDelegate();
 			}
 
-			DMWorldSubsystem->GetCustomObjectPropertyEditorDelegate().ExecuteIfBound(InObjectProperty);
+			DMWorldSubsystem->ExecuteCustomObjectPropertyEditorDelegate(InObjectProperty);
 		}
 	}
 }
@@ -427,10 +428,10 @@ void FDynamicMaterialEditorModule::OnActorSelected(AActor* InActor, UWorld* InWo
 		{
 			if (bInInvokeTab)
 			{
-				DMWorldSubsystem->GetInvokeTabDelegate().ExecuteIfBound();
+				DMWorldSubsystem->ExecuteInvokeTabDelegate();
 			}
 
-			DMWorldSubsystem->GetSetCustomEditorActorDelegate().ExecuteIfBound(InActor);
+			DMWorldSubsystem->ExecuteSetCustomEditorActorDelegate(InActor);
 		}
 	}
 }
@@ -450,6 +451,7 @@ TSharedRef<SWidget> FDynamicMaterialEditorModule::CreateEditor(UDynamicMaterialM
 
 		if (IsValid(WorldSubsystem))
 		{
+			WorldSubsystem->GetGetCustomEditorModelDelegate().BindSP(NewEditor, &SDMEditor::GetMaterialModel);
 			WorldSubsystem->GetSetCustomEditorModelDelegate().BindSP(NewEditor, &SDMEditor::SetMaterialModel);
 			WorldSubsystem->GetCustomObjectPropertyEditorDelegate().BindSP(NewEditor, &SDMEditor::SetMaterialObjectProperty);
 			WorldSubsystem->GetSetCustomEditorActorDelegate().BindSP(NewEditor, &SDMEditor::OnActorSelected);
@@ -547,10 +549,28 @@ void FDynamicMaterialEditorModule::OpenEditor(UWorld* InWorld) const
 	}
 	else if (UDMWorldSubsystem* DMWorldSubsystem = InWorld->GetSubsystem<UDMWorldSubsystem>())
 	{
-		DMWorldSubsystem->GetInvokeTabDelegate().ExecuteIfBound();
+		DMWorldSubsystem->ExecuteInvokeTabDelegate();
 	}
 }
 
+UDynamicMaterialModel* FDynamicMaterialEditorModule::GetOpenedMaterialModel(UWorld* InWorld) const
+{
+	if (IsValid(InWorld))
+	{
+		if (TSharedPtr<SDMEditor> Editor = FDMLevelEditorIntegration::GetEditorForWorld(InWorld))
+		{
+			return Editor->GetMaterialModel();
+		}
+
+		if (UDMWorldSubsystem* DMWorldSubsystem = InWorld->GetSubsystem<UDMWorldSubsystem>())
+		{
+			return DMWorldSubsystem->ExecuteGetCustomEditorModelDelegate();
+		}
+	}
+
+	return nullptr;
+}
+
 void FDynamicMaterialEditorModule::ProcessBuildRequest(UObject* InToBuild, bool bInDirtyAssets)
 {
 	if (!IsValid(InToBuild))
diff --git a/Engine/Plugins/Experimental/DynamicMaterial/Source/DynamicMaterialEditor/Private/DynamicMaterialEditorModule.h b/Engine/Plugins/Experimental/DynamicMaterial/Source/DynamicMaterialEditor/Private/DynamicMaterialEditorModule.h
--- a/Engine/Plugins/Experimental/DynamicMaterial/Source/DynamicMaterialEditor/Private/DynamicMaterialEditorModule.h
+++ b/Engine/Plugins/Experimental/DynamicMaterial/Source/DynamicMaterialEditor/Private/DynamicMaterialEditorModule.h
@@ -2,19 +2,20 @@
 
 #pragma once
 
+#include "IDynamicMaterialEditorModule.h"
+#include "TickableEditorObject.h"
+
 #include "DMEDefs.h"
 #include "DMObjectMaterialProperty.h"
-#include "IDynamicMaterialEditorModule.h"
 #include "Templates/SharedPointer.h"
-#include "TickableEditorObject.h"
 
 class AActor;
 class FDMMaterialFunctionLibrary;
 class FUICommandList;
 class IAssetTypeActions;
 class ILevelEditor;
-class SDMEditor;
 class SDMComponentEdit;
+class SDMEditor;
 class SDockTab;
 class SWidget;
 class UDMMaterialComponent;
@@ -92,6 +93,7 @@ public:
 	virtual void RegisterMaterialModelCreatedCallback(const TSharedRef<IDMOnWizardCompleteCallback> InCallback)  override;
 	virtual void UnregisterMaterialModelCreatedCallback(const TSharedRef<IDMOnWizardCompleteCallback> InCallback) override;
 	virtual void OpenEditor(UWorld* InWorld) const override;
+	virtual UDynamicMaterialModel* GetOpenedMaterialModel(UWorld* InWorld) const override;
 	virtual void OpenMaterialModel(UDynamicMaterialModel* InMaterialModel, UWorld* InWorld, bool bInInvokeTab) const override;
 	virtual void OpenMaterialObjectProperty(const FDMObjectMaterialProperty& InObjectProperty, UWorld* InWorld, bool bInInvokeTab) const override;
 	virtual void OpenMaterialInstance(UDynamicMaterialInstance* InInstance, UWorld* InWorld, bool bInInvokeTab) const override;
diff --git a/Engine/Plugins/Experimental/DynamicMaterial/Source/DynamicMaterialEditor/Private/Model/DynamicMaterialModelEditorOnlyData.cpp b/Engine/Plugins/Experimental/DynamicMaterial/Source/DynamicMaterialEditor/Private/Model/DynamicMaterialModelEditorOnlyData.cpp
--- a/Engine/Plugins/Experimental/DynamicMaterial/Source/DynamicMaterialEditor/Private/Model/DynamicMaterialModelEditorOnlyData.cpp
+++ b/Engine/Plugins/Experimental/DynamicMaterial/Source/DynamicMaterialEditor/Private/Model/DynamicMaterialModelEditorOnlyData.cpp
@@ -1,6 +1,7 @@
 // Copyright Epic Games, Inc. All Rights Reserved.
 
 #include "Model/DynamicMaterialModelEditorOnlyData.h"
+
 #include "AssetRegistry/AssetRegistryModule.h"
 #include "AssetToolsModule.h"
 #include "Components/DMMaterialLayer.h"
@@ -127,6 +128,16 @@ UDynamicMaterialModelEditorOnlyData* UDynamicMaterialModelEditorOnlyData::Get(ID
 	return Cast<UDynamicMaterialModelEditorOnlyData>(InInterface);
 }
 
+UDynamicMaterialModelEditorOnlyData* UDynamicMaterialModelEditorOnlyData::Get(UDynamicMaterialInstance* InInstance)
+{
+	if (InInstance)
+	{
+		return Get(InInstance->GetMaterialModel());
+	}
+
+	return nullptr;
+}
+
 UDynamicMaterialModelEditorOnlyData::UDynamicMaterialModelEditorOnlyData()
 	: State(EDMState::Idle)
 	, Domain(EMaterialDomain::MD_Surface)
@@ -796,6 +807,12 @@ TSharedRef<FDMMaterialBuildState> UDynamicMaterialModelEditorOnlyData::CreateBui
 	return BuildState;
 }
 
+bool UDynamicMaterialModelEditorOnlyData::AddTextureSet(UDMTextureSet* InTextureSet, bool bInReplaceSlots)
+{
+	// TODO Integrate!
+	return false;
+}
+
 bool UDynamicMaterialModelEditorOnlyData::NeedsWizard() const
 {
 	return ChannelListPreset == NAME_None;
diff --git a/Engine/Plugins/Experimental/DynamicMaterial/Source/DynamicMaterialEditor/Private/Slate/SDMEditor.cpp b/Engine/Plugins/Experimental/DynamicMaterial/Source/DynamicMaterialEditor/Private/Slate/SDMEditor.cpp
--- a/Engine/Plugins/Experimental/DynamicMaterial/Source/DynamicMaterialEditor/Private/Slate/SDMEditor.cpp
+++ b/Engine/Plugins/Experimental/DynamicMaterial/Source/DynamicMaterialEditor/Private/Slate/SDMEditor.cpp
@@ -194,8 +194,7 @@ namespace UE::DynamicMaterialEditor::Private
 			{
 				if (UDMWorldSubsystem* WorldSubsystem = World->GetSubsystem<UDMWorldSubsystem>())
 				{
-					if (WorldSubsystem->GetIsValidDelegate().IsBound()
-						&& WorldSubsystem->GetIsValidDelegate().Execute(InMaterialModel) == false)
+					if (!WorldSubsystem->ExecuteIsValidDelegate(InMaterialModel))
 					{
 						return false;
 					}
diff --git a/Engine/Plugins/Experimental/DynamicMaterial/Source/DynamicMaterialEditor/Private/Slate/SDMToolBar.cpp b/Engine/Plugins/Experimental/DynamicMaterial/Source/DynamicMaterialEditor/Private/Slate/SDMToolBar.cpp
--- a/Engine/Plugins/Experimental/DynamicMaterial/Source/DynamicMaterialEditor/Private/Slate/SDMToolBar.cpp
+++ b/Engine/Plugins/Experimental/DynamicMaterial/Source/DynamicMaterialEditor/Private/Slate/SDMToolBar.cpp
@@ -599,14 +599,12 @@ FReply SDMToolBar::OnUseClicked()
 
 	if (DMSubsystem && DMSubsystem->GetMaterialValueSetterDelegate().IsBound())
 	{
-		if (IsValid(CurrentModel)
-			&& DMSubsystem->GetIsValidDelegate().IsBound()
-			&& !DMSubsystem->GetIsValidDelegate().Execute(CurrentModel))
+		if (!DMSubsystem->ExecuteIsValidDelegate(CurrentModel))
 		{
 			return FReply::Handled();
 		}
 
-		DMSubsystem->GetMaterialValueSetterDelegate().Execute(*CurrentActorProperty, SelectedInstance);
+		DMSubsystem->ExecuteMaterialValueSetterDelegate(*CurrentActorProperty, SelectedInstance);
 	}
 	else
 	{
diff --git a/Engine/Plugins/Experimental/DynamicMaterial/Source/DynamicMaterialEditor/Private/Utils/DMBlueprintFunctionLibrary.cpp b/Engine/Plugins/Experimental/DynamicMaterial/Source/DynamicMaterialEditor/Private/Utils/DMBlueprintFunctionLibrary.cpp
--- a/Engine/Plugins/Experimental/DynamicMaterial/Source/DynamicMaterialEditor/Private/Utils/DMBlueprintFunctionLibrary.cpp
+++ b/Engine/Plugins/Experimental/DynamicMaterial/Source/DynamicMaterialEditor/Private/Utils/DMBlueprintFunctionLibrary.cpp
@@ -222,14 +222,9 @@ UDynamicMaterialModel* UDMBlueprintFunctionLibrary::CreateDynamicMaterialInObjec
 	{
 		if (IsValid(World))
 		{
-			UDMWorldSubsystem* const WorldSubsystem = World->GetSubsystem<UDMWorldSubsystem>();
-
-			if (IsValid(WorldSubsystem))
+			if (UDMWorldSubsystem* const WorldSubsystem = World->GetSubsystem<UDMWorldSubsystem>())
 			{
-				if (WorldSubsystem->GetMaterialValueSetterDelegate().IsBound())
-				{
-					bSubsystemTakenOver = WorldSubsystem->GetMaterialValueSetterDelegate().Execute(InMaterialProperty, NewInstance);
-				}
+				bSubsystemTakenOver = WorldSubsystem->ExecuteMaterialValueSetterDelegate(InMaterialProperty, NewInstance);
 			}
 		}
 	}
diff --git a/Engine/Plugins/Experimental/DynamicMaterial/Source/DynamicMaterialEditor/Public/DMWorldSubsystem.h b/Engine/Plugins/Experimental/DynamicMaterial/Source/DynamicMaterialEditor/Public/DMWorldSubsystem.h
--- a/Engine/Plugins/Experimental/DynamicMaterial/Source/DynamicMaterialEditor/Public/DMWorldSubsystem.h
+++ b/Engine/Plugins/Experimental/DynamicMaterial/Source/DynamicMaterialEditor/Public/DMWorldSubsystem.h
@@ -2,9 +2,11 @@
  
 #pragma once
 
-#include "Delegates/Delegate.h"
 #include "Subsystems/WorldSubsystem.h"
+
+#include "Delegates/Delegate.h"
 #include "Templates/SharedPointer.h"
+
 #include "DMWorldSubsystem.generated.h"
 
 class AActor;
@@ -13,6 +15,7 @@ class UDynamicMaterialInstance;
 class UDynamicMaterialModel;
 struct FDMObjectMaterialProperty;
 
+DECLARE_DELEGATE_RetVal(UDynamicMaterialModel*, FDMGetMaterialModelDelegate)
 DECLARE_DELEGATE_OneParam(FDMSetMaterialModelDelegate, UDynamicMaterialModel*)
 DECLARE_DELEGATE_OneParam(FDMSetMaterialObjectPropertyDelegate, const FDMObjectMaterialProperty&)
 DECLARE_DELEGATE_OneParam(FDMSetMaterialActorDelegate, AActor*)
@@ -28,34 +31,79 @@ class UDMWorldSubsystem : public UWorldSubsystem
 public:
 	UDMWorldSubsystem();
  
-	const TSharedPtr<IDetailKeyframeHandler>& GetKeyframeHandler() const { return KeyframeHandler; }
+	const TSharedPtr<IDetailKeyframeHandler>& GetKeyframeHandler() const
+	{
+		return KeyframeHandler;
+	}
+
+	void SetKeyframeHandler(const TSharedPtr<IDetailKeyframeHandler>& InKeyframeHandler)
+	{
+		KeyframeHandler = InKeyframeHandler;
+	}
+
+	/** Gets the material in a custom editor tab. */
+	FDMGetMaterialModelDelegate::RegistrationType& GetGetCustomEditorModelDelegate() const
+	{
+		return CustomModelEditorGetDelegate;
+	}
 
-	void SetKeyframeHandler(const TSharedPtr<IDetailKeyframeHandler>& InKeyframeHandler) { KeyframeHandler = InKeyframeHandler; }
+	UDynamicMaterialModel* ExecuteGetCustomEditorModelDelegate();
 
 	/** Sets material in a custom editor tab. */
-	FDMSetMaterialModelDelegate& GetSetCustomEditorModelDelegate() { return CustomModelEditorDelegate; }
+	FDMSetMaterialModelDelegate::RegistrationType& GetSetCustomEditorModelDelegate() const
+	{
+		return CustomModelEditorSetDelegate;
+	}
+
+	void ExecuteSetCustomEditorModelDelegate(UDynamicMaterialModel* InMaterialModel);
 
 	/** Sets the object property in custom editor tab. */
-	FDMSetMaterialObjectPropertyDelegate& GetCustomObjectPropertyEditorDelegate() { return CustomObjectPropertyEditorDelegate; }
+	FDMSetMaterialObjectPropertyDelegate::RegistrationType& GetCustomObjectPropertyEditorDelegate() const
+	{
+		return CustomObjectPropertyEditorDelegate;
+	}
+
+	void ExecuteCustomObjectPropertyEditorDelegate(const FDMObjectMaterialProperty& InObjectProperty);
 
 	/** Sets actor in a custom editor tab. */
-	FDMSetMaterialActorDelegate& GetSetCustomEditorActorDelegate() { return CustomActorEditorDelegate; }
+	FDMSetMaterialActorDelegate::RegistrationType& GetSetCustomEditorActorDelegate() const
+	{
+		return CustomActorEditorDelegate;
+	}
+
+	void ExecuteSetCustomEditorActorDelegate(AActor* InActor);
 
 	/** Returns true if the supplied material is valid for this world. */
-	FDMIsValidDelegate& GetIsValidDelegate() { return IsValidDelegate; }
+	FDMIsValidDelegate::RegistrationType& GetIsValidDelegate() const
+	{
+		return IsValidDelegate;
+	}
+
+	bool ExecuteIsValidDelegate(UDynamicMaterialModel* InMaterialModel);
 
 	/** Used to redirect SetMaterial to different objects/paths. */
-	FDMSetMaterialValueDelegate& GetMaterialValueSetterDelegate() { return SetMaterialValueDelegate; }
+	FDMSetMaterialValueDelegate::RegistrationType& GetMaterialValueSetterDelegate() const
+	{
+		return SetMaterialValueDelegate;
+	}
+
+	bool ExecuteMaterialValueSetterDelegate(const FDMObjectMaterialProperty& InObjectProperty, UDynamicMaterialInstance* InMaterialInstance);
 
 	/** Used to show the tab to the user. */
-	FDMInvokeTabDelegate& GetInvokeTabDelegate() { return InvokeTabDelegate; }
+	FDMInvokeTabDelegate::RegistrationType& GetInvokeTabDelegate() const
+	{
+		return InvokeTabDelegate;
+	}
+
+	void ExecuteInvokeTabDelegate();
 
 protected:
 	TSharedPtr<IDetailKeyframeHandler> KeyframeHandler;
-	FDMSetMaterialModelDelegate CustomModelEditorDelegate;
-	FDMSetMaterialObjectPropertyDelegate CustomObjectPropertyEditorDelegate;
-	FDMSetMaterialActorDelegate CustomActorEditorDelegate;
-	FDMIsValidDelegate IsValidDelegate;
-	FDMSetMaterialValueDelegate SetMaterialValueDelegate;
-	FDMInvokeTabDelegate InvokeTabDelegate;
+	mutable FDMGetMaterialModelDelegate CustomModelEditorGetDelegate;
+	mutable FDMSetMaterialModelDelegate CustomModelEditorSetDelegate;
+	mutable FDMSetMaterialObjectPropertyDelegate CustomObjectPropertyEditorDelegate;
+	mutable FDMSetMaterialActorDelegate CustomActorEditorDelegate;
+	mutable FDMIsValidDelegate IsValidDelegate;
+	mutable FDMSetMaterialValueDelegate SetMaterialValueDelegate;
+	mutable FDMInvokeTabDelegate InvokeTabDelegate;
 };
diff --git a/Engine/Plugins/Experimental/DynamicMaterial/Source/DynamicMaterialEditor/Public/IDynamicMaterialEditorModule.h b/Engine/Plugins/Experimental/DynamicMaterial/Source/DynamicMaterialEditor/Public/IDynamicMaterialEditorModule.h
--- a/Engine/Plugins/Experimental/DynamicMaterial/Source/DynamicMaterialEditor/Public/IDynamicMaterialEditorModule.h
+++ b/Engine/Plugins/Experimental/DynamicMaterial/Source/DynamicMaterialEditor/Public/IDynamicMaterialEditorModule.h
@@ -2,9 +2,10 @@
 
 #pragma once
 
+#include "Modules/ModuleManager.h"
+
 #include "Containers/ContainersFwd.h"
 #include "Delegates/Delegate.h"
-#include "Modules/ModuleManager.h"
 #include "Templates/SharedPointerFwd.h"
 
 class AActor;
@@ -51,6 +52,8 @@ public:
 
 	virtual void OpenEditor(UWorld* InWorld) const = 0;
 
+	virtual UDynamicMaterialModel* GetOpenedMaterialModel(UWorld* InWorld) const = 0;
+
 	virtual void OpenMaterialModel(UDynamicMaterialModel* InMaterialModel, UWorld* InWorld, bool bInInvokeTab) const = 0;
 
 	virtual void OpenMaterialObjectProperty(const FDMObjectMaterialProperty& InObjectProperty, UWorld* InWorld, bool bInInvokeTab) const = 0;
diff --git a/Engine/Plugins/Experimental/DynamicMaterial/Source/DynamicMaterialEditor/Public/Model/DynamicMaterialModelEditorOnlyData.h b/Engine/Plugins/Experimental/DynamicMaterial/Source/DynamicMaterialEditor/Public/Model/DynamicMaterialModelEditorOnlyData.h
--- a/Engine/Plugins/Experimental/DynamicMaterial/Source/DynamicMaterialEditor/Public/Model/DynamicMaterialModelEditorOnlyData.h
+++ b/Engine/Plugins/Experimental/DynamicMaterial/Source/DynamicMaterialEditor/Public/Model/DynamicMaterialModelEditorOnlyData.h
@@ -2,13 +2,16 @@
 
 #pragma once
 
-#include "Model/IDynamicMaterialModelEditorOnlyDataInterface.h"
 #include "DMEDefs.h"
+#include "Misc/NotifyHook.h"
+#include "Model/IDynamicMaterialModelEditorOnlyDataInterface.h"
+#include "UObject/Object.h"
+
 #include "Engine/EngineTypes.h"
 #include "MaterialDomain.h"
-#include "Misc/NotifyHook.h"
 #include "UObject/WeakObjectPtrFwd.h"
 #include "UObject/WeakObjectPtrTemplatesFwd.h"
+
 #include "DynamicMaterialModelEditorOnlyData.generated.h"
 
 class UDMMaterialComponent;
@@ -17,6 +20,7 @@ class UDMMaterialProperty;
 class UDMMaterialSlot;
 class UDMMaterialValue;
 class UDMMaterialValueFloat1;
+class UDMTextureSet;
 class UDMTextureUV;
 class UDynamicMaterialModel;
 class UDynamicMaterialModelEditorOnlyData;
@@ -76,6 +80,7 @@ public:
 	static UDynamicMaterialModelEditorOnlyData* Get(TWeakObjectPtr<UDynamicMaterialModel> InModelWeak);
 	static UDynamicMaterialModelEditorOnlyData* Get(const TScriptInterface<IDynamicMaterialModelEditorOnlyDataInterface>& InInterface);
 	static UDynamicMaterialModelEditorOnlyData* Get(IDynamicMaterialModelEditorOnlyDataInterface* InInterface);
+	static UDynamicMaterialModelEditorOnlyData* Get(UDynamicMaterialInstance* InInstance);
 
 	UDynamicMaterialModelEditorOnlyData();
 
@@ -176,6 +181,15 @@ public:
 
 	TSharedRef<FDMMaterialBuildState> CreateBuildState(UMaterial* InMaterialToBuild, bool bInDirtyAssets = true) const;
 
+	/**
+	 * Integrates a Texture Set with the model.
+	 * @param InTextureSet The set to integrate.
+	 * @param bInReplaceSlots Whether to add to or completely replace slots.
+	 * @return True on success.
+	 */
+	UFUNCTION(BlueprintCallable, Category = "Material Designer")
+	bool AddTextureSet(UDMTextureSet* InTextureSet, bool bInReplaceSlots);
+
 	bool NeedsWizard() const;
 
 	void OnWizardComplete();
diff --git a/Engine/Plugins/Experimental/DynamicMaterial/Source/DynamicMaterialTextureSet/DynamicMaterialTextureSet.Build.cs b/Engine/Plugins/Experimental/DynamicMaterial/Source/DynamicMaterialTextureSet/DynamicMaterialTextureSet.Build.cs
new file mode 100644
--- /dev/null
+++ b/Engine/Plugins/Experimental/DynamicMaterial/Source/DynamicMaterialTextureSet/DynamicMaterialTextureSet.Build.cs
@@ -0,0 +1,28 @@
+// Copyright Epic Games, Inc. All Rights Reserved.
+
+using UnrealBuildTool;
+
+public class DynamicMaterialTextureSet : ModuleRules
+{
+    public DynamicMaterialTextureSet(ReadOnlyTargetRules Target) : base(Target)
+    {
+        PCHUsage = ModuleRules.PCHUsageMode.UseExplicitOrSharedPCHs;
+
+        PublicDependencyModuleNames.AddRange(
+            new string[]
+            {
+                "CoreUObject",
+                "Engine"
+            }
+        );
+
+        PrivateDependencyModuleNames.AddRange(
+            new string[]
+            {
+				"Core"
+            }
+        );
+
+        ShortName = "DynMatTexSet";
+    }
+}
diff --git a/Engine/Plugins/Experimental/DynamicMaterial/Source/DynamicMaterialTextureSet/Private/DMTextureSet.cpp b/Engine/Plugins/Experimental/DynamicMaterial/Source/DynamicMaterialTextureSet/Private/DMTextureSet.cpp
new file mode 100644
--- /dev/null
+++ b/Engine/Plugins/Experimental/DynamicMaterial/Source/DynamicMaterialTextureSet/Private/DMTextureSet.cpp
@@ -0,0 +1,87 @@
+// Copyright Epic Games, Inc. All Rights Reserved.
+
+#include "DMTextureSet.h"
+
+#include "Engine/Texture.h"
+#include "SceneTypes.h"
+
+UDMTextureSet::UDMTextureSet()
+{
+	Textures.Reserve(13); // EMaterialProperty has more than 13 entries.
+
+	Textures.Emplace(EMaterialProperty::MP_EmissiveColor);
+	Textures.Emplace(EMaterialProperty::MP_Opacity);
+	Textures.Emplace(EMaterialProperty::MP_OpacityMask);
+	Textures.Emplace(EMaterialProperty::MP_BaseColor);
+	Textures.Emplace(EMaterialProperty::MP_Metallic);
+	Textures.Emplace(EMaterialProperty::MP_Specular);
+	Textures.Emplace(EMaterialProperty::MP_Roughness);
+	Textures.Emplace(EMaterialProperty::MP_Anisotropy);
+	Textures.Emplace(EMaterialProperty::MP_Normal);
+	Textures.Emplace(EMaterialProperty::MP_Tangent);
+	Textures.Emplace(EMaterialProperty::MP_SubsurfaceColor);
+	Textures.Emplace(EMaterialProperty::MP_AmbientOcclusion);
+	Textures.Emplace(EMaterialProperty::MP_Refraction);
+}
+
+bool UDMTextureSet::HasMaterialProperty(TEnumAsByte<EMaterialProperty> InMaterialProperty) const
+{
+	return Textures.Contains(InMaterialProperty);
+}
+
+const TMap<TEnumAsByte<EMaterialProperty>, FDMMaterialTexture>& UDMTextureSet::GetTextures() const
+{
+	return Textures;
+}
+
+bool UDMTextureSet::HasMaterialTexture(TEnumAsByte<EMaterialProperty> InMaterialProperty) const
+{
+	if (const FDMMaterialTexture* MaterialTexture = Textures.Find(InMaterialProperty))
+	{
+		return !MaterialTexture->Texture.IsNull();
+	}
+
+	return false;
+}
+
+bool UDMTextureSet::GetMaterialTexture(TEnumAsByte<EMaterialProperty> InMaterialProperty, FDMMaterialTexture& OutMaterialTexture) const
+{
+	if (const FDMMaterialTexture* MaterialTexture = Textures.Find(InMaterialProperty))
+	{
+		OutMaterialTexture = *MaterialTexture;
+		return true;
+	}
+
+	return false;
+}
+
+const FDMMaterialTexture* UDMTextureSet::GetMaterialTexture(TEnumAsByte<EMaterialProperty> InMaterialProperty) const
+{
+	return Textures.Find(InMaterialProperty);
+}
+
+void UDMTextureSet::SetMaterialTexture(TEnumAsByte<EMaterialProperty> InMaterialProperty, const FDMMaterialTexture& InMaterialTexture)
+{
+	if (FDMMaterialTexture* MaterialTexture = Textures.Find(InMaterialProperty))
+	{
+		*MaterialTexture = InMaterialTexture;
+	}
+}
+
+bool UDMTextureSet::ContainsTexture(UTexture* InTexture) const
+{
+	if (!IsValid(InTexture))
+	{
+		return false;
+	}
+
+	for (const TPair<TEnumAsByte<EMaterialProperty>, FDMMaterialTexture>& Pair : Textures)
+	{
+		if (Pair.Value.Texture == InTexture)
+		{
+			return true;
+		}
+	}
+
+	return false;
+}
diff --git a/Engine/Plugins/Experimental/DynamicMaterial/Source/DynamicMaterialTextureSet/Private/DynamicMaterialTextureSetModule.cpp b/Engine/Plugins/Experimental/DynamicMaterial/Source/DynamicMaterialTextureSet/Private/DynamicMaterialTextureSetModule.cpp
new file mode 100644
--- /dev/null
+++ b/Engine/Plugins/Experimental/DynamicMaterial/Source/DynamicMaterialTextureSet/Private/DynamicMaterialTextureSetModule.cpp
@@ -0,0 +1,5 @@
+// Copyright Epic Games, Inc. All Rights Reserved.
+
+#include "Modules/ModuleManager.h"
+
+IMPLEMENT_MODULE(FDefaultModuleImpl, DynamicMaterialTextureSet)
diff --git a/Engine/Plugins/Experimental/DynamicMaterial/Source/DynamicMaterialTextureSet/Public/DMMaterialTexture.h b/Engine/Plugins/Experimental/DynamicMaterial/Source/DynamicMaterialTextureSet/Public/DMMaterialTexture.h
new file mode 100644
--- /dev/null
+++ b/Engine/Plugins/Experimental/DynamicMaterial/Source/DynamicMaterialTextureSet/Public/DMMaterialTexture.h
@@ -0,0 +1,22 @@
+// Copyright Epic Games, Inc. All Rights Reserved.
+
+#pragma once
+
+#include "DMTextureChannelMask.h"
+#include "UObject/SoftObjectPtr.h"
+
+#include "DMMaterialTexture.generated.h"
+
+class UTexture;
+
+USTRUCT(BlueprintType)
+struct FDMMaterialTexture
+{
+	GENERATED_BODY()
+
+	UPROPERTY(EditAnywhere, BlueprintReadWrite, Category = "Material Designer")
+	TSoftObjectPtr<UTexture> Texture;
+
+	UPROPERTY(EditAnywhere, BlueprintReadWrite, Category = "Material Designer")
+	EDMTextureChannelMask TextureChannel = EDMTextureChannelMask::RGBA;
+};
diff --git a/Engine/Plugins/Experimental/DynamicMaterial/Source/DynamicMaterialTextureSet/Public/DMTextureChannelMask.h b/Engine/Plugins/Experimental/DynamicMaterial/Source/DynamicMaterialTextureSet/Public/DMTextureChannelMask.h
new file mode 100644
--- /dev/null
+++ b/Engine/Plugins/Experimental/DynamicMaterial/Source/DynamicMaterialTextureSet/Public/DMTextureChannelMask.h
@@ -0,0 +1,18 @@
+// Copyright Epic Games, Inc. All Rights Reserved.
+
+#pragma once
+
+#include "HAL/Platform.h"
+#include "UObject/ObjectMacros.h"
+
+UENUM(BlueprintType)
+enum class EDMTextureChannelMask : uint8
+{
+	Red = 0,
+	Green = 1 << 0,
+	Blue = 1 << 1,
+	Alpha = 1 << 2,
+	RGB = Red|Green|Blue,
+	RGBA = Red|Green|Blue|Alpha,
+};
+ENUM_CLASS_FLAGS(EDMTextureChannelMask)
diff --git a/Engine/Plugins/Experimental/DynamicMaterial/Source/DynamicMaterialTextureSet/Public/DMTextureSet.h b/Engine/Plugins/Experimental/DynamicMaterial/Source/DynamicMaterialTextureSet/Public/DMTextureSet.h
new file mode 100644
--- /dev/null
+++ b/Engine/Plugins/Experimental/DynamicMaterial/Source/DynamicMaterialTextureSet/Public/DMTextureSet.h
@@ -0,0 +1,84 @@
+// Copyright Epic Games, Inc. All Rights Reserved.
+
+#pragma once
+
+#include "UObject/Object.h"
+
+#include "Containers/EnumAsByte.h"
+#include "Containers/Map.h"
+#include "DMMaterialTexture.h"
+
+#include "DMTextureSet.generated.h"
+
+enum EMaterialProperty : int;
+
+UCLASS(BlueprintType, ClassGroup = "Material Designer")
+class DYNAMICMATERIALTEXTURESET_API UDMTextureSet : public UObject
+{
+	GENERATED_BODY()
+
+public:
+	UDMTextureSet();
+
+	virtual ~UDMTextureSet() override = default;
+
+	/**
+	 * Checks whether a given Material Property exists in the Texture Map. Does not check whether
+	 * that a Texture is assigned to it.
+	 * @param InMaterialProperty The Material Property to check.
+	 * @return True if the property exists in the Texture Map.
+	 */
+	UFUNCTION(BlueprintPure, Category = "Material Designer")
+	bool HasMaterialProperty(TEnumAsByte<EMaterialProperty> InMaterialProperty) const;
+
+	/**
+	 * @return Gets the entire Texture Map.
+	 */
+	const TMap<TEnumAsByte<EMaterialProperty>, FDMMaterialTexture>& GetTextures() const;
+
+	/**
+	 * Checks whether a given Material Property has a Texture assigned to it.
+	 * @param InMaterialProperty The Material Property to check.
+	 * @return True if the Material Property has an assigned Texture.
+	 */
+	UFUNCTION(BlueprintPure, Category = "Material Designer")
+	bool HasMaterialTexture(TEnumAsByte<EMaterialProperty> InMaterialProperty) const;
+
+	/**
+	 * Gets the Material Texture associated with a Material Property. Does not check whether a Texture
+	 * is assigned to it.
+	 * @param InMaterialProperty The Material Property to check.
+	 * @param OutMaterialTexture The found Material Texture.
+	 * @return True if the Material Property exists within the Texture Map.
+	 */
+	UFUNCTION(BlueprintPure, Category = "Material Designer")
+	bool GetMaterialTexture(TEnumAsByte<EMaterialProperty> InMaterialProperty, FDMMaterialTexture& OutMaterialTexture) const;
+
+	/**
+	 * Gets the Material Texture associated with a Material Property. Does not check whether a Texture
+	 * is assigned to it.
+	 * @param InMaterialProperty The Material Property to check.
+	 * @return A pointer to the found Material Texture, if it exists, or nullptr.
+	 */
+	const FDMMaterialTexture* GetMaterialTexture(TEnumAsByte<EMaterialProperty> InMaterialProperty) const;
+
+	/**
+	 * Sets the Material Texture for a given Material Property. Can be used to unset Textures.
+	 * @param InMaterialProperty The Material Property to check.
+	 * @param InMaterialTexture The Material Texture to set on the given Material Property.
+	 */
+	UFUNCTION(BlueprintCallable, Category = "Material Designer")
+	void SetMaterialTexture(TEnumAsByte<EMaterialProperty> InMaterialProperty, const FDMMaterialTexture& InMaterialTexture);
+
+	/**
+	 * Checks the Texture Map to see if a given Texture exists within it.
+	 * @param InTexture The Texture to search for.
+	 * @return True if the Texture exits in the Texture Map.
+	 */
+	UFUNCTION(BlueprintPure, Category = "Material Designer")
+	bool ContainsTexture(UTexture* InTexture) const;
+
+protected:
+	UPROPERTY(EditAnywhere, BlueprintReadWrite, Category = "Material Designer", EditFixedSize, meta = (ReadOnlyKeys, AllowPrivateAccess))
+	TMap<TEnumAsByte<EMaterialProperty>, FDMMaterialTexture> Textures;
+};
diff --git a/Engine/Plugins/Experimental/DynamicMaterial/Source/DynamicMaterialTextureSetEditor/DynamicMaterialTextureSetEditor.Build.cs b/Engine/Plugins/Experimental/DynamicMaterial/Source/DynamicMaterialTextureSetEditor/DynamicMaterialTextureSetEditor.Build.cs
new file mode 100644
--- /dev/null
+++ b/Engine/Plugins/Experimental/DynamicMaterial/Source/DynamicMaterialTextureSetEditor/DynamicMaterialTextureSetEditor.Build.cs
@@ -0,0 +1,40 @@
+// Copyright Epic Games, Inc. All Rights Reserved.
+
+using UnrealBuildTool;
+
+public class DynamicMaterialTextureSetEditor : ModuleRules
+{
+	public DynamicMaterialTextureSetEditor(ReadOnlyTargetRules Target) : base(Target)
+	{
+		PCHUsage = ModuleRules.PCHUsageMode.UseExplicitOrSharedPCHs;
+
+		PublicDependencyModuleNames.AddRange(
+			new string[]
+			{
+				"CoreUObject",
+				"UnrealEd"
+			}
+		);
+
+		PrivateDependencyModuleNames.AddRange(
+			new string[]
+			{
+				"AssetDefinition",
+				"Core",
+				"DeveloperSettings",
+				"DynamicMaterial",
+				"DynamicMaterialEditor",
+				"DynamicMaterialTextureSet",
+				"EditorWidgets",
+				"Engine",
+				"InputCore",
+				"PropertyEditor",
+				"Slate",
+				"SlateCore",
+				"StructUtils"
+			}
+		);
+
+		ShortName = "DynMatTexSetEd";
+	}
+}
diff --git a/Engine/Plugins/Experimental/DynamicMaterial/Source/DynamicMaterialTextureSetEditor/Private/AssetDefinition_DMTextureSet.cpp b/Engine/Plugins/Experimental/DynamicMaterial/Source/DynamicMaterialTextureSetEditor/Private/AssetDefinition_DMTextureSet.cpp
new file mode 100644
--- /dev/null
+++ b/Engine/Plugins/Experimental/DynamicMaterial/Source/DynamicMaterialTextureSetEditor/Private/AssetDefinition_DMTextureSet.cpp
@@ -0,0 +1,35 @@
+// Copyright Epic Games, Inc. All Rights Reserved.
+
+#include "AssetDefinition_DMTextureSet.h"
+
+#include "DMTextureSet.h"
+
+#define LOCTEXT_NAMESPACE "AssetDefinition_DMTextureSet"
+
+FText UAssetDefinition_DMTextureSet::GetAssetDisplayName() const
+{
+	return LOCTEXT("DMTextureSet", "Texture Set");
+}
+
+FText UAssetDefinition_DMTextureSet::GetAssetDisplayName(const FAssetData& InAssetData) const
+{
+	return GetAssetDisplayName();
+}
+
+TSoftClassPtr<> UAssetDefinition_DMTextureSet::GetAssetClass() const
+{
+	return UDMTextureSet::StaticClass();
+}
+
+FLinearColor UAssetDefinition_DMTextureSet::GetAssetColor() const
+{
+	return FLinearColor(FColor(192, 64, 64));
+}
+
+TConstArrayView<FAssetCategoryPath> UAssetDefinition_DMTextureSet::GetAssetCategories() const
+{
+	static const TArray<FAssetCategoryPath> Categories = {EAssetCategoryPaths::Texture};
+	return Categories;
+}
+
+#undef LOCTEXT_NAMESPACE
diff --git a/Engine/Plugins/Experimental/DynamicMaterial/Source/DynamicMaterialTextureSetEditor/Private/AssetDefinition_DMTextureSet.h b/Engine/Plugins/Experimental/DynamicMaterial/Source/DynamicMaterialTextureSetEditor/Private/AssetDefinition_DMTextureSet.h
new file mode 100644
--- /dev/null
+++ b/Engine/Plugins/Experimental/DynamicMaterial/Source/DynamicMaterialTextureSetEditor/Private/AssetDefinition_DMTextureSet.h
@@ -0,0 +1,23 @@
+// Copyright Epic Games, Inc. All Rights Reserved.
+
+#pragma once
+
+#include "AssetDefinitionDefault.h"
+#include "AssetDefinition_DMTextureSet.generated.h"
+
+class UAssetDefinition_DMTextureSetInterface;
+
+UCLASS()
+class UAssetDefinition_DMTextureSet : public UAssetDefinitionDefault
+{
+	GENERATED_BODY()
+
+public:
+	//~ Begin UAssetDefinition Begin
+	virtual FText GetAssetDisplayName() const override;
+	virtual FText GetAssetDisplayName(const FAssetData& InAssetData) const override;
+	virtual TSoftClassPtr<UObject> GetAssetClass() const override;
+	virtual FLinearColor GetAssetColor() const override;
+	virtual TConstArrayView<FAssetCategoryPath> GetAssetCategories() const override;
+	//~ End UAssetDefinition End
+};
diff --git a/Engine/Plugins/Experimental/DynamicMaterial/Source/DynamicMaterialTextureSetEditor/Private/DMTextureSetBlueprintFunctionLibrary.cpp b/Engine/Plugins/Experimental/DynamicMaterial/Source/DynamicMaterialTextureSetEditor/Private/DMTextureSetBlueprintFunctionLibrary.cpp
new file mode 100644
--- /dev/null
+++ b/Engine/Plugins/Experimental/DynamicMaterial/Source/DynamicMaterialTextureSetEditor/Private/DMTextureSetBlueprintFunctionLibrary.cpp
@@ -0,0 +1,113 @@
+// Copyright Epic Games, Inc. All Rights Reserved.
+
+#include "DMTextureSetBlueprintFunctionLibrary.h"
+
+#include "AssetRegistry/AssetData.h"
+#include "CoreGlobals.h"
+#include "DMTextureSet.h"
+#include "DMTextureSetFactory.h"
+#include "DMTextureSetSettings.h"
+#include "Engine/Texture.h"
+#include "Framework/Application/SlateApplication.h"
+#include "Widgets/SDMTextureSetBuilder.h"
+#include "Widgets/SWindow.h"
+
+#define LOCTEXT_NAMESPACE "UDMTextureSetBlueprintFunctionLibrary"
+
+UDMTextureSet* UDMTextureSetBlueprintFunctionLibrary::CreateTextureSetFromAssets(const TArray<FAssetData>& InAssets)
+{
+	if (InAssets.IsEmpty())
+	{
+		return nullptr;
+	}
+
+	const UDMTextureSetSettings* TextureSetSettings = GetDefault<UDMTextureSetSettings>();
+
+	if (!TextureSetSettings || TextureSetSettings->Filters.IsEmpty())
+	{
+		return nullptr;
+	}
+
+	UDMTextureSet* TextureSet = Cast<UDMTextureSet>(GetMutableDefault<UDMTextureSetFactory>()->FactoryCreateNew(
+		UDMTextureSet::StaticClass(),
+		GetTransientPackage(),
+		NAME_None,
+		RF_Transactional,
+		/* Context */ nullptr,
+		GWarn
+	));
+
+	if (!TextureSet)
+	{
+		return nullptr;
+	}
+
+	for (const FAssetData& Asset : InAssets)
+	{
+		UClass* AssetClass = Asset.GetClass(EResolveClass::Yes);
+
+		if (!AssetClass || !AssetClass->IsChildOf(UTexture::StaticClass()))
+		{
+			continue;
+		}
+
+		UTexture* Texture = Cast<UTexture>(Asset.GetAsset());
+
+		if (!Texture)
+		{
+			continue;
+		}
+
+		const FString AssetName = Asset.AssetName.ToString();
+
+		for (const FDMTextureSetFilter& Filter : TextureSetSettings->Filters)
+		{
+			if (Filter.MaterialProperties.IsEmpty())
+			{
+				continue;
+			}
+
+			if (!Filter.MatchesFilter(AssetName))
+			{
+				continue;
+			}
+
+			for (const TPair<TEnumAsByte<EMaterialProperty>, EDMTextureChannelMask>& Pair : Filter.MaterialProperties)
+			{
+				if (TextureSet->HasMaterialTexture(Pair.Key))
+				{
+					continue;
+				}
+
+				TextureSet->SetMaterialTexture(Pair.Key, {Texture, Pair.Value});
+			}
+		}
+	}
+
+	return TextureSet;
+}
+
+void UDMTextureSetBlueprintFunctionLibrary::CreateTextureSetFromAssetsInteractive(const TArray<FAssetData>& InAssets, 
+	FDMTextureSetBuilderOnComplete InOnComplete)
+{
+	UDMTextureSet* TextureSet = CreateTextureSetFromAssets(InAssets);
+
+	if (!TextureSet)
+	{
+		InOnComplete.ExecuteIfBound(nullptr, /* Was Accepted */ false);
+		return;
+	}
+
+	TSharedRef<SWindow> Window = SNew(SWindow)
+		.SizingRule(ESizingRule::Autosized)
+		.SupportsMinimize(false)
+		.SupportsMaximize(false)
+		.Title(LOCTEXT("TextureSetBuilder", "Texture Set Builder"))
+		[
+			SNew(SDMTextureSetBuilder, TextureSet, InAssets, InOnComplete)
+		];
+
+	FSlateApplication::Get().AddWindow(Window, /* Show Immediately */ true);
+}
+
+#undef LOCTEXT_NAMESPACE
diff --git a/Engine/Plugins/Experimental/DynamicMaterial/Source/DynamicMaterialTextureSetEditor/Private/DMTextureSetContentBrowserIntegration.cpp b/Engine/Plugins/Experimental/DynamicMaterial/Source/DynamicMaterialTextureSetEditor/Private/DMTextureSetContentBrowserIntegration.cpp
new file mode 100644
--- /dev/null
+++ b/Engine/Plugins/Experimental/DynamicMaterial/Source/DynamicMaterialTextureSetEditor/Private/DMTextureSetContentBrowserIntegration.cpp
@@ -0,0 +1,312 @@
+// Copyright Epic Games, Inc. All Rights Reserved.
+
+#include "DMTextureSetContentBrowserIntegration.h"
+
+#include "AssetRegistry/AssetData.h"
+#include "AssetRegistry/AssetRegistryModule.h"
+#include "AssetToolsModule.h"
+#include "ContentBrowserModule.h"
+#include "CoreGlobals.h"
+#include "DMTextureSet.h"
+#include "DMTextureSetBlueprintFunctionLibrary.h"
+#include "DMTextureSetStyle.h"
+#include "Engine/Texture.h"
+#include "Framework/MultiBox/MultiBoxBuilder.h"
+#include "IDynamicMaterialEditorModule.h"
+#include "Material/DynamicMaterialInstance.h"
+#include "Material/DynamicMaterialInstanceFactory.h"
+#include "Materials/Material.h"
+#include "Model/DynamicMaterialModelEditorOnlyData.h"
+#include "ScopedTransaction.h"
+#include "Styling/SlateIconFinder.h"
+
+#define LOCTEXT_NAMESPACE "FDMTextureSetContentBrowserIntegration"
+
+FDelegateHandle FDMTextureSetContentBrowserIntegration::ExtenderDelegateHandle;
+
+void FDMTextureSetContentBrowserIntegration::Integrate()
+{
+	FDMTextureSetStyle::Get();
+
+	FContentBrowserModule& ContentBrowserModule = FModuleManager::LoadModuleChecked<FContentBrowserModule>("ContentBrowser");
+	TArray<FContentBrowserMenuExtender_SelectedAssets>& CBMenuExtenderDelegates = ContentBrowserModule.GetAllAssetViewContextMenuExtenders();
+	CBMenuExtenderDelegates.Add(FContentBrowserMenuExtender_SelectedAssets::CreateStatic(&FDMTextureSetContentBrowserIntegration::OnExtendContentBrowserAssetSelectionMenu));
+	ExtenderDelegateHandle = CBMenuExtenderDelegates.Last().GetHandle();
+}
+
+void FDMTextureSetContentBrowserIntegration::Disintegrate()
+{
+	if (!ExtenderDelegateHandle.IsValid())
+	{
+		return;
+	}
+
+	if (FContentBrowserModule* ContentBrowserModule = FModuleManager::GetModulePtr<FContentBrowserModule>("ContentBrowser"))
+	{
+		TArray<FContentBrowserMenuExtender_SelectedAssets>& CBMenuExtenderDelegates = ContentBrowserModule->GetAllAssetViewContextMenuExtenders();
+
+		CBMenuExtenderDelegates.RemoveAll(
+			[](const FContentBrowserMenuExtender_SelectedAssets& InElement)
+			{
+				return InElement.GetHandle() == ExtenderDelegateHandle;
+			}
+		);
+	}
+
+	ExtenderDelegateHandle.Reset();
+}
+
+TSharedRef<FExtender> FDMTextureSetContentBrowserIntegration::OnExtendContentBrowserAssetSelectionMenu(const TArray<FAssetData>& InSelectedAssets)
+{
+	TSharedRef<FExtender> Extender = MakeShared<FExtender>();
+	bool bHasTexture = false;
+
+	for (const FAssetData& SelectedAsset : InSelectedAssets)
+	{
+		if (UClass* AssetClass = SelectedAsset.GetClass(EResolveClass::Yes))
+		{
+			if (AssetClass->IsChildOf<UTexture>())
+			{
+				bHasTexture = true;
+				break;
+			}
+		}
+	}
+
+	if (!bHasTexture)
+	{
+		return Extender;
+	}
+
+	Extender->AddMenuExtension(
+		"GetAssetActions",
+		EExtensionHook::After,
+		nullptr,
+		FMenuExtensionDelegate::CreateLambda(
+			[InSelectedAssets](FMenuBuilder& InMenuBuilder) {
+				InMenuBuilder.AddMenuEntry(
+					LOCTEXT("CreateTextureSet", "Create Texture Set"),
+					LOCTEXT("CreateTextureSetTooltip", "Creates an asset listing a group of textures and the material properties they are associated with."),
+					FSlateIconFinder::FindIconForClass(UTexture::StaticClass()),
+					FUIAction(FExecuteAction::CreateStatic(&FDMTextureSetContentBrowserIntegration::CreateTextureSet, InSelectedAssets))
+				);
+			})
+	);
+
+#if 0
+	Extender->AddMenuExtension(
+		"GetAssetActions",
+		EExtensionHook::After,
+		nullptr,
+		FMenuExtensionDelegate::CreateLambda(
+			[InSelectedAssets](FMenuBuilder& InMenuBuilder) {
+				InMenuBuilder.AddMenuEntry(
+					LOCTEXT("CreateMaterialDesignerInstanceFromTextureSet", "Create Material Designer Instance"),
+					LOCTEXT("CreateMaterialDesignerInstanceFromTextureSetTooltip", "Creates a Material Designer Instance in the content browser using a Texture Set."),
+					FSlateIconFinder::FindIconForClass(UMaterial::StaticClass()),
+					FUIAction(FExecuteAction::CreateStatic(&FDMTextureSetContentBrowserIntegration::CreateMaterialDesignerInstanceFromTextureSet, InSelectedAssets))
+				);
+			})
+	);
+
+	Extender->AddMenuExtension(
+		"GetAssetActions",
+		EExtensionHook::After,
+		nullptr,
+		FMenuExtensionDelegate::CreateLambda(
+			[InSelectedAssets](FMenuBuilder& InMenuBuilder) {
+				InMenuBuilder.AddMenuEntry(
+					LOCTEXT("UpdateMaterialDesignerInstanceFromTextureSetAdd", "Update Material Designer Instance (Add)"),
+					LOCTEXT("UpdateMaterialDesignerInstanceFromTextureAddSetTooltip", "Updates the opened Material Designer Instance using a Texture Set, adding new layers to the Model."),
+					FSlateIconFinder::FindIconForClass(UMaterial::StaticClass()),
+					FUIAction(FExecuteAction::CreateStatic(&FDMTextureSetContentBrowserIntegration::UpdateMaterialDesignerInstanceFromTextureSet, InSelectedAssets, /* Replace */ false))
+				);
+			})
+	);
+
+	Extender->AddMenuExtension(
+		"GetAssetActions",
+		EExtensionHook::After,
+		nullptr,
+		FMenuExtensionDelegate::CreateLambda(
+			[InSelectedAssets](FMenuBuilder& InMenuBuilder) {
+				InMenuBuilder.AddMenuEntry(
+					LOCTEXT("UpdateMaterialDesignerInstanceFromTextureSetReplace", "Update Material Designer Instance (Replace)"),
+					LOCTEXT("UpdateMaterialDesignerInstanceFromTextureSetReplaceTooltip", "Updates the opened Material Designer Instance using a Texture Set, replacing slots in the Model."),
+					FSlateIconFinder::FindIconForClass(UMaterial::StaticClass()),
+					FUIAction(FExecuteAction::CreateStatic(&FDMTextureSetContentBrowserIntegration::UpdateMaterialDesignerInstanceFromTextureSet, InSelectedAssets, /* Replace */ true))
+				);
+			})
+	);
+#endif
+
+	return Extender;
+}
+
+void FDMTextureSetContentBrowserIntegration::CreateTextureSet(TArray<FAssetData> InSelectedAssets)
+{
+	if (InSelectedAssets.IsEmpty())
+	{
+		return;
+	}
+
+	UDMTextureSetBlueprintFunctionLibrary::CreateTextureSetFromAssetsInteractive(
+		InSelectedAssets,
+		FDMTextureSetBuilderOnComplete::CreateStatic(
+			&FDMTextureSetContentBrowserIntegration::OnCreateTextureSetComplete,
+			InSelectedAssets[0].PackagePath.ToString()
+		)
+	);
+}
+
+void FDMTextureSetContentBrowserIntegration::OnCreateTextureSetComplete(UDMTextureSet* InTextureSet, bool bInAccepted, FString InPath)
+{
+	if (!InTextureSet || !bInAccepted)
+	{
+		return;
+	}
+
+	const FAssetToolsModule& AssetToolsModule = FModuleManager::LoadModuleChecked<FAssetToolsModule>("AssetTools");
+
+	FString UniquePackageName;
+	FString UniqueAssetName;
+
+	const FString BasePackageName = InPath / TEXT("NewTextureSet");
+	AssetToolsModule.Get().CreateUniqueAssetName(BasePackageName, TEXT(""), UniquePackageName, UniqueAssetName);
+
+	UPackage* Package = CreatePackage(*UniquePackageName);
+
+	if (!Package)
+	{
+		return;
+	}
+
+	InTextureSet->SetFlags(RF_Standalone | RF_Public);
+	InTextureSet->Rename(*UniqueAssetName, Package, REN_DontCreateRedirectors);
+
+	FAssetRegistryModule::AssetCreated(InTextureSet);
+}
+
+void FDMTextureSetContentBrowserIntegration::CreateMaterialDesignerInstanceFromTextureSet(TArray<FAssetData> InSelectedAssets)
+{
+	if (InSelectedAssets.IsEmpty())
+	{
+		return;
+	}
+
+	UDMTextureSetBlueprintFunctionLibrary::CreateTextureSetFromAssetsInteractive(
+		InSelectedAssets,
+		FDMTextureSetBuilderOnComplete::CreateStatic(
+			&FDMTextureSetContentBrowserIntegration::OnCreateMaterialDesignerInstanceFromTextureSetComplete,
+			InSelectedAssets[0].PackagePath.ToString()
+		)
+	);
+}
+
+void FDMTextureSetContentBrowserIntegration::OnCreateMaterialDesignerInstanceFromTextureSetComplete(UDMTextureSet* InTextureSet, bool bInAccepted,
+	FString InPath)
+{
+	if (!InTextureSet || !bInAccepted)
+	{
+		return;
+	}
+
+	UDynamicMaterialInstance* Instance = Cast<UDynamicMaterialInstance>(GetMutableDefault<UDynamicMaterialInstanceFactory>()->FactoryCreateNew(
+		UDynamicMaterialInstance::StaticClass(),
+		GetTransientPackage(),
+		NAME_None,
+		RF_Transient,
+		/* Context */ nullptr,
+		GWarn
+	));
+
+	if (!Instance)
+	{
+		return;
+	}
+
+	UDynamicMaterialModelEditorOnlyData* EditorOnlyData = UDynamicMaterialModelEditorOnlyData::Get(Instance);
+
+	if (!EditorOnlyData)
+	{
+		return;
+	}
+
+	if (!EditorOnlyData->AddTextureSet(InTextureSet, /* Replace */ true))
+	{
+		return;
+	}
+
+	const FAssetToolsModule& AssetToolsModule = FModuleManager::LoadModuleChecked<FAssetToolsModule>("AssetTools");
+
+	FString UniquePackageName;
+	FString UniqueAssetName;
+
+	const FString BasePackageName = InPath / TEXT("NewMaterialDesignerInstance");
+	AssetToolsModule.Get().CreateUniqueAssetName(BasePackageName, TEXT(""), UniquePackageName, UniqueAssetName);
+
+	UPackage* Package = CreatePackage(*UniquePackageName);
+
+	if (!Package)
+	{
+		return;
+	}
+
+	Instance->SetFlags(RF_Standalone | RF_Public);
+	Instance->Rename(*UniqueAssetName, Package, REN_DontCreateRedirectors);
+
+	FAssetRegistryModule::AssetCreated(Instance);
+}
+
+void FDMTextureSetContentBrowserIntegration::UpdateMaterialDesignerInstanceFromTextureSet(TArray<FAssetData> InSelectedAssets, bool bInReplace)
+{
+	if (InSelectedAssets.IsEmpty())
+	{
+		return;
+	}
+
+	UDMTextureSetBlueprintFunctionLibrary::CreateTextureSetFromAssetsInteractive(
+		InSelectedAssets,
+		FDMTextureSetBuilderOnComplete::CreateStatic(
+			&FDMTextureSetContentBrowserIntegration::OnUpdateMaterialDesignerInstanceFromTextureSetComplete,
+			bInReplace
+		)
+	);
+}
+
+void FDMTextureSetContentBrowserIntegration::OnUpdateMaterialDesignerInstanceFromTextureSetComplete(UDMTextureSet* InTextureSet, bool bInAccepted,
+	bool bInReplace)
+{
+	if (!InTextureSet || !bInAccepted)
+	{
+		return;
+	}
+
+	IDynamicMaterialEditorModule& DynamicMaterialEditorModule = IDynamicMaterialEditorModule::Get();
+
+	UDynamicMaterialModel* Model = DynamicMaterialEditorModule.GetOpenedMaterialModel(nullptr);
+
+	if (!Model)
+	{
+		return;
+	}
+
+	UDynamicMaterialModelEditorOnlyData* EditorOnlyData = UDynamicMaterialModelEditorOnlyData::Get(Model);
+
+	if (!EditorOnlyData)
+	{
+		return;
+	}
+
+	FScopedTransaction Transaction(LOCTEXT("AddTextureSet", "Material Designer Add Texture Set"));
+	EditorOnlyData->Modify();
+
+	const bool bSuccess = EditorOnlyData->AddTextureSet(InTextureSet, bInReplace);
+
+	if (!bSuccess)
+	{
+		Transaction.Cancel();
+	}
+}
+
+#undef LOCTEXT_NAMESPACE
diff --git a/Engine/Plugins/Experimental/DynamicMaterial/Source/DynamicMaterialTextureSetEditor/Private/DMTextureSetContentBrowserIntegration.h b/Engine/Plugins/Experimental/DynamicMaterial/Source/DynamicMaterialTextureSetEditor/Private/DMTextureSetContentBrowserIntegration.h
new file mode 100644
--- /dev/null
+++ b/Engine/Plugins/Experimental/DynamicMaterial/Source/DynamicMaterialTextureSetEditor/Private/DMTextureSetContentBrowserIntegration.h
@@ -0,0 +1,37 @@
+// Copyright Epic Games, Inc. All Rights Reserved.
+
+#pragma once
+
+#include "Containers/ContainersFwd.h"
+#include "Delegates/Delegate.h"
+#include "Delegates/DelegateCombinations.h"
+#include "Templates/SharedPointerFwd.h"
+
+class FExtender;
+class UDMTextureSet;
+struct FAssetData;
+
+class FDMTextureSetContentBrowserIntegration
+{
+public:
+	static void Integrate();
+
+	static void Disintegrate();
+
+protected:
+	static TSharedRef<FExtender> OnExtendContentBrowserAssetSelectionMenu(const TArray<FAssetData>& InSelectedAssets);
+
+	static void CreateTextureSet(TArray<FAssetData> InSelectedAssets);
+
+	static void OnCreateTextureSetComplete(UDMTextureSet* InTextureSet, bool bInAccepted, FString InPath);
+
+	static void CreateMaterialDesignerInstanceFromTextureSet(TArray<FAssetData> InSelectedAssets);
+
+	static void OnCreateMaterialDesignerInstanceFromTextureSetComplete(UDMTextureSet* InTextureSet, bool bInAccepted, FString InPath);
+
+	static void UpdateMaterialDesignerInstanceFromTextureSet(TArray<FAssetData> InSelectedAssets, bool bInReplace);
+
+	static void OnUpdateMaterialDesignerInstanceFromTextureSetComplete(UDMTextureSet* InTextureSet, bool bInAccepted, bool bInReplace);
+
+	static FDelegateHandle ExtenderDelegateHandle;
+};
diff --git a/Engine/Plugins/Experimental/DynamicMaterial/Source/DynamicMaterialTextureSetEditor/Private/DMTextureSetFactory.cpp b/Engine/Plugins/Experimental/DynamicMaterial/Source/DynamicMaterialTextureSetEditor/Private/DMTextureSetFactory.cpp
new file mode 100644
--- /dev/null
+++ b/Engine/Plugins/Experimental/DynamicMaterial/Source/DynamicMaterialTextureSetEditor/Private/DMTextureSetFactory.cpp
@@ -0,0 +1,41 @@
+// Copyright Epic Games, Inc. All Rights Reserved.
+
+#include "DMTextureSetFactory.h"
+
+#include "DMTextureSet.h"
+#include "Modules/ModuleManager.h"
+
+#define LOCTEXT_NAMESPACE "MaterialDesignerInstanceFactory"
+
+UDMTextureSetFactory::UDMTextureSetFactory()
+{
+	SupportedClass = UDMTextureSet::StaticClass();
+
+	bCreateNew = true;
+	bEditAfterNew = true;
+	bEditorImport = false;
+	bText = false;
+}
+
+UObject* UDMTextureSetFactory::FactoryCreateNew(UClass* Class, UObject* InParent, FName Name,
+	EObjectFlags Flags, UObject* Context, FFeedbackContext* Warn)
+{
+	check(Class->IsChildOf(UDMTextureSet::StaticClass()));
+
+	UDMTextureSet* NewInstance = NewObject<UDMTextureSet>(InParent, Class, Name, Flags | RF_Transactional);
+	check(NewInstance);
+
+	return NewInstance;
+}
+
+FText UDMTextureSetFactory::GetDisplayName() const
+{
+	return LOCTEXT("UDMTextureSet", "Texture Set");
+}
+
+FText UDMTextureSetFactory::GetToolTip() const
+{
+	return LOCTEXT("UDMTextureSetTooltip", "A set of textures which are associated with a material property.");
+}
+
+#undef LOCTEXT_NAMESPACE
diff --git a/Engine/Plugins/Experimental/DynamicMaterial/Source/DynamicMaterialTextureSetEditor/Private/DMTextureSetFilter.cpp b/Engine/Plugins/Experimental/DynamicMaterial/Source/DynamicMaterialTextureSetEditor/Private/DMTextureSetFilter.cpp
new file mode 100644
--- /dev/null
+++ b/Engine/Plugins/Experimental/DynamicMaterial/Source/DynamicMaterialTextureSetEditor/Private/DMTextureSetFilter.cpp
@@ -0,0 +1,29 @@
+// Copyright Epic Games, Inc. All Rights Reserved.
+
+#include "DMTextureSetFilter.h"
+
+FDMTextureSetFilter::FDMTextureSetFilter()
+	: FilterStrings({TEXT("_")})
+	, MaterialProperties({{EMaterialProperty::MP_BaseColor, EDMTextureChannelMask::RGBA}})
+{
+}
+
+bool FDMTextureSetFilter::MatchesFilter(const FString& InAssetName) const
+{
+	for (const FString& FilterString : FilterStrings)
+	{
+		if (FilterString.StartsWith(TEXT("_")))
+		{
+			if (InAssetName.EndsWith(FilterString))
+			{
+				return true;
+			}
+		}
+		else if (InAssetName.Find(FilterString) != INDEX_NONE)
+		{
+			return true;
+		}
+	}
+
+	return false;
+}
diff --git a/Engine/Plugins/Experimental/DynamicMaterial/Source/DynamicMaterialTextureSetEditor/Private/DMTextureSetFilter.h b/Engine/Plugins/Experimental/DynamicMaterial/Source/DynamicMaterialTextureSetEditor/Private/DMTextureSetFilter.h
new file mode 100644
--- /dev/null
+++ b/Engine/Plugins/Experimental/DynamicMaterial/Source/DynamicMaterialTextureSetEditor/Private/DMTextureSetFilter.h
@@ -0,0 +1,38 @@
+// Copyright Epic Games, Inc. All Rights Reserved.
+
+#pragma once
+
+#include "Containers/Map.h"
+#include "DMTextureChannelMask.h"
+#include "SceneTypes.h"
+
+#include "DMTextureSetFilter.generated.h"
+
+USTRUCT()
+struct FDMTextureSetFilter
+{
+	GENERATED_BODY()
+
+	FDMTextureSetFilter();
+
+	/**
+	 * Portion of the name of the Texture asset to search for. For instance _Normal or _ORB.
+	 * Will match any of the given strings.
+	 * Filters starting with _ will only match if they are at the end of an asset name.
+	 */
+	UPROPERTY(EditAnywhere, Category = "Material Designer")
+	TArray<FString> FilterStrings;
+
+	/**
+	 * Where the matched Texture assets should be placed into the Texture Set. Links to the channel for the given asset.
+	 */
+	UPROPERTY(EditAnywhere, Category = "Material Designer")
+	TMap<TEnumAsByte<EMaterialProperty>, EDMTextureChannelMask> MaterialProperties;
+
+	/**
+	 * Checks the given asset name against the filter strings.
+	 * @param InAssetName The name of the asset
+	 * @return True if any filter string was matched.
+	 */
+	bool MatchesFilter(const FString& InAssetName) const;
+};
diff --git a/Engine/Plugins/Experimental/DynamicMaterial/Source/DynamicMaterialTextureSetEditor/Private/DMTextureSetSettings.cpp b/Engine/Plugins/Experimental/DynamicMaterial/Source/DynamicMaterialTextureSetEditor/Private/DMTextureSetSettings.cpp
new file mode 100644
--- /dev/null
+++ b/Engine/Plugins/Experimental/DynamicMaterial/Source/DynamicMaterialTextureSetEditor/Private/DMTextureSetSettings.cpp
@@ -0,0 +1,63 @@
+// Copyright Epic Games, Inc. All Rights Reserved.
+
+#include "DMTextureSetSettings.h"
+
+UDMTextureSetSettings::UDMTextureSetSettings()
+{
+	FDMTextureSetFilter BaseColor;
+	BaseColor.FilterStrings = {TEXT("Base_Color"), TEXT("BaseColor"), TEXT("Base_Colour"), TEXT("BaseColour"), TEXT("_BC"), TEXT("Diffuse"), TEXT("Albedo")};
+	BaseColor.MaterialProperties = {{EMaterialProperty::MP_BaseColor, EDMTextureChannelMask::RGBA}};
+
+	FDMTextureSetFilter Roughness;
+	Roughness.FilterStrings = {TEXT("Roughness"), TEXT("Rough"), TEXT("_R")};
+	Roughness.MaterialProperties = {{EMaterialProperty::MP_Roughness, EDMTextureChannelMask::RGBA}};
+
+	FDMTextureSetFilter Normal;
+	Normal.FilterStrings = {TEXT("Normal"), TEXT("Norm"), TEXT("_N")};
+	Normal.MaterialProperties = {{EMaterialProperty::MP_Normal, EDMTextureChannelMask::RGBA}};
+
+	FDMTextureSetFilter Metallic;
+	Metallic.FilterStrings = {TEXT("Metallic"), TEXT("Metal"), TEXT("_M")};
+	Metallic.MaterialProperties = {{EMaterialProperty::MP_Metallic, EDMTextureChannelMask::RGBA}};
+
+	FDMTextureSetFilter AmbientOcclusion;
+	AmbientOcclusion.FilterStrings = {TEXT("AmbientOcclusion"), TEXT("Ambient_Occlusion"), TEXT("_AO")};
+	AmbientOcclusion.MaterialProperties = {{EMaterialProperty::MP_AmbientOcclusion, EDMTextureChannelMask::RGBA}};
+
+	FDMTextureSetFilter Specular;
+	Specular.FilterStrings = {TEXT("Specular"), TEXT("_S")};
+	Specular.MaterialProperties = {{EMaterialProperty::MP_EmissiveColor, EDMTextureChannelMask::RGBA}};
+
+	FDMTextureSetFilter Emissive;
+	Emissive.FilterStrings = {TEXT("Emissive"), TEXT("Emission"), TEXT("_E")};
+	Emissive.MaterialProperties = {{EMaterialProperty::MP_EmissiveColor, EDMTextureChannelMask::RGBA}};
+
+	FDMTextureSetFilter Opacity;
+	Opacity.FilterStrings = {TEXT("Opacity"), TEXT("_O"), TEXT("Alpha"), TEXT("_A")};
+	Opacity.MaterialProperties = {{EMaterialProperty::MP_EmissiveColor, EDMTextureChannelMask::RGBA}};
+
+	FDMTextureSetFilter ORM;
+	ORM.FilterStrings = {TEXT("_ORM")};
+	ORM.MaterialProperties = {
+		{EMaterialProperty::MP_Opacity, EDMTextureChannelMask::Red},
+		{EMaterialProperty::MP_Roughness, EDMTextureChannelMask::Green},
+		{EMaterialProperty::MP_Metallic, EDMTextureChannelMask::Blue}
+	};
+
+	Filters.Reserve(9);
+
+	Filters.Add(MoveTemp(BaseColor));
+	Filters.Add(MoveTemp(Roughness));
+	Filters.Add(MoveTemp(Normal));
+	Filters.Add(MoveTemp(Metallic));
+	Filters.Add(MoveTemp(AmbientOcclusion));
+	Filters.Add(MoveTemp(Specular));
+	Filters.Add(MoveTemp(Emissive));
+	Filters.Add(MoveTemp(Opacity));
+	Filters.Add(MoveTemp(ORM));
+}
+
+UDMTextureSetSettings* UDMTextureSetSettings::Get()
+{
+	return GetMutableDefault<UDMTextureSetSettings>();
+}
diff --git a/Engine/Plugins/Experimental/DynamicMaterial/Source/DynamicMaterialTextureSetEditor/Private/DMTextureSetSettings.h b/Engine/Plugins/Experimental/DynamicMaterial/Source/DynamicMaterialTextureSetEditor/Private/DMTextureSetSettings.h
new file mode 100644
--- /dev/null
+++ b/Engine/Plugins/Experimental/DynamicMaterial/Source/DynamicMaterialTextureSetEditor/Private/DMTextureSetSettings.h
@@ -0,0 +1,27 @@
+// Copyright Epic Games, Inc. All Rights Reserved.
+
+#pragma once
+
+#include "Engine/DeveloperSettings.h"
+
+#include "DMTextureSetFilter.h"
+
+#include "DMTextureSetSettings.generated.h"
+
+/**
+ * DM Texture Set Settings
+ */
+UCLASS(Config = EditorPerProjectUserSettings, meta = (DisplayName = "Material Designer"))
+class UDMTextureSetSettings : public UDeveloperSettings
+{
+	GENERATED_BODY()
+
+public:
+	UDMTextureSetSettings();
+	virtual ~UDMTextureSetSettings() override = default;
+
+	static UDMTextureSetSettings* Get();
+
+	UPROPERTY(EditAnywhere, Category = "Material Designer", meta = (TitleProperty=FilterStrings))
+	TArray<FDMTextureSetFilter> Filters;
+};
diff --git a/Engine/Plugins/Experimental/DynamicMaterial/Source/DynamicMaterialTextureSetEditor/Private/DMTextureSetStyle.cpp b/Engine/Plugins/Experimental/DynamicMaterial/Source/DynamicMaterialTextureSetEditor/Private/DMTextureSetStyle.cpp
new file mode 100644
--- /dev/null
+++ b/Engine/Plugins/Experimental/DynamicMaterial/Source/DynamicMaterialTextureSetEditor/Private/DMTextureSetStyle.cpp
@@ -0,0 +1,20 @@
+// Copyright Epic Games, Inc. All Rights Reserved.
+
+#include "DMTextureSetStyle.h"
+
+#include "Brushes/SlateImageBrush.h"
+#include "Styling/StyleColors.h"
+
+const ISlateStyle& FDMTextureSetStyle::Get()
+{
+	static FDMTextureSetStyle Instance;
+	return Instance;
+}
+
+FDMTextureSetStyle::FDMTextureSetStyle()
+	: FSlateStyleSet(TEXT("DMTextureSetStyle"))
+{
+	Set("TextureSetConfig.Background", new FSlateRoundedBoxBrush(
+		FStyleColors::Recessed.GetSpecifiedColor(), 6.0f,
+		FStyleColors::Header.GetSpecifiedColor(), 2.0f));
+}
diff --git a/Engine/Plugins/Experimental/DynamicMaterial/Source/DynamicMaterialTextureSetEditor/Private/DMTextureSetStyle.h b/Engine/Plugins/Experimental/DynamicMaterial/Source/DynamicMaterialTextureSetEditor/Private/DMTextureSetStyle.h
new file mode 100644
--- /dev/null
+++ b/Engine/Plugins/Experimental/DynamicMaterial/Source/DynamicMaterialTextureSetEditor/Private/DMTextureSetStyle.h
@@ -0,0 +1,16 @@
+// Copyright Epic Games, Inc. All Rights Reserved.
+
+#pragma once
+
+#include "Styling/ISlateStyle.h"
+#include "Styling/SlateStyle.h"
+
+class FDMTextureSetStyle final : public FSlateStyleSet
+{
+public:
+	static const ISlateStyle& Get();
+
+	FDMTextureSetStyle();
+
+	virtual ~FDMTextureSetStyle() override = default;
+};
diff --git a/Engine/Plugins/Experimental/DynamicMaterial/Source/DynamicMaterialTextureSetEditor/Private/DynamicMaterialTextureSetEditorModule.cpp b/Engine/Plugins/Experimental/DynamicMaterial/Source/DynamicMaterialTextureSetEditor/Private/DynamicMaterialTextureSetEditorModule.cpp
new file mode 100644
--- /dev/null
+++ b/Engine/Plugins/Experimental/DynamicMaterial/Source/DynamicMaterialTextureSetEditor/Private/DynamicMaterialTextureSetEditorModule.cpp
@@ -0,0 +1,20 @@
+// Copyright Epic Games, Inc. All Rights Reserved.
+
+#include "DynamicMaterialTextureSetEditorModule.h"
+
+#include "DMTextureSetContentBrowserIntegration.h"
+#include "DMTextureSetStyle.h"
+#include "Modules/ModuleManager.h"
+
+void FDynamicMaterialTextureSetEditorModule::StartupModule()
+{
+	FDMTextureSetStyle::Get();
+	FDMTextureSetContentBrowserIntegration::Integrate();
+}
+
+void FDynamicMaterialTextureSetEditorModule::ShutdownModule()
+{
+	FDMTextureSetContentBrowserIntegration::Disintegrate();
+}
+
+IMPLEMENT_MODULE(FDynamicMaterialTextureSetEditorModule, DynamicMaterialTextureSetEditor)
diff --git a/Engine/Plugins/Experimental/DynamicMaterial/Source/DynamicMaterialTextureSetEditor/Private/DynamicMaterialTextureSetEditorModule.h b/Engine/Plugins/Experimental/DynamicMaterial/Source/DynamicMaterialTextureSetEditor/Private/DynamicMaterialTextureSetEditorModule.h
new file mode 100644
--- /dev/null
+++ b/Engine/Plugins/Experimental/DynamicMaterial/Source/DynamicMaterialTextureSetEditor/Private/DynamicMaterialTextureSetEditorModule.h
@@ -0,0 +1,14 @@
+// Copyright Epic Games, Inc. All Rights Reserved.
+
+#pragma once
+
+#include "Modules/ModuleInterface.h"
+
+class FDynamicMaterialTextureSetEditorModule : public IModuleInterface
+{
+public:
+	//~ Begin IModuleInterface
+	virtual void StartupModule() override;
+	virtual void ShutdownModule() override;
+	//~ End IModuleInterface
+};
diff --git a/Engine/Plugins/Experimental/DynamicMaterial/Source/DynamicMaterialTextureSetEditor/Private/Widgets/DMTextureSetBuilderDragDropOperation.cpp b/Engine/Plugins/Experimental/DynamicMaterial/Source/DynamicMaterialTextureSetEditor/Private/Widgets/DMTextureSetBuilderDragDropOperation.cpp
new file mode 100644
--- /dev/null
+++ b/Engine/Plugins/Experimental/DynamicMaterial/Source/DynamicMaterialTextureSetEditor/Private/Widgets/DMTextureSetBuilderDragDropOperation.cpp
@@ -0,0 +1,42 @@
+// Copyright Epic Games, Inc. All Rights Reserved.
+
+#include "Widgets/DMTextureSetBuilderDragDropOperation.h"
+
+#include "Framework/Application/SlateApplication.h"
+
+TSharedRef<FDMTextureSetBuilderDragDropOperation> FDMTextureSetBuilderDragDropOperation::New(const FAssetData& InAssetData, int32 InIndex, bool bInIsMaterialProperty)
+{
+	TSharedRef<FDMTextureSetBuilderDragDropOperation> Operation = MakeShared<FDMTextureSetBuilderDragDropOperation>();
+
+	Operation->Init({InAssetData}, {}, nullptr);
+	Operation->Index = InIndex;
+	Operation->bIsMaterialProperty = bInIsMaterialProperty;
+
+	Operation->Construct();
+	return Operation;
+}
+
+int32 FDMTextureSetBuilderDragDropOperation::GetIndex() const
+{
+	return Index;
+}
+
+bool FDMTextureSetBuilderDragDropOperation::IsMaterialProperty() const
+{
+	return bIsMaterialProperty;
+}
+
+FCursorReply FDMTextureSetBuilderDragDropOperation::OnCursorQuery()
+{
+	if (bIsMaterialProperty)
+	{
+		const bool bOverwriteTexture = FSlateApplication::Get().GetModifierKeys().IsShiftDown();
+
+		if (bOverwriteTexture)
+		{
+			return FCursorReply::Cursor(EMouseCursor::Hand);
+		}
+	}
+
+	return FAssetDragDropOp::OnCursorQuery();
+}
diff --git a/Engine/Plugins/Experimental/DynamicMaterial/Source/DynamicMaterialTextureSetEditor/Private/Widgets/DMTextureSetBuilderDragDropOperation.h b/Engine/Plugins/Experimental/DynamicMaterial/Source/DynamicMaterialTextureSetEditor/Private/Widgets/DMTextureSetBuilderDragDropOperation.h
new file mode 100644
--- /dev/null
+++ b/Engine/Plugins/Experimental/DynamicMaterial/Source/DynamicMaterialTextureSetEditor/Private/Widgets/DMTextureSetBuilderDragDropOperation.h
@@ -0,0 +1,27 @@
+// Copyright Epic Games, Inc. All Rights Reserved.
+
+#pragma once
+
+#include "DragAndDrop/AssetDragDropOp.h"
+
+class FDMTextureSetBuilderDragDropOperation : public FAssetDragDropOp
+{
+public:
+	DRAG_DROP_OPERATOR_TYPE(FDMTextureSetBuilderDragDropOperation, FAssetDragDropOp)
+
+	static TSharedRef<FDMTextureSetBuilderDragDropOperation> New(const FAssetData& InAssetData, int32 InIndex, bool bInIsMaterialProperty);
+
+	virtual ~FDMTextureSetBuilderDragDropOperation() override = default;
+
+	int32 GetIndex() const;
+
+	bool IsMaterialProperty() const;
+
+	//~ Begin FDragDropOperation
+	virtual FCursorReply OnCursorQuery();
+	//~ End FDragDropOperation
+
+protected:
+	int32 Index = -1;
+	bool bIsMaterialProperty = false;
+};
diff --git a/Engine/Plugins/Experimental/DynamicMaterial/Source/DynamicMaterialTextureSetEditor/Private/Widgets/DMTextureSetBuilderEntry.cpp b/Engine/Plugins/Experimental/DynamicMaterial/Source/DynamicMaterialTextureSetEditor/Private/Widgets/DMTextureSetBuilderEntry.cpp
new file mode 100644
--- /dev/null
+++ b/Engine/Plugins/Experimental/DynamicMaterial/Source/DynamicMaterialTextureSetEditor/Private/Widgets/DMTextureSetBuilderEntry.cpp
@@ -0,0 +1,33 @@
+// Copyright Epic Games, Inc. All Rights Reserved.
+
+#include "Widgets/DMTextureSetBuilderEntry.h"
+
+#include "InstancedStruct.h"
+
+FDMTextureSetBuilderEntryProvider::FDMTextureSetBuilderEntryProvider(const TSharedRef<FDMTextureSetBuilderEntry>& InEntry)
+	: Entry(InEntry)
+{
+}
+
+bool FDMTextureSetBuilderEntryProvider::IsValid() const
+{
+	return true;
+}
+
+const UStruct* FDMTextureSetBuilderEntryProvider::GetBaseStructure() const
+{
+	return FDMTextureSetBuilderEntry::StaticStruct();
+}
+
+void FDMTextureSetBuilderEntryProvider::GetInstances(TArray<TSharedPtr<FStructOnScope>>& OutInstances, const UStruct* ExpectedBaseStructure) const
+{
+	OutInstances.Add(MakeShared<FStructOnScope>(
+		FDMTextureSetBuilderEntry::StaticStruct(),
+		reinterpret_cast<uint8*>(&Entry.Get())
+	));
+}
+
+bool FDMTextureSetBuilderEntryProvider::IsPropertyIndirection() const
+{
+	return false;
+}
diff --git a/Engine/Plugins/Experimental/DynamicMaterial/Source/DynamicMaterialTextureSetEditor/Private/Widgets/DMTextureSetBuilderEntry.h b/Engine/Plugins/Experimental/DynamicMaterial/Source/DynamicMaterialTextureSetEditor/Private/Widgets/DMTextureSetBuilderEntry.h
new file mode 100644
--- /dev/null
+++ b/Engine/Plugins/Experimental/DynamicMaterial/Source/DynamicMaterialTextureSetEditor/Private/Widgets/DMTextureSetBuilderEntry.h
@@ -0,0 +1,47 @@
+// Copyright Epic Games, Inc. All Rights Reserved.
+
+#pragma once
+
+#include "IStructureDataProvider.h"
+
+#include "DMTextureChannelMask.h"
+#include "SceneTypes.h"
+#include "Templates/SharedPointer.h"
+#include "UObject/ObjectPtr.h"
+
+#include "DMTextureSetBuilderEntry.generated.h"
+
+class UTexture;
+
+USTRUCT()
+struct FDMTextureSetBuilderEntry
+{
+	GENERATED_BODY()
+
+	UPROPERTY(EditAnywhere, Category = "Material Designer")
+	TEnumAsByte<EMaterialProperty> MaterialProperty = EMaterialProperty::MP_BaseColor;
+
+	UPROPERTY(EditAnywhere, Category = "Material Designer")
+	TObjectPtr<UTexture> Texture;
+
+	UPROPERTY(EditAnywhere, Category = "Material Designer")
+	EDMTextureChannelMask ChannelMask = EDMTextureChannelMask::RGBA;
+};
+
+class FDMTextureSetBuilderEntryProvider : public IStructureDataProvider
+{
+public:
+	explicit FDMTextureSetBuilderEntryProvider(const TSharedRef<FDMTextureSetBuilderEntry>& InEntry);
+
+	virtual ~FDMTextureSetBuilderEntryProvider() override = default;
+
+	//~ Begin IStructureDataProvider
+	virtual bool IsValid() const override;
+	virtual const UStruct* GetBaseStructure() const override;
+	virtual void GetInstances(TArray<TSharedPtr<FStructOnScope>>& OutInstances, const UStruct* ExpectedBaseStructure) const override;
+	virtual bool IsPropertyIndirection() const override;
+	//~ End IStructureDataProvider
+
+protected:
+	TSharedRef<FDMTextureSetBuilderEntry> Entry;
+};
diff --git a/Engine/Plugins/Experimental/DynamicMaterial/Source/DynamicMaterialTextureSetEditor/Private/Widgets/SDMTextureSetBuilder.cpp b/Engine/Plugins/Experimental/DynamicMaterial/Source/DynamicMaterialTextureSetEditor/Private/Widgets/SDMTextureSetBuilder.cpp
new file mode 100644
--- /dev/null
+++ b/Engine/Plugins/Experimental/DynamicMaterial/Source/DynamicMaterialTextureSetEditor/Private/Widgets/SDMTextureSetBuilder.cpp
@@ -0,0 +1,409 @@
+// Copyright Epic Games, Inc. All Rights Reserved.
+
+#include "Widgets/SDMTextureSetBuilder.h"
+
+#include "DMTextureSet.h"
+#include "DMTextureSetStyle.h"
+#include "Engine/Texture.h"
+#include "SAssetDropTarget.h"
+#include "SceneTypes.h"
+#include "Widgets/DMTextureSetBuilderDragDropOperation.h"
+#include "Widgets/DMTextureSetBuilderEntry.h"
+#include "Widgets/Input/SButton.h"
+#include "Widgets/Layout/SBorder.h"
+#include "Widgets/Layout/SBox.h"
+#include "Widgets/Layout/SGridPanel.h"
+#include "Widgets/SBoxPanel.h"
+#include "Widgets/SDMTextureSetBuilderCellBase.h"
+#include "Widgets/SDMTextureSetBuilderMaterialPropertyCell.h"
+#include "Widgets/SDMTextureSetBuilderUnassignedTextureCell.h"
+
+#define LOCTEXT_NAMESPACE "SDMTextureSetBuilder"
+
+namespace UE::DynamicMaterialEditor::TextureSet::Private
+{
+	constexpr int32 WidthMax = 7;
+}
+
+void SDMTextureSetBuilder::Construct(const FArguments& InArgs, UDMTextureSet* InTextureSet, const TArray<FAssetData>& InAssets, 
+	FDMTextureSetBuilderOnComplete InOnComplete)
+{
+	TextureSet.Reset(InTextureSet);
+	Assets = InAssets;
+	OnComplete = InOnComplete;
+
+	if (!IsValid(InTextureSet))
+	{
+		OnComplete.ExecuteIfBound(nullptr, /* Was Accepted */ false);
+		return;
+	}
+
+	using namespace UE::DynamicMaterialEditor::TextureSet::Private;
+
+	MaterialPropertyGrid = SNew(SGridPanel);
+	int32 WidthIndex = 0;
+	int32 HeightIndex = 0;
+
+	MaterialProperties.Reserve(InTextureSet->GetTextures().Num());
+
+	for (const TPair<TEnumAsByte<EMaterialProperty>, FDMMaterialTexture>& SetElement : InTextureSet->GetTextures())
+	{
+		TSharedRef<FDMTextureSetBuilderEntry> Entry = MakeShared<FDMTextureSetBuilderEntry>(
+			SetElement.Key,
+			SetElement.Value.Texture.LoadSynchronous(),
+			EDMTextureChannelMask::RGBA
+		);
+
+		MaterialPropertyGrid->AddSlot(WidthIndex, HeightIndex)
+			.Padding(5.f)
+			[
+				GenerateMaterialPropertyCell(Entry, MaterialProperties.Num())
+			];
+
+		MaterialProperties.Add(Entry);
+
+		++WidthIndex;
+
+		if (WidthIndex == WidthMax)
+		{
+			WidthIndex = 0;
+			++HeightIndex;
+		}
+	}
+
+	UnassignedTextureGrid = SNew(SGridPanel);
+	
+	CreateUnassignedTextureSlots();
+
+	ChildSlot
+	[
+		SNew(SBox)
+		.Padding(5.0f)
+		[
+			SNew(SVerticalBox)
+
+			+ SVerticalBox::Slot()
+			.AutoHeight()
+			.Padding(5.f, 5.f, 5.f, 0.f)
+			[
+				MaterialPropertyGrid.ToSharedRef()
+			]
+
+			+ SVerticalBox::Slot()
+			.AutoHeight()
+			.Padding(5.f, 5.f, 5.f, 0.f)
+			[
+				SNew(SBorder)
+				.Padding(5.f)
+				.BorderImage(FDMTextureSetStyle::Get().GetBrush("TextureSetConfig.Background"))
+				[
+					SNew(SAssetDropTarget)
+					.OnAreAssetsAcceptableForDrop(this, &SDMTextureSetBuilder::OnAssetDraggedOver)
+					.OnAssetsDropped(this, &SDMTextureSetBuilder::OnAssetsDropped)
+					[
+						UnassignedTextureGrid.ToSharedRef()
+					]
+				]
+			]
+
+			+ SVerticalBox::Slot()
+			.AutoHeight()
+			.Padding(5.f, 5.f, 5.f, 5.f)
+			[
+				SNew(SHorizontalBox)
+
+				+ SHorizontalBox::Slot()
+				.FillWidth(1.f)
+				[
+					SNullWidget::NullWidget
+				]
+
+				+ SHorizontalBox::Slot()
+				.AutoWidth()
+				[
+					SNew(SBox)
+					.Padding(10.f, 5.f)
+					[
+						SNew(SButton)
+						.Text(LOCTEXT("Accept", "Accept"))
+						.ContentPadding(FMargin(5.f, 3.f))
+						.ButtonStyle(FAppStyle::Get(), "PrimaryButton")
+						.OnClicked(this, &SDMTextureSetBuilder::OnAcceptClicked)
+					]
+				]
+
+				+ SHorizontalBox::Slot()
+				.AutoWidth()
+				[
+					SNew(SBox)
+					.Padding(10.f, 5.f)
+					[
+						SNew(SButton)
+						.Text(LOCTEXT("Cancel", "Cancel"))
+						.ContentPadding(FMargin(5.f, 3.f))
+						.ButtonStyle(FAppStyle::Get(), "Button")
+						.OnClicked(this, &SDMTextureSetBuilder::OnCancelClicked)
+					]
+				]
+
+				+ SHorizontalBox::Slot()
+				.FillWidth(1.f)
+				[
+					SNullWidget::NullWidget
+				]
+			]
+		]
+	];
+}
+
+void SDMTextureSetBuilder::SwapTexture(int32 InFromIndex, bool bInIsFromMaterialProperty, int32 InToIndex, bool bInIsToMaterialProperty)
+{
+	if (InFromIndex == InToIndex && bInIsFromMaterialProperty == bInIsToMaterialProperty)
+	{
+		return;
+	}
+
+	const bool bSwapTextures = !FSlateApplication::Get().GetModifierKeys().IsShiftDown();
+
+	TSharedPtr<SDMTextureSetBuilderCellBase> FromCell = GetCell(InFromIndex, bInIsFromMaterialProperty);
+	TSharedPtr<SDMTextureSetBuilderCellBase> ToCell = GetCell(InToIndex, bInIsToMaterialProperty);
+
+	if (!FromCell.IsValid() || !ToCell.IsValid())
+	{
+		return;
+	}
+
+	UTexture* FromTexture = FromCell->GetTexture();
+
+	if (bSwapTextures)
+	{
+		UTexture* ToTexture = ToCell->GetTexture();
+
+		FromCell->SetTexture(ToTexture);
+
+		if (bInIsFromMaterialProperty)
+		{
+			MaterialProperties[InFromIndex]->Texture = ToTexture;
+		}
+	}
+
+	ToCell->SetTexture(FromTexture);
+
+	if (bInIsToMaterialProperty)
+	{
+		MaterialProperties[InToIndex]->Texture = FromTexture;
+	}
+
+	if (!bInIsFromMaterialProperty)
+	{
+		UnassignedTextures.Remove(FromTexture);
+	}
+
+	if (!bInIsToMaterialProperty)
+	{
+		if (UnassignedTextures.Contains(FromTexture))
+		{
+			UnassignedTextures.Add(FromTexture);
+		}
+	}
+
+	if (!bInIsFromMaterialProperty || !bInIsToMaterialProperty || !bSwapTextures)
+	{
+		CreateUnassignedTextureSlots();
+	}
+}
+
+TSharedRef<SWidget> SDMTextureSetBuilder::GenerateMaterialPropertyCell(const TSharedRef<FDMTextureSetBuilderEntry>& InListItem, int32 InIndex)
+{
+	return SNew(SDMTextureSetBuilderMaterialPropertyCell, SharedThis(this), InListItem, InIndex);
+}
+
+TSharedRef<SWidget> SDMTextureSetBuilder::GenerateUnusedTextureCell(UTexture* InTexture, int32 InIndex)
+{
+	return SNew(SDMTextureSetBuilderUnassignedTextureCell, SharedThis(this), InTexture, InIndex);
+}
+
+FReply SDMTextureSetBuilder::OnAcceptClicked()
+{
+	UDMTextureSet* TextureSetObject = TextureSet.Get();
+
+	if (!IsValid(TextureSetObject))
+	{
+		return FReply::Handled();
+	}
+
+	for (const TPair<TEnumAsByte<EMaterialProperty>, FDMMaterialTexture>& SetElement : TextureSet->GetTextures())
+	{
+		TextureSet->SetMaterialTexture(SetElement.Key, {nullptr, EDMTextureChannelMask::RGBA});
+	}
+
+	for (const TSharedRef<FDMTextureSetBuilderEntry>& Entry : MaterialProperties)
+	{
+		if (!TextureSet->HasMaterialProperty(Entry->MaterialProperty))
+		{
+			continue;
+		}
+
+		TextureSet->SetMaterialTexture(Entry->MaterialProperty, {Entry->Texture, Entry->ChannelMask});
+	}
+
+	bAccepted = true;
+
+	Close();
+
+	return FReply::Handled();
+}
+
+FReply SDMTextureSetBuilder::OnCancelClicked()
+{
+	Close();
+
+	return FReply::Handled();
+}
+
+void SDMTextureSetBuilder::Close()
+{
+	if (TSharedPtr<SWindow> CurrentWindow = FSlateApplication::Get().FindWidgetWindow(AsShared()))
+	{
+		CurrentWindow->RequestDestroyWindow();
+	}
+
+	OnComplete.ExecuteIfBound(TextureSet.Get(), /* Was Accepted */ bAccepted);
+}
+
+void SDMTextureSetBuilder::CreateUnassignedTextureSlots()
+{
+	using namespace UE::DynamicMaterialEditor::TextureSet::Private;
+
+	UnassignedTextures.Empty(TextureSet->GetTextures().Num());
+	UnassignedTextureGrid->ClearChildren();
+
+	int32 WidthIndex = 0;
+	int32 HeightIndex = 0;
+
+	for (const FAssetData& Asset : Assets)
+	{
+		UTexture* Texture = Cast<UTexture>(Asset.GetAsset());
+
+		if (!Texture || IsTextureAssignedToMaterialProperty(Texture))
+		{
+			continue;
+		}
+
+		UnassignedTextureGrid->AddSlot(WidthIndex, HeightIndex)
+			.Padding(5.f)
+			[
+				GenerateUnusedTextureCell(Texture, UnassignedTextures.Num())
+			];
+
+		UnassignedTextures.Add(Texture);
+
+		++WidthIndex;
+
+		if (WidthIndex == WidthMax)
+		{
+			WidthIndex = 0;
+			++HeightIndex;
+		}
+	}
+
+	if (UnassignedTextures.IsEmpty())
+	{
+		UnassignedTextureGrid->AddSlot(WidthIndex, HeightIndex)
+			.Padding(5.f)
+			[
+				GenerateUnusedTextureCell(nullptr, UnassignedTextures.Num())
+			];
+
+		UnassignedTextures.Add(nullptr);
+	}
+}
+
+bool SDMTextureSetBuilder::OnAssetDraggedOver(TArrayView<FAssetData> InAssets)
+{
+	for (const FAssetData& Asset : InAssets)
+	{
+		UClass* AssetClass = Asset.GetClass(EResolveClass::Yes);
+
+		if (AssetClass && AssetClass->IsChildOf(UTexture::StaticClass()))
+		{
+			return true;
+		}
+	}
+
+	return false;
+}
+
+void SDMTextureSetBuilder::OnAssetsDropped(const FDragDropEvent& InDragDropEvent, TArrayView<FAssetData> InAssets)
+{
+	if (TSharedPtr<FDMTextureSetBuilderDragDropOperation> BuilderOperation = InDragDropEvent.GetOperationAs<FDMTextureSetBuilderDragDropOperation>())
+	{
+		if (BuilderOperation->IsMaterialProperty())
+		{
+			if (TSharedPtr<SDMTextureSetBuilderCellBase> FromCell = GetCell(BuilderOperation->GetIndex(), /* Is Material Property */ true))
+			{
+				if (UTexture* Texture = FromCell->GetTexture())
+				{
+					FromCell->SetTexture(nullptr);
+
+					if (!UnassignedTextures.Contains(Texture))
+					{
+						CreateUnassignedTextureSlots();
+					}
+				}
+			}
+		}
+	}
+}
+
+TSharedPtr<SDMTextureSetBuilderCellBase> SDMTextureSetBuilder::GetCell(int32 InIndex, bool bInIsMaterialProperty) const
+{
+	FChildren* Children = bInIsMaterialProperty
+		? MaterialPropertyGrid->GetChildren()
+		: UnassignedTextureGrid->GetChildren();;
+
+	if (InIndex >= Children->Num())
+	{
+		return nullptr;
+	}
+
+	return StaticCastSharedRef<SDMTextureSetBuilderCellBase>(Children->GetChildAt(InIndex));
+}
+
+UTexture* SDMTextureSetBuilder::GetTexture(int32 InIndex, bool bInIsMaterialProperty) const
+{
+	if (TSharedPtr<SDMTextureSetBuilderCellBase> Cell = GetCell(InIndex, bInIsMaterialProperty))
+	{
+		return Cell->GetTexture();
+	}
+
+	return nullptr;
+}
+
+void SDMTextureSetBuilder::SetTexture(int32 InIndex, bool bInIsMaterialProperty, UTexture* InTexture) const
+{
+	if (TSharedPtr<SDMTextureSetBuilderCellBase> Cell = GetCell(InIndex, bInIsMaterialProperty))
+	{
+		Cell->SetTexture(InTexture);
+	}
+}
+
+bool SDMTextureSetBuilder::IsTextureAssignedToMaterialProperty(UTexture* InTexture) const
+{
+	FChildren* Children = MaterialPropertyGrid->GetChildren();
+	const int32 ChildCount = Children->Num();
+
+	for (int32 Index = 0; Index < ChildCount; ++Index)
+	{
+		TSharedRef<SDMTextureSetBuilderCellBase> Cell = StaticCastSharedRef<SDMTextureSetBuilderCellBase>(Children->GetChildAt(Index));
+
+		if (Cell->GetTexture() == InTexture)
+		{
+			return true;
+		}
+	}
+
+	return false;
+}
+
+#undef LOCTEXT_NAMESPACE
diff --git a/Engine/Plugins/Experimental/DynamicMaterial/Source/DynamicMaterialTextureSetEditor/Private/Widgets/SDMTextureSetBuilder.h b/Engine/Plugins/Experimental/DynamicMaterial/Source/DynamicMaterialTextureSetEditor/Private/Widgets/SDMTextureSetBuilder.h
new file mode 100644
--- /dev/null
+++ b/Engine/Plugins/Experimental/DynamicMaterial/Source/DynamicMaterialTextureSetEditor/Private/Widgets/SDMTextureSetBuilder.h
@@ -0,0 +1,65 @@
+// Copyright Epic Games, Inc. All Rights Reserved.
+
+#pragma once
+
+#include "DMTextureSetBlueprintFunctionLibrary.h"
+#include "Widgets/SCompoundWidget.h"
+
+#include "Containers/Array.h"
+#include "Templates/SharedPointer.h"
+#include "UObject/StrongObjectPtr.h"
+
+class SDMTextureSetBuilderCellBase;
+class SGridPanel;
+class UDMTextureSet;
+class UTexture;
+struct FAssetData;
+struct FDMTextureSetBuilderEntry;
+
+class SDMTextureSetBuilder : public SCompoundWidget
+{
+	SLATE_BEGIN_ARGS(SDMTextureSetBuilder) {}
+	SLATE_END_ARGS()
+
+public:
+	void Construct(const FArguments& InArgs, UDMTextureSet* InTextureSet, const TArray<FAssetData>& InAssets, FDMTextureSetBuilderOnComplete InOnComplete);
+
+	bool WasAccepted() const { return bAccepted; }
+
+	void SwapTexture(int32 InFromIndex, bool bInIsFromMaterialProperty, int32 InToIndex, bool bInIsToMaterialProperty);
+
+protected:
+	TStrongObjectPtr<UDMTextureSet> TextureSet;
+	TArray<FAssetData> Assets;
+	TArray<TSharedRef<FDMTextureSetBuilderEntry>> MaterialProperties;
+	TArray<UTexture*> UnassignedTextures;
+	bool bAccepted = false;
+	FDMTextureSetBuilderOnComplete OnComplete;
+
+	TSharedPtr<SGridPanel> MaterialPropertyGrid;
+	TSharedPtr<SGridPanel> UnassignedTextureGrid;
+
+	TSharedRef<SWidget> GenerateMaterialPropertyCell(const TSharedRef<FDMTextureSetBuilderEntry>& InListItem, int32 InIndex);
+
+	TSharedRef<SWidget> GenerateUnusedTextureCell(UTexture* InTexture, int32 InIndex);
+
+	FReply OnAcceptClicked();
+
+	FReply OnCancelClicked();
+
+	void Close();
+
+	void CreateUnassignedTextureSlots();
+
+	bool OnAssetDraggedOver(TArrayView<FAssetData> InAssets);
+
+	void OnAssetsDropped(const FDragDropEvent& InDragDropEvent, TArrayView<FAssetData> InAssets);
+
+	TSharedPtr<SDMTextureSetBuilderCellBase> GetCell(int32 InIndex, bool bInIsMaterialProperty) const;
+
+	UTexture* GetTexture(int32 InIndex, bool bInIsMaterialProperty) const;
+
+	void SetTexture(int32 InIndex, bool bInIsMaterialProperty, UTexture* InTexture) const;
+
+	bool IsTextureAssignedToMaterialProperty(UTexture* InTexture) const;
+};
diff --git a/Engine/Plugins/Experimental/DynamicMaterial/Source/DynamicMaterialTextureSetEditor/Private/Widgets/SDMTextureSetBuilderCellBase.cpp b/Engine/Plugins/Experimental/DynamicMaterial/Source/DynamicMaterialTextureSetEditor/Private/Widgets/SDMTextureSetBuilderCellBase.cpp
new file mode 100644
--- /dev/null
+++ b/Engine/Plugins/Experimental/DynamicMaterial/Source/DynamicMaterialTextureSetEditor/Private/Widgets/SDMTextureSetBuilderCellBase.cpp
@@ -0,0 +1,120 @@
+// Copyright Epic Games, Inc. All Rights Reserved.
+
+#include "Widgets/SDMTextureSetBuilderCellBase.h"
+
+#include "AssetRegistry/AssetData.h"
+#include "Engine/Texture.h"
+#include "Widgets/DMTextureSetBuilderDragDropOperation.h"
+#include "Widgets/Images/SImage.h"
+#include "Widgets/SBoxPanel.h"
+#include "Widgets/SDMTextureSetBuilder.h"
+
+#define LOCTEXT_NAMESPACE "SDMTextureSetBuilderCellBase"
+
+SDMTextureSetBuilderCellBase::SDMTextureSetBuilderCellBase()
+	: Index(-1)
+	, bIsMaterialProperty(false)
+{
+}
+
+void SDMTextureSetBuilderCellBase::Construct(const FArguments& InArgs, const TSharedRef<SDMTextureSetBuilder>& InTextureSetBuilder,
+	UTexture* InTexture, int32 InIndex, bool bInIsMaterialProperty)
+{
+	TextureSetBuilderWeak = InTextureSetBuilder;
+	Texture.Reset(InTexture);
+	Index = InIndex;
+	bIsMaterialProperty = bInIsMaterialProperty;
+}
+
+UTexture* SDMTextureSetBuilderCellBase::GetTexture() const
+{
+	return Texture.Get();
+}
+
+void SDMTextureSetBuilderCellBase::SetTexture(UTexture* InTexture)
+{
+	Texture.Reset(InTexture);
+}
+
+FReply SDMTextureSetBuilderCellBase::OnMouseButtonDown(const FGeometry& MyGeometry, const FPointerEvent& MouseEvent)
+{
+	if (MouseEvent.GetEffectingButton() == EKeys::LeftMouseButton && Texture.IsValid())
+	{
+		return FReply::Handled().DetectDrag(SharedThis(this), EKeys::LeftMouseButton);
+	}
+
+	return SCompoundWidget::OnMouseButtonDown(MyGeometry, MouseEvent);
+}
+
+FReply SDMTextureSetBuilderCellBase::OnDragDetected(const FGeometry& MyGeometry, const FPointerEvent& MouseEvent)
+{
+	const TSharedRef<FDMTextureSetBuilderDragDropOperation> Operation = FDMTextureSetBuilderDragDropOperation::New(
+		FAssetData(Texture.Get()),
+		Index,
+		bIsMaterialProperty
+	);
+
+	return FReply::Handled().BeginDragDrop(Operation);
+}
+
+bool SDMTextureSetBuilderCellBase::OnAssetDraggedOver(TArrayView<FAssetData> InAssets)
+{
+	for (const FAssetData& Asset : InAssets)
+	{
+		UClass* AssetClass = Asset.GetClass(EResolveClass::Yes);
+
+		if (AssetClass && AssetClass->IsChildOf(UTexture::StaticClass()))
+		{
+			return true;
+		}
+	}
+
+	return false;
+}
+
+void SDMTextureSetBuilderCellBase::OnAssetsDropped(const FDragDropEvent& InDragDropEvent, TArrayView<FAssetData> InAssets)
+{
+	if (TSharedPtr<FDMTextureSetBuilderDragDropOperation> BuilderOperation = InDragDropEvent.GetOperationAs<FDMTextureSetBuilderDragDropOperation>())
+	{
+		if (TSharedPtr<SDMTextureSetBuilder> TextureSetBuilder = TextureSetBuilderWeak.Pin())
+		{
+			TextureSetBuilder->SwapTexture(
+				BuilderOperation->GetIndex(),
+				BuilderOperation->IsMaterialProperty(),
+				Index,
+				bIsMaterialProperty
+			);
+		}
+	}
+}
+
+EVisibility SDMTextureSetBuilderCellBase::GetImageVisibility() const
+{
+	return Texture.IsValid() ? EVisibility::Visible : EVisibility::Hidden;
+}
+
+FText SDMTextureSetBuilderCellBase::GetToolTipText() const
+{
+	UTexture* TextureObject = Texture.Get();
+
+	if (!TextureObject)
+	{
+		return LOCTEXT("NoTexture", "Texture slot empty.");
+	}
+
+	return FText::FromString(TextureObject->GetPathName());
+}
+
+FText SDMTextureSetBuilderCellBase::GetTextureName() const
+{
+	UTexture* TextureObject = Texture.Get();
+
+	if (!TextureObject)
+	{
+		return LOCTEXT("-", "-");
+	}
+
+	return FText::FromString(TextureObject->GetName());
+}
+
+#undef LOCTEXT_NAMESPACE
diff --git a/Engine/Plugins/Experimental/DynamicMaterial/Source/DynamicMaterialTextureSetEditor/Private/Widgets/SDMTextureSetBuilderCellBase.h b/Engine/Plugins/Experimental/DynamicMaterial/Source/DynamicMaterialTextureSetEditor/Private/Widgets/SDMTextureSetBuilderCellBase.h
new file mode 100644
--- /dev/null
+++ b/Engine/Plugins/Experimental/DynamicMaterial/Source/DynamicMaterialTextureSetEditor/Private/Widgets/SDMTextureSetBuilderCellBase.h
@@ -0,0 +1,51 @@
+// Copyright Epic Games, Inc. All Rights Reserved.
+
+#pragma once
+
+#include "Widgets/SCompoundWidget.h"
+
+#include "Templates/SharedPointer.h"
+#include "UObject/StrongObjectPtr.h"
+
+class SDMTextureSetBuilder;
+class SImage;
+class UTexture;
+
+class SDMTextureSetBuilderCellBase : public SCompoundWidget
+{
+	SLATE_BEGIN_ARGS(SDMTextureSetBuilderCellBase) {}
+	SLATE_END_ARGS()
+
+public:
+	SDMTextureSetBuilderCellBase();
+
+	virtual ~SDMTextureSetBuilderCellBase() override = default;
+
+	void Construct(const FArguments& InArgs, const TSharedRef<SDMTextureSetBuilder>& InTextureSetBuilder,
+		UTexture* InTexture, int32 InIndex, bool bInIsMaterialProperty);
+
+	UTexture* GetTexture() const;
+
+	virtual void SetTexture(UTexture* InTexture);
+
+	//~ Begin SWidget
+	virtual FReply OnMouseButtonDown(const FGeometry& MyGeometry, const FPointerEvent& MouseEvent) override;
+	virtual FReply OnDragDetected(const FGeometry& MyGeometry, const FPointerEvent& MouseEvent) override;
+	//~ End SWidget
+
+protected:
+	TWeakPtr<SDMTextureSetBuilder> TextureSetBuilderWeak;
+	TStrongObjectPtr<UTexture> Texture;
+	int32 Index;
+	bool bIsMaterialProperty;
+
+	bool OnAssetDraggedOver(TArrayView<FAssetData> InAssets);
+
+	void OnAssetsDropped(const FDragDropEvent& InDragDropEvent, TArrayView<FAssetData> InAssets);
+
+	EVisibility GetImageVisibility() const;
+
+	FText GetToolTipText() const;
+
+	FText GetTextureName() const;
+};
diff --git a/Engine/Plugins/Experimental/DynamicMaterial/Source/DynamicMaterialTextureSetEditor/Private/Widgets/SDMTextureSetBuilderMaterialPropertyCell.cpp b/Engine/Plugins/Experimental/DynamicMaterial/Source/DynamicMaterialTextureSetEditor/Private/Widgets/SDMTextureSetBuilderMaterialPropertyCell.cpp
new file mode 100644
--- /dev/null
+++ b/Engine/Plugins/Experimental/DynamicMaterial/Source/DynamicMaterialTextureSetEditor/Private/Widgets/SDMTextureSetBuilderMaterialPropertyCell.cpp
@@ -0,0 +1,214 @@
+// Copyright Epic Games, Inc. All Rights Reserved.
+
+#include "Widgets/SDMTextureSetBuilderMaterialPropertyCell.h"
+
+#include "DMTextureSetBuilderEntry.h"
+#include "Engine/Texture.h"
+#include "ISinglePropertyView.h"
+#include "Materials/Material.h"
+#include "Materials/MaterialInstanceDynamic.h"
+#include "Modules/ModuleManager.h"
+#include "PropertyEditorModule.h"
+#include "SAssetDropTarget.h"
+#include "SceneTypes.h"
+#include "SDMTextureSetBuilder.h"
+#include "Widgets/Images/SImage.h"
+#include "Widgets/SBoxPanel.h"
+#include "Widgets/SOverlay.h"
+#include "Widgets/Text/STextBlock.h"
+
+#define LOCTEXT_NAMESPACE "SDMTextureSetBuilderMaterialPropertyCell"
+
+namespace UE::DynamicMaterial::TextureSet::Private
+{
+	const TCHAR* Material_Red =   TEXT("/Script/Engine.Material'/DynamicMaterial/Materials/TextureSet_Red.TextureSet_Red'");
+	const TCHAR* Material_Green = TEXT("/Script/Engine.Material'/DynamicMaterial/Materials/TextureSet_Green.TextureSet_Green'");
+	const TCHAR* Material_Blue =  TEXT("/Script/Engine.Material'/DynamicMaterial/Materials/TextureSet_Blue.TextureSet_Blue'");
+	const TCHAR* Material_Alpha = TEXT("/Script/Engine.Material'/DynamicMaterial/Materials/TextureSet_Alpha.TextureSet_Alpha'");
+	const TCHAR* Material_RGB =   TEXT("/Script/Engine.Material'/DynamicMaterial/Materials/TextureSet_RGB.TextureSet_RGB'");
+	const TCHAR* Material_All =   TEXT("/Script/Engine.Material'/DynamicMaterial/Materials/TextureSet_All.TextureSet_All'");
+
+	const TCHAR* Parameter_Texture = TEXT("Texture");
+}
+
+SDMTextureSetBuilderMaterialPropertyCell::SDMTextureSetBuilderMaterialPropertyCell()
+	: MaterialBrush(FSlateMaterialBrush(FVector2D(100.f)))
+{
+}
+
+void SDMTextureSetBuilderMaterialPropertyCell::Construct(const FArguments& InArgs, const TSharedRef<SDMTextureSetBuilder>& InTextureSetBuilder, 
+	const TSharedRef<FDMTextureSetBuilderEntry>& InEntry, int32 InIndex)
+{
+	SDMTextureSetBuilderCellBase::Construct(
+		SDMTextureSetBuilderCellBase::FArguments(), 
+		InTextureSetBuilder,
+		InEntry->Texture,
+		InIndex,
+		/* Is Material Property */ true
+	);
+
+	Entry = InEntry;
+
+	UEnum* MaterialPropertyEnum = StaticEnum<EMaterialProperty>();
+
+	if (!MaterialPropertyEnum)
+	{
+		return;
+	}
+
+	SetMaterialForChannelMask();
+
+	EntryProvider = MakeShared<FDMTextureSetBuilderEntryProvider>(InEntry);
+
+	FPropertyEditorModule& PropertyEditorModule = FModuleManager::Get().LoadModuleChecked<FPropertyEditorModule>("PropertyEditor");
+
+	FSinglePropertyParams InitParams;
+	InitParams.NotifyHook = this;
+	InitParams.NamePlacement = EPropertyNamePlacement::Hidden;
+
+	TSharedPtr<ISinglePropertyView> PropertyView = PropertyEditorModule.CreateSingleProperty(
+		EntryProvider.ToSharedRef(),
+		GET_MEMBER_NAME_CHECKED(FDMTextureSetBuilderEntry, ChannelMask),
+		InitParams
+	);
+
+	if (!PropertyView.IsValid())
+	{
+		return;
+	}
+
+	PropertyView->SetEnabled(TAttribute<bool>::CreateSP(this, &SDMTextureSetBuilderMaterialPropertyCell::GetPropertyEnabled));
+
+	ChildSlot
+	[
+		SNew(SVerticalBox)
+
+		+ SVerticalBox::Slot()
+		.AutoHeight()
+		.HAlign(EHorizontalAlignment::HAlign_Center)
+		.Padding(0.f, 0.f, 0.f, 5.f)
+		[
+			SNew(STextBlock)
+			.Text(MaterialPropertyEnum->GetDisplayNameTextByValue(Entry->MaterialProperty))
+		]
+
+		+ SVerticalBox::Slot()
+		.AutoHeight()
+		.HAlign(EHorizontalAlignment::HAlign_Center)
+		.Padding(0.f, 0.f, 0.f, 5.f)
+		[
+			SNew(SAssetDropTarget)
+			.OnAreAssetsAcceptableForDrop(this, &SDMTextureSetBuilderMaterialPropertyCell::OnAssetDraggedOver)
+			.OnAssetsDropped(this, &SDMTextureSetBuilderMaterialPropertyCell::OnAssetsDropped)
+			[
+				SNew(SOverlay)
+				.ToolTipText(this, &SDMTextureSetBuilderMaterialPropertyCell::GetToolTipText)
+
+				+ SOverlay::Slot()
+				[
+					SNew(SImage)
+					.Image(&MaterialBrush)
+					.DesiredSizeOverride(FVector2D(100.f))
+					.Visibility(this, &SDMTextureSetBuilderMaterialPropertyCell::GetImageVisibility)
+				]
+
+				+ SOverlay::Slot()
+				.Padding(5.f)
+				[
+					SNew(STextBlock)
+					.Text(this, &SDMTextureSetBuilderMaterialPropertyCell::GetTextureName)
+					.WrappingPolicy(ETextWrappingPolicy::AllowPerCharacterWrapping)
+					.WrapTextAt(90.f)
+				]
+			]
+		]
+
+		+ SVerticalBox::Slot()
+		.AutoHeight()
+		.HAlign(EHorizontalAlignment::HAlign_Left)
+		.Padding(0.f, 0.f, 0.f, 0.f)
+		[
+			PropertyView.ToSharedRef()
+		]
+	];
+}
+
+void SDMTextureSetBuilderMaterialPropertyCell::SetTexture(UTexture* InTexture)
+{
+	SDMTextureSetBuilderCellBase::SetTexture(InTexture);
+
+	using namespace UE::DynamicMaterial::TextureSet::Private;
+
+	if (UMaterialInstanceDynamic* MIDObject = MID.Get())
+	{
+		MIDObject->SetTextureParameterValue(Parameter_Texture, InTexture);
+	}	
+}
+
+bool SDMTextureSetBuilderMaterialPropertyCell::GetPropertyEnabled() const
+{
+	return Texture.IsValid();
+}
+
+void SDMTextureSetBuilderMaterialPropertyCell::NotifyPostChange(const FPropertyChangedEvent& PropertyChangedEvent, FProperty* PropertyThatChanged)
+{
+	if (PropertyThatChanged && PropertyThatChanged->GetFName() == GET_MEMBER_NAME_CHECKED(FDMTextureSetBuilderEntry, ChannelMask))
+	{
+		SetMaterialForChannelMask();
+	}
+}
+
+void SDMTextureSetBuilderMaterialPropertyCell::SetMaterialForChannelMask()
+{
+	using namespace UE::DynamicMaterial::TextureSet::Private;
+
+	UMaterial* ParentMaterial = nullptr;
+
+	switch (Entry->ChannelMask)
+	{
+		case EDMTextureChannelMask::Red:
+			ParentMaterial = LoadObject<UMaterial>(GetTransientPackage(), Material_Red);
+			break;
+
+		case EDMTextureChannelMask::Green:
+			ParentMaterial = LoadObject<UMaterial>(GetTransientPackage(), Material_Green);
+			break;
+
+		case EDMTextureChannelMask::Blue:
+			ParentMaterial = LoadObject<UMaterial>(GetTransientPackage(), Material_Blue);
+			break;
+
+		case EDMTextureChannelMask::Alpha:
+			ParentMaterial = LoadObject<UMaterial>(GetTransientPackage(), Material_Alpha);
+			break;
+
+		case EDMTextureChannelMask::RGB:
+			ParentMaterial = LoadObject<UMaterial>(GetTransientPackage(), Material_RGB);
+			break;
+
+		case EDMTextureChannelMask::RGBA:
+			ParentMaterial = LoadObject<UMaterial>(GetTransientPackage(), Material_All);
+			break;
+	}
+
+	if (!ParentMaterial)
+	{
+		return;
+	}
+
+	UMaterialInstanceDynamic* MIDObject = MID.Get();
+
+	if (!MIDObject || MIDObject->Parent != ParentMaterial)
+	{
+		MIDObject = UMaterialInstanceDynamic::Create(ParentMaterial, GetTransientPackage());
+		MID.Reset(MIDObject);
+	}
+
+	if (MIDObject)
+	{
+		MIDObject->SetTextureParameterValue(Parameter_Texture, Texture.Get());
+		MaterialBrush.SetMaterial(MIDObject);
+	}
+}
+
+#undef LOCTEXT_NAMESPACE
diff --git a/Engine/Plugins/Experimental/DynamicMaterial/Source/DynamicMaterialTextureSetEditor/Private/Widgets/SDMTextureSetBuilderMaterialPropertyCell.h b/Engine/Plugins/Experimental/DynamicMaterial/Source/DynamicMaterialTextureSetEditor/Private/Widgets/SDMTextureSetBuilderMaterialPropertyCell.h
new file mode 100644
--- /dev/null
+++ b/Engine/Plugins/Experimental/DynamicMaterial/Source/DynamicMaterialTextureSetEditor/Private/Widgets/SDMTextureSetBuilderMaterialPropertyCell.h
@@ -0,0 +1,46 @@
+// Copyright Epic Games, Inc. All Rights Reserved.
+
+#pragma once
+
+#include "Misc/NotifyHook.h"
+#include "Widgets/SDMTextureSetBuilderCellBase.h"
+
+#include "SlateMaterialBrush.h"
+#include "Templates/SharedPointer.h"
+
+class FDMTextureSetBuilderEntryProvider;
+class SDMTextureSetBuilder;
+class UMaterialInstanceDynamic;
+struct FDMTextureSetBuilderEntry;
+
+class SDMTextureSetBuilderMaterialPropertyCell : public SDMTextureSetBuilderCellBase, public FNotifyHook
+{
+	SLATE_BEGIN_ARGS(SDMTextureSetBuilderMaterialPropertyCell) {}
+	SLATE_END_ARGS()
+
+public:
+	SDMTextureSetBuilderMaterialPropertyCell();
+
+	virtual ~SDMTextureSetBuilderMaterialPropertyCell() override = default;
+
+	void Construct(const FArguments& InArgs, const TSharedRef<SDMTextureSetBuilder>& InTextureSetBuilder, 
+		const TSharedRef<FDMTextureSetBuilderEntry>& InEntry, int32 InIndex);
+
+	//~ Begin SDMTextureSetBuilderCellBase
+	virtual void SetTexture(UTexture* InTexture) override;
+	//~ End SDMTextureSetBuilderCellBase
+
+	//~ Begin FNotifyHook
+	virtual void NotifyPostChange(const FPropertyChangedEvent& PropertyChangedEvent, FProperty* PropertyThatChanged) override;
+	//~ End FNotifyHook
+
+protected:
+	TSharedPtr<FDMTextureSetBuilderEntry> Entry;
+	TSharedPtr<FDMTextureSetBuilderEntryProvider> EntryProvider;
+	TStrongObjectPtr<UMaterialInstanceDynamic> MID;
+	FSlateMaterialBrush MaterialBrush;
+
+	bool GetPropertyEnabled() const;
+
+	void SetMaterialForChannelMask();
+};
diff --git a/Engine/Plugins/Experimental/DynamicMaterial/Source/DynamicMaterialTextureSetEditor/Private/Widgets/SDMTextureSetBuilderUnassignedTextureCell.cpp b/Engine/Plugins/Experimental/DynamicMaterial/Source/DynamicMaterialTextureSetEditor/Private/Widgets/SDMTextureSetBuilderUnassignedTextureCell.cpp
new file mode 100644
--- /dev/null
+++ b/Engine/Plugins/Experimental/DynamicMaterial/Source/DynamicMaterialTextureSetEditor/Private/Widgets/SDMTextureSetBuilderUnassignedTextureCell.cpp
@@ -0,0 +1,85 @@
+// Copyright Epic Games, Inc. All Rights Reserved.
+
+#include "Widgets/SDMTextureSetBuilderUnassignedTextureCell.h"
+
+#include "AssetRegistry/AssetData.h"
+#include "Engine/Texture.h"
+#include "SDMTextureSetBuilder.h"
+#include "Widgets/Images/SImage.h"
+#include "Widgets/SBoxPanel.h"
+#include "Widgets/SOverlay.h"
+#include "Widgets/Text/STextBlock.h"
+
+#define LOCTEXT_NAMESPACE "SDMTextureSetBuilderUnassignedTextureCell"
+
+SDMTextureSetBuilderUnassignedTextureCell::SDMTextureSetBuilderUnassignedTextureCell()
+	: TextureBrush(FSlateImageBrush((UObject*)nullptr, FVector2D(100.f)))
+{
+}
+
+void SDMTextureSetBuilderUnassignedTextureCell::Construct(const FArguments& InArgs, const TSharedRef<SDMTextureSetBuilder>& InTextureSetBuilder,
+	UTexture* InTexture, int32 InIndex)
+{
+	SDMTextureSetBuilderCellBase::Construct(
+		SDMTextureSetBuilderCellBase::FArguments(),
+		InTextureSetBuilder,
+		InTexture,
+		InIndex,
+		/* Material Property */ false
+	);
+
+	TextureSetBuilderWeak = InTextureSetBuilder;
+	Index = InIndex;
+
+	TextureBrush.SetResourceObject(InTexture);
+
+	ChildSlot
+	[
+		SNew(SVerticalBox)
+		.Visibility(this, &SDMTextureSetBuilderUnassignedTextureCell::GetImageVisibility)
+
+		+ SVerticalBox::Slot()
+		.AutoHeight()
+		.HAlign(EHorizontalAlignment::HAlign_Center)
+		.Padding(0.f, 0.f, 0.f, 5.f)
+		[
+			SNew(STextBlock)
+			.Text(LOCTEXT("Unassigned", "Unassigned"))
+		]
+
+		+ SVerticalBox::Slot()
+		.AutoHeight()
+		.HAlign(EHorizontalAlignment::HAlign_Center)
+		.Padding(0.f, 0.f, 0.f, 0.f)
+		[
+			SNew(SOverlay)
+			.ToolTipText(this, &SDMTextureSetBuilderUnassignedTextureCell::GetToolTipText)
+
+			+ SOverlay::Slot()
+			[
+				SNew(SImage)
+				.Image(&TextureBrush)
+				.DesiredSizeOverride(FVector2D(100.f))
+				.Visibility(this, &SDMTextureSetBuilderUnassignedTextureCell::GetImageVisibility)
+			]
+
+			+ SOverlay::Slot()
+			.Padding(5.f)
+			[
+				SNew(STextBlock)
+				.Text(this, &SDMTextureSetBuilderUnassignedTextureCell::GetTextureName)
+				.WrappingPolicy(ETextWrappingPolicy::AllowPerCharacterWrapping)
+				.WrapTextAt(90.f)
+			]
+		]
+	];
+}
+
+void SDMTextureSetBuilderUnassignedTextureCell::SetTexture(UTexture* InTexture)
+{
+	SDMTextureSetBuilderCellBase::SetTexture(InTexture);
+
+	TextureBrush.SetResourceObject(InTexture);
+}
+
+#undef LOCTEXT_NAMESPACE
diff --git a/Engine/Plugins/Experimental/DynamicMaterial/Source/DynamicMaterialTextureSetEditor/Private/Widgets/SDMTextureSetBuilderUnassignedTextureCell.h b/Engine/Plugins/Experimental/DynamicMaterial/Source/DynamicMaterialTextureSetEditor/Private/Widgets/SDMTextureSetBuilderUnassignedTextureCell.h
new file mode 100644
--- /dev/null
+++ b/Engine/Plugins/Experimental/DynamicMaterial/Source/DynamicMaterialTextureSetEditor/Private/Widgets/SDMTextureSetBuilderUnassignedTextureCell.h
@@ -0,0 +1,31 @@
+// Copyright Epic Games, Inc. All Rights Reserved.
+
+#pragma once
+
+#include "Widgets/SDMTextureSetBuilderCellBase.h"
+
+#include "Brushes/SlateImageBrush.h"
+#include "Templates/SharedPointer.h"
+
+class SDMTextureSetBuilder;
+
+class SDMTextureSetBuilderUnassignedTextureCell : public SDMTextureSetBuilderCellBase
+{
+	SLATE_BEGIN_ARGS(SDMTextureSetBuilderUnassignedTextureCell) {}
+	SLATE_END_ARGS()
+
+public:
+	SDMTextureSetBuilderUnassignedTextureCell();
+
+	virtual ~SDMTextureSetBuilderUnassignedTextureCell() override = default;
+
+	void Construct(const FArguments& InArgs, const TSharedRef<SDMTextureSetBuilder>& InTextureSetBuilder,
+		UTexture* InTexture, int32 InIndex);
+
+	//~ Begin SDMTextureSetBuilderCellBase
+	virtual void SetTexture(UTexture* InTexture) override;
+	//~ End SDMTextureSetBuilderCellBase
+
+protected:
+	FSlateImageBrush TextureBrush;
+};
diff --git a/Engine/Plugins/Experimental/DynamicMaterial/Source/DynamicMaterialTextureSetEditor/Public/DMTextureSetBlueprintFunctionLibrary.h b/Engine/Plugins/Experimental/DynamicMaterial/Source/DynamicMaterialTextureSetEditor/Public/DMTextureSetBlueprintFunctionLibrary.h
new file mode 100644
--- /dev/null
+++ b/Engine/Plugins/Experimental/DynamicMaterial/Source/DynamicMaterialTextureSetEditor/Public/DMTextureSetBlueprintFunctionLibrary.h
@@ -0,0 +1,38 @@
+// Copyright Epic Games, Inc. All Rights Reserved.
+
+#pragma once
+
+#include "Kismet/BlueprintFunctionLibrary.h"
+
+#include "DMTextureSetBlueprintFunctionLibrary.generated.h"
+
+class UDMTextureSet;
+struct FAssetData;
+
+DECLARE_DELEGATE_TwoParams(FDMTextureSetBuilderOnComplete, UDMTextureSet*, /* Was Accepted */bool);
+
+/**
+ * Material Designer Texture Set Blueprint Function Library
+ */
+UCLASS()
+class UDMTextureSetBlueprintFunctionLibrary : public UBlueprintFunctionLibrary
+{
+	GENERATED_BODY()
+
+public:
+	/**
+	 * Uses the filters from the Texture Set Settings to create a Texture Set based on the given assets.
+	 * @param InAssets The texture assets to assign to the texture slot.
+	 * @return A new texture set or nullptr if no textures were filtered.
+	 */
+	UFUNCTION(BlueprintCallable, Category = "Material Designer")
+	static DYNAMICMATERIALTEXTURESETEDITOR_API UDMTextureSet* CreateTextureSetFromAssets(const TArray<FAssetData>& InAssets);
+
+	/**
+	 * Uses the filters from the Texture Set Settings to create a Texture Set based on the given assets
+	 * and presents the user with a UI to confirm slot assignments.
+	 * @param InAssets The texture assets to assign to the texture slot.
+	 * @param InOnComplete Called when the interactive panel is closed, successful or not.
+	 */
+	static DYNAMICMATERIALTEXTURESETEDITOR_API void CreateTextureSetFromAssetsInteractive(const TArray<FAssetData>& InAssets, FDMTextureSetBuilderOnComplete InOnComplete);
+};
diff --git a/Engine/Plugins/Experimental/DynamicMaterial/Source/DynamicMaterialTextureSetEditor/Public/DMTextureSetFactory.h b/Engine/Plugins/Experimental/DynamicMaterial/Source/DynamicMaterialTextureSetEditor/Public/DMTextureSetFactory.h
new file mode 100644
--- /dev/null
+++ b/Engine/Plugins/Experimental/DynamicMaterial/Source/DynamicMaterialTextureSetEditor/Public/DMTextureSetFactory.h
@@ -0,0 +1,23 @@
+// Copyright Epic Games, Inc. All Rights Reserved.
+
+#pragma once
+
+#include "Factories/Factory.h"
+
+#include "DMTextureSetFactory.generated.h"
+
+UCLASS(MinimalAPI)
+class UDMTextureSetFactory : public UFactory
+{
+	GENERATED_BODY()
+
+public:
+	UDMTextureSetFactory();
+
+	//~ Begin UFactory
+	DYNAMICMATERIALTEXTURESETEDITOR_API virtual UObject* FactoryCreateNew(UClass* Class, UObject* InParent, FName Name, EObjectFlags Flags,
+		UObject* Context, FFeedbackContext* Warn) override;
+	virtual FText GetDisplayName() const override;
+	virtual FText GetToolTip() const override;
+	//~ End UFactory
+};
