Material Designer
- Horde issue 479734

diff --git a/Engine/Plugins/Experimental/DynamicMaterial/Source/DynamicMaterialTextureSetEditor/Private/Widgets/DMTextureSetBuilderEntry.cpp b/Engine/Plugins/Experimental/DynamicMaterial/Source/DynamicMaterialTextureSetEditor/Private/Widgets/DMTextureSetBuilderEntry.cpp
--- a/Engine/Plugins/Experimental/DynamicMaterial/Source/DynamicMaterialTextureSetEditor/Private/Widgets/DMTextureSetBuilderEntry.cpp
+++ b/Engine/Plugins/Experimental/DynamicMaterial/Source/DynamicMaterialTextureSetEditor/Private/Widgets/DMTextureSetBuilderEntry.cpp
@@ -2,7 +2,19 @@
 
 #include "Widgets/DMTextureSetBuilderEntry.h"
 
-#include "InstancedStruct.h"
+#include "SceneTypes.h"
+
+FDMTextureSetBuilderEntry::FDMTextureSetBuilderEntry()
+	: FDMTextureSetBuilderEntry(EMaterialProperty::MP_BaseColor, nullptr, EDMTextureChannelMask::RGBA)
+{
+}
+
+FDMTextureSetBuilderEntry::FDMTextureSetBuilderEntry(EMaterialProperty InMaterialProperty, UTexture* InTexture, EDMTextureChannelMask InMask)
+	: MaterialProperty(InMaterialProperty)
+	, Texture(InTexture)
+	, ChannelMask(InMask)
+{
+}
 
 FDMTextureSetBuilderEntryProvider::FDMTextureSetBuilderEntryProvider(const TSharedRef<FDMTextureSetBuilderEntry>& InEntry)
 	: Entry(InEntry)
diff --git a/Engine/Plugins/Experimental/DynamicMaterial/Source/DynamicMaterialTextureSetEditor/Private/Widgets/DMTextureSetBuilderEntry.h b/Engine/Plugins/Experimental/DynamicMaterial/Source/DynamicMaterialTextureSetEditor/Private/Widgets/DMTextureSetBuilderEntry.h
--- a/Engine/Plugins/Experimental/DynamicMaterial/Source/DynamicMaterialTextureSetEditor/Private/Widgets/DMTextureSetBuilderEntry.h
+++ b/Engine/Plugins/Experimental/DynamicMaterial/Source/DynamicMaterialTextureSetEditor/Private/Widgets/DMTextureSetBuilderEntry.h
@@ -18,6 +18,10 @@ struct FDMTextureSetBuilderEntry
 {
 	GENERATED_BODY()
 
+	FDMTextureSetBuilderEntry();
+
+	FDMTextureSetBuilderEntry(EMaterialProperty InMaterialProperty, UTexture* InTexture, EDMTextureChannelMask InMask);
+
 	UPROPERTY(EditAnywhere, Category = "Material Designer")
 	TEnumAsByte<EMaterialProperty> MaterialProperty = EMaterialProperty::MP_BaseColor;
 
