Material Designer
- Horde issue 479790

diff --git a/Engine/Plugins/Experimental/DynamicMaterial/Source/DynamicMaterialTextureSetEditor/Private/AssetDefinition_DMTextureSet.cpp b/Engine/Plugins/Experimental/DynamicMaterial/Source/DynamicMaterialTextureSetEditor/Private/AssetDefinition_DMTextureSet.cpp
--- a/Engine/Plugins/Experimental/DynamicMaterial/Source/DynamicMaterialTextureSetEditor/Private/AssetDefinition_DMTextureSet.cpp
+++ b/Engine/Plugins/Experimental/DynamicMaterial/Source/DynamicMaterialTextureSetEditor/Private/AssetDefinition_DMTextureSet.cpp
@@ -8,7 +8,7 @@
 
 FText UAssetDefinition_DMTextureSet::GetAssetDisplayName() const
 {
-	return LOCTEXT("DMTextureSet", "Texture Set");
+	return LOCTEXT("DMTextureSet", "Material Designer Texture Set");
 }
 
 FText UAssetDefinition_DMTextureSet::GetAssetDisplayName(const FAssetData& InAssetData) const
diff --git a/Engine/Plugins/Experimental/DynamicMaterial/Source/DynamicMaterialTextureSetEditor/Private/DMTextureSetBlueprintFunctionLibrary.cpp b/Engine/Plugins/Experimental/DynamicMaterial/Source/DynamicMaterialTextureSetEditor/Private/DMTextureSetBlueprintFunctionLibrary.cpp
--- a/Engine/Plugins/Experimental/DynamicMaterial/Source/DynamicMaterialTextureSetEditor/Private/DMTextureSetBlueprintFunctionLibrary.cpp
+++ b/Engine/Plugins/Experimental/DynamicMaterial/Source/DynamicMaterialTextureSetEditor/Private/DMTextureSetBlueprintFunctionLibrary.cpp
@@ -2,7 +2,6 @@
 
 #include "DMTextureSetBlueprintFunctionLibrary.h"
 
-#include "AssetRegistry/AssetData.h"
 #include "CoreGlobals.h"
 #include "DMTextureSet.h"
 #include "DMTextureSetFactory.h"
diff --git a/Engine/Plugins/Experimental/DynamicMaterial/Source/DynamicMaterialTextureSetEditor/Public/DMTextureSetBlueprintFunctionLibrary.h b/Engine/Plugins/Experimental/DynamicMaterial/Source/DynamicMaterialTextureSetEditor/Public/DMTextureSetBlueprintFunctionLibrary.h
--- a/Engine/Plugins/Experimental/DynamicMaterial/Source/DynamicMaterialTextureSetEditor/Public/DMTextureSetBlueprintFunctionLibrary.h
+++ b/Engine/Plugins/Experimental/DynamicMaterial/Source/DynamicMaterialTextureSetEditor/Public/DMTextureSetBlueprintFunctionLibrary.h
@@ -4,10 +4,11 @@
 
 #include "Kismet/BlueprintFunctionLibrary.h"
 
+#include "AssetRegistry/AssetData.h"
+
 #include "DMTextureSetBlueprintFunctionLibrary.generated.h"
 
 class UDMTextureSet;
-struct FAssetData;
 
 DECLARE_DELEGATE_TwoParams(FDMTextureSetBuilderOnComplete, UDMTextureSet*, /* Was Accepted */bool);
 
