Material Designer
- Scripted actions part 2 / Material Designer integration
- Some small UI changes to the texture set interface

diff --git a/Engine/Plugins/Experimental/DynamicMaterial/Source/DynamicMaterial/DynamicMaterial.Build.cs b/Engine/Plugins/Experimental/DynamicMaterial/Source/DynamicMaterial/DynamicMaterial.Build.cs
--- a/Engine/Plugins/Experimental/DynamicMaterial/Source/DynamicMaterial/DynamicMaterial.Build.cs
+++ b/Engine/Plugins/Experimental/DynamicMaterial/Source/DynamicMaterial/DynamicMaterial.Build.cs
@@ -13,6 +13,7 @@ public class DynamicMaterial : ModuleRules
 			{
 				"Core",
 				"CoreUObject",
+				"DynamicMaterialTextureSet",
 				"Slate",
 				"SlateCore",
 				"UMG"
diff --git a/Engine/Plugins/Experimental/DynamicMaterial/Source/DynamicMaterial/Private/DMDefs.cpp b/Engine/Plugins/Experimental/DynamicMaterial/Source/DynamicMaterial/Private/DMDefs.cpp
--- a/Engine/Plugins/Experimental/DynamicMaterial/Source/DynamicMaterial/Private/DMDefs.cpp
+++ b/Engine/Plugins/Experimental/DynamicMaterial/Source/DynamicMaterial/Private/DMDefs.cpp
@@ -1,6 +1,8 @@
 // Copyright Epic Games, Inc. All Rights Reserved.
 
 #include "DMDefs.h"
+
+#include "DMTextureSetMaterialProperty.h"
 #include "Misc/Paths.h"
 
 int32 FDMUpdateGuard::GuardCount = 0;
@@ -20,6 +22,108 @@ FString UE::DynamicMaterial::CreateNodeComment(const ANSICHAR* InFile, int InLin
 	}
 }
 
+EDMTextureSetMaterialProperty UE::DynamicMaterial::MaterialPropertyTypeToMaterialProperty(EDMMaterialPropertyType InPropertyType)
+{
+	switch (InPropertyType)
+	{
+		case EDMMaterialPropertyType::BaseColor:
+			return EDMTextureSetMaterialProperty::BaseColor;
+
+		case EDMMaterialPropertyType::EmissiveColor:
+			return EDMTextureSetMaterialProperty::EmissiveColor;
+
+		case EDMMaterialPropertyType::Opacity:
+			return EDMTextureSetMaterialProperty::Opacity;
+
+		case EDMMaterialPropertyType::OpacityMask:
+			return EDMTextureSetMaterialProperty::OpacityMask;
+
+		case EDMMaterialPropertyType::Roughness:
+			return EDMTextureSetMaterialProperty::Roughness;
+
+		case EDMMaterialPropertyType::Specular:
+			return EDMTextureSetMaterialProperty::Specular;
+
+		case EDMMaterialPropertyType::Metallic:
+			return EDMTextureSetMaterialProperty::Metallic;
+
+		case EDMMaterialPropertyType::Normal:
+			return EDMTextureSetMaterialProperty::Normal;
+
+		case EDMMaterialPropertyType::PixelDepthOffset:
+			return EDMTextureSetMaterialProperty::PixelDepthOffset;
+
+		case EDMMaterialPropertyType::WorldPositionOffset:
+			return EDMTextureSetMaterialProperty::WorldPositionOffset;
+
+		case EDMMaterialPropertyType::AmbientOcclusion:
+			return EDMTextureSetMaterialProperty::AmbientOcclusion;
+
+		case EDMMaterialPropertyType::Anisotropy:
+			return EDMTextureSetMaterialProperty::Anisotropy;
+
+		case EDMMaterialPropertyType::Refraction:
+			return EDMTextureSetMaterialProperty::Refraction;
+
+		case EDMMaterialPropertyType::Tangent:
+			return EDMTextureSetMaterialProperty::Tangent;
+
+		default:
+			return EDMTextureSetMaterialProperty::None;
+	}
+}
+
+EDMMaterialPropertyType UE::DynamicMaterial::MaterialPropertyToMaterialPropertyType(EDMTextureSetMaterialProperty InPropertyType)
+{
+	switch (InPropertyType)
+	{
+		case EDMTextureSetMaterialProperty::BaseColor:
+			return EDMMaterialPropertyType::BaseColor;
+
+		case EDMTextureSetMaterialProperty::EmissiveColor:
+			return EDMMaterialPropertyType::EmissiveColor;
+
+		case EDMTextureSetMaterialProperty::Opacity:
+			return EDMMaterialPropertyType::Opacity;
+
+		case EDMTextureSetMaterialProperty::OpacityMask:
+			return EDMMaterialPropertyType::OpacityMask;
+
+		case EDMTextureSetMaterialProperty::Roughness:
+			return EDMMaterialPropertyType::Roughness;
+
+		case EDMTextureSetMaterialProperty::Specular:
+			return EDMMaterialPropertyType::Specular;
+
+		case EDMTextureSetMaterialProperty::Metallic:
+			return EDMMaterialPropertyType::Metallic;
+
+		case EDMTextureSetMaterialProperty::Normal:
+			return EDMMaterialPropertyType::Normal;
+
+		case EDMTextureSetMaterialProperty::PixelDepthOffset:
+			return EDMMaterialPropertyType::PixelDepthOffset;
+
+		case EDMTextureSetMaterialProperty::WorldPositionOffset:
+			return EDMMaterialPropertyType::WorldPositionOffset;
+
+		case EDMTextureSetMaterialProperty::AmbientOcclusion:
+			return EDMMaterialPropertyType::AmbientOcclusion;
+
+		case EDMTextureSetMaterialProperty::Anisotropy:
+			return EDMMaterialPropertyType::Anisotropy;
+
+		case EDMTextureSetMaterialProperty::Refraction:
+			return EDMMaterialPropertyType::Refraction;
+
+		case EDMTextureSetMaterialProperty::Tangent:
+			return EDMMaterialPropertyType::Tangent;
+
+		default:
+			return EDMMaterialPropertyType::None;
+	}
+}
+
 bool FDMInitializationGuard::IsInitializing()
 {
 	return GuardCount > 0;
diff --git a/Engine/Plugins/Experimental/DynamicMaterial/Source/DynamicMaterial/Public/DMDefs.h b/Engine/Plugins/Experimental/DynamicMaterial/Source/DynamicMaterial/Public/DMDefs.h
--- a/Engine/Plugins/Experimental/DynamicMaterial/Source/DynamicMaterial/Public/DMDefs.h
+++ b/Engine/Plugins/Experimental/DynamicMaterial/Source/DynamicMaterial/Public/DMDefs.h
@@ -9,6 +9,7 @@
 
 class UTexture;
 class UDMMaterialStage;
+enum class EDMTextureSetMaterialProperty : uint8;
 
 /** Some short names provided so they fit nicely on buttons! */
 UENUM(BlueprintType)
@@ -163,11 +164,15 @@ private:
 
 namespace UE::DynamicMaterial
 {
-	/** Designed to be used with preprocessor macros. See below. */
-	FString DYNAMICMATERIAL_API CreateNodeComment(const ANSICHAR* InFile, int InLine, const ANSICHAR* InFunction, const FString* InComment = nullptr);
-
 	constexpr int32 RenameFlags = REN_DontCreateRedirectors | REN_DoNotDirty | REN_ForceNoResetLoaders | REN_NonTransactional;
 	constexpr EClassFlags InvalidClassFlags = CLASS_Abstract | CLASS_Hidden | CLASS_Deprecated | CLASS_NewerVersionExists;
+
+	/** Designed to be used with preprocessor macros. See below. */
+	DYNAMICMATERIAL_API FString CreateNodeComment(const ANSICHAR* InFile, int InLine, const ANSICHAR* InFunction, const FString* InComment = nullptr);
+
+	DYNAMICMATERIAL_API EDMTextureSetMaterialProperty MaterialPropertyTypeToMaterialProperty(EDMMaterialPropertyType InPropertyType);
+
+	DYNAMICMATERIAL_API EDMMaterialPropertyType MaterialPropertyToMaterialPropertyType(EDMTextureSetMaterialProperty InPropertyType);
 }
 
 #define UE_DM_NodeComment_Default UE::DynamicMaterial::CreateNodeComment(__FILE__, __LINE__, __FUNCTION__)
diff --git a/Engine/Plugins/Experimental/DynamicMaterial/Source/DynamicMaterialEditor/DynamicMaterialEditor.Build.cs b/Engine/Plugins/Experimental/DynamicMaterial/Source/DynamicMaterialEditor/DynamicMaterialEditor.Build.cs
--- a/Engine/Plugins/Experimental/DynamicMaterial/Source/DynamicMaterialEditor/DynamicMaterialEditor.Build.cs
+++ b/Engine/Plugins/Experimental/DynamicMaterial/Source/DynamicMaterialEditor/DynamicMaterialEditor.Build.cs
@@ -30,6 +30,7 @@ public class DynamicMaterialEditor : ModuleRules
 				"DeveloperSettings",
 				"DynamicMaterialShadersEditor",
 				"DynamicMaterialTextureSet",
+				"DynamicMaterialTextureSetEditor",
 				"EditorSubsystem",
 				"EditorWidgets",
 				"InputCore",
diff --git a/Engine/Plugins/Experimental/DynamicMaterial/Source/DynamicMaterialEditor/Private/DMContentBrowserIntegration.cpp b/Engine/Plugins/Experimental/DynamicMaterial/Source/DynamicMaterialEditor/Private/DMContentBrowserIntegration.cpp
new file mode 100644
--- /dev/null
+++ b/Engine/Plugins/Experimental/DynamicMaterial/Source/DynamicMaterialEditor/Private/DMContentBrowserIntegration.cpp
@@ -0,0 +1,193 @@
+// Copyright Epic Games, Inc. All Rights Reserved.
+
+#include "DMContentBrowserIntegration.h"
+
+#include "AssetRegistry/AssetRegistryModule.h"
+#include "AssetToolsModule.h"
+#include "CoreGlobals.h"
+#include "DMTextureSet.h"
+#include "DMTextureSetBlueprintFunctionLibrary.h"
+#include "DMTextureSetContentBrowserIntegration.h"
+#include "Engine/Texture.h"
+#include "Framework/MultiBox/MultiBoxBuilder.h"
+#include "IDynamicMaterialEditorModule.h"
+#include "Material/DynamicMaterialInstance.h"
+#include "Material/DynamicMaterialInstanceFactory.h"
+#include "Materials/Material.h"
+#include "Model/DynamicMaterialModelEditorOnlyData.h"
+#include "ScopedTransaction.h"
+#include "Styling/SlateIconFinder.h"
+
+#define LOCTEXT_NAMESPACE "FDMContentBrowserIntegration"
+
+FDelegateHandle FDMContentBrowserIntegration::PopulateHandle;
+
+void FDMContentBrowserIntegration::Integrate()
+{
+	if (PopulateHandle.IsValid())
+	{
+		Disintegrate();
+	}
+
+	PopulateHandle = FDMTextureSetContentBrowserIntegration::GetPopulateExtenderDelegate().AddStatic(&FDMContentBrowserIntegration::ExtendMenu);
+}
+
+void FDMContentBrowserIntegration::Disintegrate()
+{
+	if (!PopulateHandle.IsValid())
+	{
+		return;
+	}
+
+	FDMTextureSetContentBrowserIntegration::GetPopulateExtenderDelegate().Remove(PopulateHandle);
+
+	PopulateHandle.Reset();
+}
+
+void FDMContentBrowserIntegration::ExtendMenu(FMenuBuilder& InMenuBuilder, const TArray<FAssetData>& InSelectedAssets)
+{
+	InMenuBuilder.AddMenuEntry(
+		LOCTEXT("CreateMaterialDesignerInstanceFromTextureSet", "Create Material Designer Instance"),
+		LOCTEXT("CreateMaterialDesignerInstanceFromTextureSetTooltip", "Creates a Material Designer Instance in the content browser using a Texture Set."),
+		FSlateIconFinder::FindIconForClass(UMaterial::StaticClass()),
+		FUIAction(FExecuteAction::CreateStatic(&FDMContentBrowserIntegration::CreateMaterialDesignerInstanceFromTextureSet, InSelectedAssets))
+	);
+
+	InMenuBuilder.AddMenuEntry(
+		LOCTEXT("UpdateMaterialDesignerInstanceFromTextureSetAdd", "Update Material Designer Instance (Add)"),
+		LOCTEXT("UpdateMaterialDesignerInstanceFromTextureAddSetTooltip", "Updates the opened Material Designer Instance using a Texture Set, adding new layers to the Model."),
+		FSlateIconFinder::FindIconForClass(UMaterial::StaticClass()),
+		FUIAction(FExecuteAction::CreateStatic(&FDMContentBrowserIntegration::UpdateMaterialDesignerInstanceFromTextureSet, InSelectedAssets, /* Replace */ false))
+	);
+
+	InMenuBuilder.AddMenuEntry(
+		LOCTEXT("UpdateMaterialDesignerInstanceFromTextureSetReplace", "Update Material Designer Instance (Replace)"),
+		LOCTEXT("UpdateMaterialDesignerInstanceFromTextureSetReplaceTooltip", "Updates the opened Material Designer Instance using a Texture Set, replacing slots in the Model."),
+		FSlateIconFinder::FindIconForClass(UMaterial::StaticClass()),
+		FUIAction(FExecuteAction::CreateStatic(&FDMContentBrowserIntegration::UpdateMaterialDesignerInstanceFromTextureSet, InSelectedAssets, /* Replace */ true))
+	);
+}
+
+void FDMContentBrowserIntegration::CreateMaterialDesignerInstanceFromTextureSet(TArray<FAssetData> InSelectedAssets)
+{
+	if (InSelectedAssets.IsEmpty())
+	{
+		return;
+	}
+
+	UDMTextureSetBlueprintFunctionLibrary::CreateTextureSetFromAssetsInteractive(
+		InSelectedAssets,
+		FDMTextureSetBuilderOnComplete::CreateStatic(
+			&FDMContentBrowserIntegration::OnCreateMaterialDesignerInstanceFromTextureSetComplete,
+			InSelectedAssets[0].PackagePath.ToString()
+		)
+	);
+}
+
+void FDMContentBrowserIntegration::OnCreateMaterialDesignerInstanceFromTextureSetComplete(UDMTextureSet* InTextureSet, bool bInAccepted,
+	FString InPath)
+{
+	if (!InTextureSet || !bInAccepted)
+	{
+		return;
+	}
+
+	UDynamicMaterialInstance* Instance = Cast<UDynamicMaterialInstance>(GetMutableDefault<UDynamicMaterialInstanceFactory>()->FactoryCreateNew(
+		UDynamicMaterialInstance::StaticClass(),
+		GetTransientPackage(),
+		NAME_None,
+		RF_Transactional,
+		/* Context */ nullptr,
+		GWarn
+	));
+
+	if (!Instance)
+	{
+		return;
+	}
+
+	UDynamicMaterialModelEditorOnlyData* EditorOnlyData = UDynamicMaterialModelEditorOnlyData::Get(Instance);
+
+	if (!EditorOnlyData)
+	{
+		return;
+	}
+
+	if (!EditorOnlyData->AddTextureSet(InTextureSet, /* Replace */ true))
+	{
+		return;
+	}
+
+	const FAssetToolsModule& AssetToolsModule = FModuleManager::LoadModuleChecked<FAssetToolsModule>("AssetTools");
+
+	FString UniquePackageName;
+	FString UniqueAssetName;
+
+	const FString BasePackageName = InPath / TEXT("NewMaterialDesignerInstance");
+	AssetToolsModule.Get().CreateUniqueAssetName(BasePackageName, TEXT(""), UniquePackageName, UniqueAssetName);
+
+	UPackage* Package = CreatePackage(*UniquePackageName);
+
+	if (!Package)
+	{
+		return;
+	}
+
+	Instance->SetFlags(RF_Standalone | RF_Public);
+	Instance->Rename(*UniqueAssetName, Package, REN_DontCreateRedirectors);
+
+	FAssetRegistryModule::AssetCreated(Instance);
+}
+
+void FDMContentBrowserIntegration::UpdateMaterialDesignerInstanceFromTextureSet(TArray<FAssetData> InSelectedAssets, bool bInReplace)
+{
+	if (InSelectedAssets.IsEmpty())
+	{
+		return;
+	}
+
+	UDMTextureSetBlueprintFunctionLibrary::CreateTextureSetFromAssetsInteractive(
+		InSelectedAssets,
+		FDMTextureSetBuilderOnComplete::CreateStatic(
+			&FDMContentBrowserIntegration::OnUpdateMaterialDesignerInstanceFromTextureSetComplete,
+			bInReplace
+		)
+	);
+}
+
+void FDMContentBrowserIntegration::OnUpdateMaterialDesignerInstanceFromTextureSetComplete(UDMTextureSet* InTextureSet, bool bInAccepted,
+	bool bInReplace)
+{
+	if (!InTextureSet || !bInAccepted)
+	{
+		return;
+	}
+
+	IDynamicMaterialEditorModule& DynamicMaterialEditorModule = IDynamicMaterialEditorModule::Get();
+
+	UDynamicMaterialModel* Model = DynamicMaterialEditorModule.GetOpenedMaterialModel(nullptr);
+
+	if (!Model)
+	{
+		return;
+	}
+
+	UDynamicMaterialModelEditorOnlyData* EditorOnlyData = UDynamicMaterialModelEditorOnlyData::Get(Model);
+
+	if (!EditorOnlyData)
+	{
+		return;
+	}
+
+	FScopedTransaction Transaction(LOCTEXT("AddTextureSet", "Material Designer Add Texture Set"));
+	EditorOnlyData->Modify();
+
+	const bool bSuccess = EditorOnlyData->AddTextureSet(InTextureSet, bInReplace);
+
+	if (!bSuccess)
+	{
+		Transaction.Cancel();
+	}
+}
+
+#undef LOCTEXT_NAMESPACE
diff --git a/Engine/Plugins/Experimental/DynamicMaterial/Source/DynamicMaterialEditor/Private/DMContentBrowserIntegration.h b/Engine/Plugins/Experimental/DynamicMaterial/Source/DynamicMaterialEditor/Private/DMContentBrowserIntegration.h
new file mode 100644
--- /dev/null
+++ b/Engine/Plugins/Experimental/DynamicMaterial/Source/DynamicMaterialEditor/Private/DMContentBrowserIntegration.h
@@ -0,0 +1,31 @@
+// Copyright Epic Games, Inc. All Rights Reserved.
+
+#pragma once
+
+#include "AssetRegistry/AssetData.h"
+#include "Containers/ContainersFwd.h"
+
+class FMenuBuilder;
+class UDMTextureSet;
+struct FAssetData;
+
+class FDMContentBrowserIntegration
+{
+public:
+	static void Integrate();
+
+	static void Disintegrate();
+
+	static void UpdateMaterialDesignerInstanceFromTextureSet(TArray<FAssetData> InSelectedAssets, bool bInReplace);
+
+protected:
+	static void ExtendMenu(FMenuBuilder& InMenuBuilder, const TArray<FAssetData>& InSelectedAssets);
+
+	static void CreateMaterialDesignerInstanceFromTextureSet(TArray<FAssetData> InSelectedAssets);
+
+	static void OnCreateMaterialDesignerInstanceFromTextureSetComplete(UDMTextureSet* InTextureSet, bool bInAccepted, FString InPath);
+
+	static void OnUpdateMaterialDesignerInstanceFromTextureSetComplete(UDMTextureSet* InTextureSet, bool bInAccepted, bool bInReplace);
+
+	static FDelegateHandle PopulateHandle;
+};
diff --git a/Engine/Plugins/Experimental/DynamicMaterial/Source/DynamicMaterialEditor/Private/DynamicMaterialEditorModule.cpp b/Engine/Plugins/Experimental/DynamicMaterial/Source/DynamicMaterialEditor/Private/DynamicMaterialEditorModule.cpp
--- a/Engine/Plugins/Experimental/DynamicMaterial/Source/DynamicMaterialEditor/Private/DynamicMaterialEditorModule.cpp
+++ b/Engine/Plugins/Experimental/DynamicMaterial/Source/DynamicMaterialEditor/Private/DynamicMaterialEditorModule.cpp
@@ -22,6 +22,7 @@
 #include "DetailsPanel/DMPropertyTypeCustomizer.h"
 #include "DetailsPanel/DMValueDetailsRowExtensions.h"
 #include "DetailsPanel/Slate/SDMMaterialListExtensionWidget.h"
+#include "DMContentBrowserIntegration.h"
 #include "DMWorldSubsystem.h"
 #include "DynamicMaterialEditorCommands.h"
 #include "DynamicMaterialEditorSettings.h"
@@ -259,6 +260,7 @@ void FDynamicMaterialEditorModule::StartupModule()
 {
 	FDynamicMaterialEditorStyle::Initialize();
 	FDynamicMaterialEditorCommands::Register();
+	FDMContentBrowserIntegration::Integrate();
 	MapCommands();
 
 	FPropertyEditorModule& PropertyModule = FModuleManager::GetModuleChecked<FPropertyEditorModule>("PropertyEditor");
@@ -321,6 +323,7 @@ void FDynamicMaterialEditorModule::ShutdownModule()
 {
 	FDynamicMaterialEditorCommands::Unregister();
 	FDynamicMaterialEditorStyle::Shutdown();
+	FDMContentBrowserIntegration::Disintegrate();
 
 	if (UObjectInitialized() && !IsEngineExitRequested() && FModuleManager::Get().IsModuleLoaded("PropertyEditor"))
 	{
@@ -555,13 +558,13 @@ void FDynamicMaterialEditorModule::OpenEditor(UWorld* InWorld) const
 
 UDynamicMaterialModel* FDynamicMaterialEditorModule::GetOpenedMaterialModel(UWorld* InWorld) const
 {
-	if (IsValid(InWorld))
+	if (TSharedPtr<SDMEditor> Editor = FDMLevelEditorIntegration::GetEditorForWorld(InWorld))
 	{
-		if (TSharedPtr<SDMEditor> Editor = FDMLevelEditorIntegration::GetEditorForWorld(InWorld))
-		{
-			return Editor->GetMaterialModel();
-		}
+		return Editor->GetMaterialModel();
+	}
 
+	if (IsValid(InWorld))
+	{
 		if (UDMWorldSubsystem* DMWorldSubsystem = InWorld->GetSubsystem<UDMWorldSubsystem>())
 		{
 			return DMWorldSubsystem->ExecuteGetCustomEditorModelDelegate();
diff --git a/Engine/Plugins/Experimental/DynamicMaterial/Source/DynamicMaterialEditor/Private/Model/DynamicMaterialModelEditorOnlyData.cpp b/Engine/Plugins/Experimental/DynamicMaterial/Source/DynamicMaterialEditor/Private/Model/DynamicMaterialModelEditorOnlyData.cpp
--- a/Engine/Plugins/Experimental/DynamicMaterial/Source/DynamicMaterialEditor/Private/Model/DynamicMaterialModelEditorOnlyData.cpp
+++ b/Engine/Plugins/Experimental/DynamicMaterial/Source/DynamicMaterialEditor/Private/Model/DynamicMaterialModelEditorOnlyData.cpp
@@ -7,6 +7,8 @@
 #include "Components/DMMaterialLayer.h"
 #include "Components/DMMaterialProperty.h"
 #include "Components/DMMaterialSlot.h"
+#include "Components/DMMaterialStageBlend.h"
+#include "Components/DMMaterialSubStage.h"
 #include "Components/DMMaterialValue.h"
 #include "Components/DMTextureUV.h"
 #include "Components/MaterialProperties/DMMPAmbientOcclusion.h"
@@ -23,9 +25,14 @@
 #include "Components/MaterialProperties/DMMPSpecular.h"
 #include "Components/MaterialProperties/DMMPTangent.h"
 #include "Components/MaterialProperties/DMMPWorldPositionOffset.h"
+#include "Components/MaterialStageExpressions/DMMSETextureSample.h"
+#include "Components/MaterialStageInputs/DMMSIExpression.h"
+#include "Components/MaterialStageInputs/DMMSIValue.h"
+#include "Components/MaterialValues/DMMaterialValueTexture.h"
 #include "CoreGlobals.h"
 #include "DMComponentPath.h"
 #include "DMDefs.h"
+#include "DMTextureSet.h"
 #include "DynamicMaterialEditorModule.h"
 #include "DynamicMaterialEditorSettings.h"
 #include "DynamicMaterialModule.h"
@@ -809,8 +816,124 @@ TSharedRef<FDMMaterialBuildState> UDynamicMaterialModelEditorOnlyData::CreateBui
 
 bool UDynamicMaterialModelEditorOnlyData::AddTextureSet(UDMTextureSet* InTextureSet, bool bInReplaceSlots)
 {
-	// TODO Integrate!
-	return false;
+	bool bMadeChange = false;
+
+	for (const TPair<EDMTextureSetMaterialProperty, FDMMaterialTexture>& MaterialTexture : InTextureSet->GetTextures())
+	{
+		const EDMMaterialPropertyType PropertyType = UE::DynamicMaterial::MaterialPropertyToMaterialPropertyType(MaterialTexture.Key);
+
+		if (PropertyType == EDMMaterialPropertyType::None)
+		{
+			continue;
+		}
+
+		UDMMaterialSlot* Slot = GetSlotForMaterialProperty(PropertyType);
+
+		if (!Slot)
+		{
+			continue;
+		}
+
+		UTexture* Texture = MaterialTexture.Value.Texture.LoadSynchronous();
+
+		if (!Texture)
+		{
+			continue;
+		}
+
+		if (GUndo)
+		{
+			Slot->Modify();
+		}
+
+		const int32 OutputMask = MaterialTexture.Value.TextureChannel == EDMTextureChannelMask::RGBA
+			? FDMMaterialStageConnectorChannel::WHOLE_CHANNEL
+			: (EnumHasAnyFlags(MaterialTexture.Value.TextureChannel, EDMTextureChannelMask::Red) ? FDMMaterialStageConnectorChannel::FIRST_CHANNEL : 0)
+				+ (EnumHasAnyFlags(MaterialTexture.Value.TextureChannel, EDMTextureChannelMask::Green) ? FDMMaterialStageConnectorChannel::FIRST_CHANNEL : 0)
+				+ (EnumHasAnyFlags(MaterialTexture.Value.TextureChannel, EDMTextureChannelMask::Blue) ? FDMMaterialStageConnectorChannel::FIRST_CHANNEL : 0)
+				+ (EnumHasAnyFlags(MaterialTexture.Value.TextureChannel, EDMTextureChannelMask::Alpha) ? FDMMaterialStageConnectorChannel::FIRST_CHANNEL : 0);
+
+		UDMMaterialLayerObject* Layer = nullptr;
+
+		{
+			const FDMUpdateGuard Guard;
+
+			Layer = Slot->AddDefaultLayer(PropertyType);
+
+			if (!ensure(Layer))
+			{
+				continue;
+			}
+
+			bMadeChange = true;
+
+			UDMMaterialStage* Stage = Layer->GetStage(EDMMaterialLayerStage::Base);
+
+			if (!ensure(Stage))
+			{
+				continue;
+			}
+
+			UDMMaterialStageInputExpression* NewExpression = UDMMaterialStageInputExpression::ChangeStageInput_Expression(
+				Stage,
+				UDMMaterialStageExpressionTextureSample::StaticClass(),
+				UDMMaterialStageBlend::InputB,
+				FDMMaterialStageConnectorChannel::WHOLE_CHANNEL,
+				0,
+				OutputMask
+			);
+
+			if (!ensure(NewExpression))
+			{
+				continue;
+			}
+
+			UDMMaterialSubStage* SubStage = NewExpression->GetSubStage();
+
+			if (!ensure(SubStage))
+			{
+				continue;
+			}
+
+			UDMMaterialStageInputValue* InputValue = UDMMaterialStageInputValue::ChangeStageInput_NewLocalValue(
+				SubStage,
+				0,
+				FDMMaterialStageConnectorChannel::WHOLE_CHANNEL,
+				EDMValueType::VT_Texture,
+				FDMMaterialStageConnectorChannel::WHOLE_CHANNEL
+			);
+
+			if (ensure(InputValue))
+			{
+				UDMMaterialValueTexture* InputTexture = Cast<UDMMaterialValueTexture>(InputValue->GetValue());
+
+				if (ensure(InputTexture))
+				{
+
+					InputTexture->SetValue(Texture);
+				}
+			}
+
+			if (bInReplaceSlots)
+			{
+				for (int32 Index = Slot->GetLayers().Num() - 1; Index >= 0; --Index)
+				{
+					UDMMaterialLayerObject* LayerIter = Slot->GetLayer(Index);
+
+					if (!LayerIter || LayerIter->GetStage(EDMMaterialLayerStage::Base) == Stage)
+					{
+						continue;
+					}
+
+					Slot->RemoveLayer(LayerIter);
+				}
+			}
+		}
+
+		Layer->Update(EDMUpdateType::Structure);
+	}
+
+	return bMadeChange;
 }
 
 bool UDynamicMaterialModelEditorOnlyData::NeedsWizard() const
diff --git a/Engine/Plugins/Experimental/DynamicMaterial/Source/DynamicMaterialEditor/Private/Slate/Layers/SDMSlotLayerView.cpp b/Engine/Plugins/Experimental/DynamicMaterial/Source/DynamicMaterialEditor/Private/Slate/Layers/SDMSlotLayerView.cpp
--- a/Engine/Plugins/Experimental/DynamicMaterial/Source/DynamicMaterialEditor/Private/Slate/Layers/SDMSlotLayerView.cpp
+++ b/Engine/Plugins/Experimental/DynamicMaterial/Source/DynamicMaterialEditor/Private/Slate/Layers/SDMSlotLayerView.cpp
@@ -1,6 +1,7 @@
 // Copyright Epic Games, Inc. All Rights Reserved.
 
 #include "Slate/Layers/SDMSlotLayerView.h"
+
 #include "Components/DMMaterialEffect.h"
 #include "Components/DMMaterialEffectStack.h"
 #include "Components/DMMaterialLayer.h"
@@ -13,6 +14,8 @@
 #include "Components/MaterialStageInputs/DMMSIFunction.h"
 #include "Components/MaterialStageInputs/DMMSIValue.h"
 #include "Components/MaterialValues/DMMaterialValueTexture.h"
+#include "DMContentBrowserIntegration.h"
+#include "DMTextureSet.h"
 #include "DragAndDrop/AssetDragDropOp.h"
 #include "DynamicMaterialEditorCommands.h"
 #include "DynamicMaterialEditorModule.h"
@@ -20,6 +23,8 @@
 #include "Framework/Application/SlateApplication.h"
 #include "Materials/MaterialFunctionInterface.h"
 #include "Menus/DMMaterialSlotLayerMenus.h"
+#include "Misc/MessageDialog.h"
+#include "Model/DynamicMaterialModelEditorOnlyData.h"
 #include "SDMLayerEffectsItem.h"
 #include "Slate/Layers/SDMSlotLayerItem.h"
 #include "Slate/SDMSlot.h"
@@ -658,6 +663,46 @@ void SDMSlotLayerView::HandleAssetDragDropOperation(FAssetDragDropOp& InAssetDra
 		return;
 	}
 
+	TArray<UTexture*> TexturesDropped;
+	TexturesDropped.Reserve(Assets.Num());
+
+	for (const FAssetData& Asset : Assets)
+	{
+		UClass* AssetClass = Asset.GetClass(EResolveClass::Yes);
+
+		if (AssetClass && AssetClass->IsChildOf(UTexture::StaticClass()))
+		{
+			if (UTexture* Texture = Cast<UTexture>(Asset.GetAsset()))
+			{
+				TexturesDropped.Add(Texture);
+			}
+		}
+	}
+
+	if (TexturesDropped.Num() > 1)
+	{
+		const EAppReturnType::Type Result = FMessageDialog::Open(
+			EAppMsgType::YesNoCancel,
+			LOCTEXT("ReplaceSlotsMultiTexture", "You are about to import multiple Textures via a Material Designer Texture Set.\n\nDo you want to replace the slot contents?\n- Yes: All layers are deleted in the matching slots.\n- No: New texture layers are added to the matching slots.\n- Cancel: Abort this operation.")
+		);
+
+		switch (Result)
+		{
+			case EAppReturnType::No:
+				FDMContentBrowserIntegration::UpdateMaterialDesignerInstanceFromTextureSet(Assets, /* Replace */ false);
+				break;
+
+			case EAppReturnType::Yes:
+				FDMContentBrowserIntegration::UpdateMaterialDesignerInstanceFromTextureSet(Assets, /* Replace */ true);
+				break;
+
+			default:
+				break;
+		}
+
+		return;
+	}
+
 	UObject* Asset = Assets[0].GetAsset();
 
 	if (!IsValid(Asset))
@@ -734,6 +779,30 @@ void SDMSlotLayerView::HandleAssetDragDropOperation(FAssetDragDropOp& InAssetDra
 			}
 		}
 	}
+	else if (UDMTextureSet* TextureSet = Cast<UDMTextureSet>(Asset))
+	{
+		if (UDynamicMaterialModelEditorOnlyData* EditorOnlyData = Slot->GetMaterialModelEditorOnlyData())
+		{
+			const EAppReturnType::Type Result = FMessageDialog::Open(
+				EAppMsgType::YesNoCancel,
+				LOCTEXT("ReplaceSlotsTextureSet", "You are about to import a Material Designer Texture Set.\n\nDo you want to replace the slot contents?\n- Yes: All layers are deleted in the matching slots.\n- No: New texture layers are added to the matching slots.\n- Cancel: Abort this operation.")
+			);
+
+			switch (Result)
+			{
+				case EAppReturnType::No:
+					EditorOnlyData->AddTextureSet(TextureSet, /* Replace */ false);
+					break;
+
+				case EAppReturnType::Yes:
+					EditorOnlyData->AddTextureSet(TextureSet, /* Replace */ true);
+					break;
+
+				default:
+					break;
+			}
+		}
+	}
 }
 
 void SDMSlotLayerView::RequestListRefresh()
diff --git a/Engine/Plugins/Experimental/DynamicMaterial/Source/DynamicMaterialEditor/Private/Slate/SDMSlot.cpp b/Engine/Plugins/Experimental/DynamicMaterial/Source/DynamicMaterialEditor/Private/Slate/SDMSlot.cpp
--- a/Engine/Plugins/Experimental/DynamicMaterial/Source/DynamicMaterialEditor/Private/Slate/SDMSlot.cpp
+++ b/Engine/Plugins/Experimental/DynamicMaterial/Source/DynamicMaterialEditor/Private/Slate/SDMSlot.cpp
@@ -958,7 +958,7 @@ void SDMSlot::OnSlotLayersUpdated(UDMMaterialSlot* InSlot)
 		return;
 	}
 
-	InvalidateComponentEditWidget();
+	InvalidateMainWidget();
 }
 
 void SDMSlot::OnSlotPropertiesUpdated(UDMMaterialSlot* InSlot)
diff --git a/Engine/Plugins/Experimental/DynamicMaterial/Source/DynamicMaterialTextureSet/Private/DMTextureSet.cpp b/Engine/Plugins/Experimental/DynamicMaterial/Source/DynamicMaterialTextureSet/Private/DMTextureSet.cpp
--- a/Engine/Plugins/Experimental/DynamicMaterial/Source/DynamicMaterialTextureSet/Private/DMTextureSet.cpp
+++ b/Engine/Plugins/Experimental/DynamicMaterial/Source/DynamicMaterialTextureSet/Private/DMTextureSet.cpp
@@ -2,39 +2,39 @@
 
 #include "DMTextureSet.h"
 
+#include "DMTextureSetMaterialProperty.h"
 #include "Engine/Texture.h"
-#include "SceneTypes.h"
 
 UDMTextureSet::UDMTextureSet()
 {
-	Textures.Reserve(13); // EMaterialProperty has more than 13 entries.
-
-	Textures.Emplace(EMaterialProperty::MP_EmissiveColor);
-	Textures.Emplace(EMaterialProperty::MP_Opacity);
-	Textures.Emplace(EMaterialProperty::MP_OpacityMask);
-	Textures.Emplace(EMaterialProperty::MP_BaseColor);
-	Textures.Emplace(EMaterialProperty::MP_Metallic);
-	Textures.Emplace(EMaterialProperty::MP_Specular);
-	Textures.Emplace(EMaterialProperty::MP_Roughness);
-	Textures.Emplace(EMaterialProperty::MP_Anisotropy);
-	Textures.Emplace(EMaterialProperty::MP_Normal);
-	Textures.Emplace(EMaterialProperty::MP_Tangent);
-	Textures.Emplace(EMaterialProperty::MP_SubsurfaceColor);
-	Textures.Emplace(EMaterialProperty::MP_AmbientOcclusion);
-	Textures.Emplace(EMaterialProperty::MP_Refraction);
+	UEnum* MaterialPropertyEnum = StaticEnum<EDMTextureSetMaterialProperty>();
+
+	Textures.Reserve(MaterialPropertyEnum->NumEnums());
+
+	for (int32 Index = 0; Index < MaterialPropertyEnum->NumEnums(); ++Index)
+	{
+		const EDMTextureSetMaterialProperty MaterialProperty = static_cast<EDMTextureSetMaterialProperty>(MaterialPropertyEnum->GetValueByIndex(Index));
+
+		if (MaterialProperty == EDMTextureSetMaterialProperty::None)
+		{
+			break;
+		}
+
+		Textures.Emplace(MaterialProperty);
+	}
 }
 
-bool UDMTextureSet::HasMaterialProperty(TEnumAsByte<EMaterialProperty> InMaterialProperty) const
+bool UDMTextureSet::HasMaterialProperty(EDMTextureSetMaterialProperty InMaterialProperty) const
 {
 	return Textures.Contains(InMaterialProperty);
 }
 
-const TMap<TEnumAsByte<EMaterialProperty>, FDMMaterialTexture>& UDMTextureSet::GetTextures() const
+const TMap<EDMTextureSetMaterialProperty, FDMMaterialTexture>& UDMTextureSet::GetTextures() const
 {
 	return Textures;
 }
 
-bool UDMTextureSet::HasMaterialTexture(TEnumAsByte<EMaterialProperty> InMaterialProperty) const
+bool UDMTextureSet::HasMaterialTexture(EDMTextureSetMaterialProperty InMaterialProperty) const
 {
 	if (const FDMMaterialTexture* MaterialTexture = Textures.Find(InMaterialProperty))
 	{
@@ -44,7 +44,7 @@ bool UDMTextureSet::HasMaterialTexture(TEnumAsByte<EMaterialProperty> InMaterial
 	return false;
 }
 
-bool UDMTextureSet::GetMaterialTexture(TEnumAsByte<EMaterialProperty> InMaterialProperty, FDMMaterialTexture& OutMaterialTexture) const
+bool UDMTextureSet::GetMaterialTexture(EDMTextureSetMaterialProperty InMaterialProperty, FDMMaterialTexture& OutMaterialTexture) const
 {
 	if (const FDMMaterialTexture* MaterialTexture = Textures.Find(InMaterialProperty))
 	{
@@ -55,12 +55,12 @@ bool UDMTextureSet::GetMaterialTexture(TEnumAsByte<EMaterialProperty> InMaterial
 	return false;
 }
 
-const FDMMaterialTexture* UDMTextureSet::GetMaterialTexture(TEnumAsByte<EMaterialProperty> InMaterialProperty) const
+const FDMMaterialTexture* UDMTextureSet::GetMaterialTexture(EDMTextureSetMaterialProperty InMaterialProperty) const
 {
 	return Textures.Find(InMaterialProperty);
 }
 
-void UDMTextureSet::SetMaterialTexture(TEnumAsByte<EMaterialProperty> InMaterialProperty, const FDMMaterialTexture& InMaterialTexture)
+void UDMTextureSet::SetMaterialTexture(EDMTextureSetMaterialProperty InMaterialProperty, const FDMMaterialTexture& InMaterialTexture)
 {
 	if (FDMMaterialTexture* MaterialTexture = Textures.Find(InMaterialProperty))
 	{
@@ -75,9 +75,9 @@ bool UDMTextureSet::ContainsTexture(UTexture* InTexture) const
 		return false;
 	}
 
-	for (const TPair<TEnumAsByte<EMaterialProperty>, FDMMaterialTexture>& Pair : Textures)
+	for (const TPair<EDMTextureSetMaterialProperty, FDMMaterialTexture>& TexturePair : Textures)
 	{
-		if (Pair.Value.Texture == InTexture)
+		if (TexturePair.Value.Texture == InTexture)
 		{
 			return true;
 		}
diff --git a/Engine/Plugins/Experimental/DynamicMaterial/Source/DynamicMaterialTextureSet/Public/DMTextureSet.h b/Engine/Plugins/Experimental/DynamicMaterial/Source/DynamicMaterialTextureSet/Public/DMTextureSet.h
--- a/Engine/Plugins/Experimental/DynamicMaterial/Source/DynamicMaterialTextureSet/Public/DMTextureSet.h
+++ b/Engine/Plugins/Experimental/DynamicMaterial/Source/DynamicMaterialTextureSet/Public/DMTextureSet.h
@@ -10,7 +10,7 @@
 
 #include "DMTextureSet.generated.h"
 
-enum EMaterialProperty : int;
+enum class EDMTextureSetMaterialProperty : uint8;
 
 UCLASS(BlueprintType, ClassGroup = "Material Designer")
 class DYNAMICMATERIALTEXTURESET_API UDMTextureSet : public UObject
@@ -29,12 +29,12 @@ public:
 	 * @return True if the property exists in the Texture Map.
 	 */
 	UFUNCTION(BlueprintPure, Category = "Material Designer")
-	bool HasMaterialProperty(TEnumAsByte<EMaterialProperty> InMaterialProperty) const;
+	bool HasMaterialProperty(EDMTextureSetMaterialProperty InMaterialProperty) const;
 
 	/**
 	 * @return Gets the entire Texture Map.
 	 */
-	const TMap<TEnumAsByte<EMaterialProperty>, FDMMaterialTexture>& GetTextures() const;
+	const TMap<EDMTextureSetMaterialProperty, FDMMaterialTexture>& GetTextures() const;
 
 	/**
 	 * Checks whether a given Material Property has a Texture assigned to it.
@@ -42,7 +42,7 @@ public:
 	 * @return True if the Material Property has an assigned Texture.
 	 */
 	UFUNCTION(BlueprintPure, Category = "Material Designer")
-	bool HasMaterialTexture(TEnumAsByte<EMaterialProperty> InMaterialProperty) const;
+	bool HasMaterialTexture(EDMTextureSetMaterialProperty InMaterialProperty) const;
 
 	/**
 	 * Gets the Material Texture associated with a Material Property. Does not check whether a Texture
@@ -52,7 +52,7 @@ public:
 	 * @return True if the Material Property exists within the Texture Map.
 	 */
 	UFUNCTION(BlueprintPure, Category = "Material Designer")
-	bool GetMaterialTexture(TEnumAsByte<EMaterialProperty> InMaterialProperty, FDMMaterialTexture& OutMaterialTexture) const;
+	bool GetMaterialTexture(EDMTextureSetMaterialProperty InMaterialProperty, FDMMaterialTexture& OutMaterialTexture) const;
 
 	/**
 	 * Gets the Material Texture associated with a Material Property. Does not check whether a Texture
@@ -60,7 +60,7 @@ public:
 	 * @param InMaterialProperty The Material Property to check.
 	 * @return A pointer to the found Material Texture, if it exists, or nullptr.
 	 */
-	const FDMMaterialTexture* GetMaterialTexture(TEnumAsByte<EMaterialProperty> InMaterialProperty) const;
+	const FDMMaterialTexture* GetMaterialTexture(EDMTextureSetMaterialProperty InMaterialProperty) const;
 
 	/**
 	 * Sets the Material Texture for a given Material Property. Can be used to unset Textures.
@@ -68,7 +68,7 @@ public:
 	 * @param InMaterialTexture The Material Texture to set on the given Material Property.
 	 */
 	UFUNCTION(BlueprintCallable, Category = "Material Designer")
-	void SetMaterialTexture(TEnumAsByte<EMaterialProperty> InMaterialProperty, const FDMMaterialTexture& InMaterialTexture);
+	void SetMaterialTexture(EDMTextureSetMaterialProperty InMaterialProperty, const FDMMaterialTexture& InMaterialTexture);
 
 	/**
 	 * Checks the Texture Map to see if a given Texture exists within it.
@@ -80,5 +80,5 @@ public:
 
 protected:
 	UPROPERTY(EditAnywhere, BlueprintReadWrite, Category = "Material Designer", EditFixedSize, meta = (ReadOnlyKeys, AllowPrivateAccess))
-	TMap<TEnumAsByte<EMaterialProperty>, FDMMaterialTexture> Textures;
+	TMap<EDMTextureSetMaterialProperty, FDMMaterialTexture> Textures;
 };
diff --git a/Engine/Plugins/Experimental/DynamicMaterial/Source/DynamicMaterialTextureSet/Public/DMTextureSetMaterialProperty.h b/Engine/Plugins/Experimental/DynamicMaterial/Source/DynamicMaterialTextureSet/Public/DMTextureSetMaterialProperty.h
new file mode 100644
--- /dev/null
+++ b/Engine/Plugins/Experimental/DynamicMaterial/Source/DynamicMaterialTextureSet/Public/DMTextureSetMaterialProperty.h
@@ -0,0 +1,35 @@
+// Copyright Epic Games, Inc. All Rights Reserved.
+
+#pragma once
+
+#include "HAL/Platform.h"
+#include "SceneTypes.h"
+#include "UObject/ObjectMacros.h"
+
+/**
+ * Copying the values from EMaterialProperty to change required meta data,
+ * such as DisplayName and Hidden.
+ */
+UENUM(BlueprintType)
+enum class EDMTextureSetMaterialProperty : uint8
+{
+	BaseColor = EMaterialProperty::MP_BaseColor,
+	Normal = EMaterialProperty::MP_Normal,
+	Metallic = EMaterialProperty::MP_Metallic,
+	Roughness = EMaterialProperty::MP_Roughness,
+	AmbientOcclusion = EMaterialProperty::MP_AmbientOcclusion,
+	Displacement = EMaterialProperty::MP_Displacement,
+	Specular = EMaterialProperty::MP_Specular,
+	SubsurfaceColor = EMaterialProperty::MP_SubsurfaceColor,
+
+	EmissiveColor = EMaterialProperty::MP_EmissiveColor,
+	Opacity = EMaterialProperty::MP_Opacity,
+	OpacityMask = EMaterialProperty::MP_OpacityMask,
+	Anisotropy = EMaterialProperty::MP_Anisotropy,
+	Refraction = EMaterialProperty::MP_Refraction,
+	Tangent = EMaterialProperty::MP_Tangent,
+	WorldPositionOffset = EMaterialProperty::MP_WorldPositionOffset,
+	PixelDepthOffset = EMaterialProperty::MP_PixelDepthOffset,
+
+	None = EMaterialProperty::MP_MAX UMETA(Hidden)
+};
diff --git a/Engine/Plugins/Experimental/DynamicMaterial/Source/DynamicMaterialTextureSetEditor/DynamicMaterialTextureSetEditor.Build.cs b/Engine/Plugins/Experimental/DynamicMaterial/Source/DynamicMaterialTextureSetEditor/DynamicMaterialTextureSetEditor.Build.cs
--- a/Engine/Plugins/Experimental/DynamicMaterial/Source/DynamicMaterialTextureSetEditor/DynamicMaterialTextureSetEditor.Build.cs
+++ b/Engine/Plugins/Experimental/DynamicMaterial/Source/DynamicMaterialTextureSetEditor/DynamicMaterialTextureSetEditor.Build.cs
@@ -22,8 +22,6 @@ public class DynamicMaterialTextureSetEditor : ModuleRules
 				"AssetDefinition",
 				"Core",
 				"DeveloperSettings",
-				"DynamicMaterial",
-				"DynamicMaterialEditor",
 				"DynamicMaterialTextureSet",
 				"EditorWidgets",
 				"Engine",
diff --git a/Engine/Plugins/Experimental/DynamicMaterial/Source/DynamicMaterialTextureSetEditor/Private/DMTextureSetBlueprintFunctionLibrary.cpp b/Engine/Plugins/Experimental/DynamicMaterial/Source/DynamicMaterialTextureSetEditor/Private/DMTextureSetBlueprintFunctionLibrary.cpp
--- a/Engine/Plugins/Experimental/DynamicMaterial/Source/DynamicMaterialTextureSetEditor/Private/DMTextureSetBlueprintFunctionLibrary.cpp
+++ b/Engine/Plugins/Experimental/DynamicMaterial/Source/DynamicMaterialTextureSetEditor/Private/DMTextureSetBlueprintFunctionLibrary.cpp
@@ -71,7 +71,7 @@ UDMTextureSet* UDMTextureSetBlueprintFunctionLibrary::CreateTextureSetFromAssets
 				continue;
 			}
 
-			for (const TPair<TEnumAsByte<EMaterialProperty>, EDMTextureChannelMask>& Pair : Filter.MaterialProperties)
+			for (const TPair<EDMTextureSetMaterialProperty, EDMTextureChannelMask>& Pair : Filter.MaterialProperties)
 			{
 				if (TextureSet->HasMaterialTexture(Pair.Key))
 				{
@@ -101,7 +101,7 @@ void UDMTextureSetBlueprintFunctionLibrary::CreateTextureSetFromAssetsInteractiv
 		.SizingRule(ESizingRule::Autosized)
 		.SupportsMinimize(false)
 		.SupportsMaximize(false)
-		.Title(LOCTEXT("TextureSetBuilder", "Texture Set Builder"))
+		.Title(LOCTEXT("TextureSetBuilder", "Material Designer Texture Set Builder"))
 		[
 			SNew(SDMTextureSetBuilder, TextureSet, InAssets, InOnComplete)
 		];
diff --git a/Engine/Plugins/Experimental/DynamicMaterial/Source/DynamicMaterialTextureSetEditor/Private/DMTextureSetContentBrowserIntegration.cpp b/Engine/Plugins/Experimental/DynamicMaterial/Source/DynamicMaterialTextureSetEditor/Private/DMTextureSetContentBrowserIntegration.cpp
--- a/Engine/Plugins/Experimental/DynamicMaterial/Source/DynamicMaterialTextureSetEditor/Private/DMTextureSetContentBrowserIntegration.cpp
+++ b/Engine/Plugins/Experimental/DynamicMaterial/Source/DynamicMaterialTextureSetEditor/Private/DMTextureSetContentBrowserIntegration.cpp
@@ -2,311 +2,9 @@
 
 #include "DMTextureSetContentBrowserIntegration.h"
 
-#include "AssetRegistry/AssetData.h"
-#include "AssetRegistry/AssetRegistryModule.h"
-#include "AssetToolsModule.h"
-#include "ContentBrowserModule.h"
-#include "CoreGlobals.h"
-#include "DMTextureSet.h"
-#include "DMTextureSetBlueprintFunctionLibrary.h"
-#include "DMTextureSetStyle.h"
-#include "Engine/Texture.h"
-#include "Framework/MultiBox/MultiBoxBuilder.h"
-#include "IDynamicMaterialEditorModule.h"
-#include "Material/DynamicMaterialInstance.h"
-#include "Material/DynamicMaterialInstanceFactory.h"
-#include "Materials/Material.h"
-#include "Model/DynamicMaterialModelEditorOnlyData.h"
-#include "ScopedTransaction.h"
-#include "Styling/SlateIconFinder.h"
+FDMTextureSetContentBrowserIntegration::FOnPopulateMenu FDMTextureSetContentBrowserIntegration::PopulateMenuDelegate;
 
-#define LOCTEXT_NAMESPACE "FDMTextureSetContentBrowserIntegration"
-
-FDelegateHandle FDMTextureSetContentBrowserIntegration::ExtenderDelegateHandle;
-
-void FDMTextureSetContentBrowserIntegration::Integrate()
-{
-	FDMTextureSetStyle::Get();
-
-	FContentBrowserModule& ContentBrowserModule = FModuleManager::LoadModuleChecked<FContentBrowserModule>("ContentBrowser");
-	TArray<FContentBrowserMenuExtender_SelectedAssets>& CBMenuExtenderDelegates = ContentBrowserModule.GetAllAssetViewContextMenuExtenders();
-	CBMenuExtenderDelegates.Add(FContentBrowserMenuExtender_SelectedAssets::CreateStatic(&FDMTextureSetContentBrowserIntegration::OnExtendContentBrowserAssetSelectionMenu));
-	ExtenderDelegateHandle = CBMenuExtenderDelegates.Last().GetHandle();
-}
-
-void FDMTextureSetContentBrowserIntegration::Disintegrate()
-{
-	if (!ExtenderDelegateHandle.IsValid())
-	{
-		return;
-	}
-
-	if (FContentBrowserModule* ContentBrowserModule = FModuleManager::GetModulePtr<FContentBrowserModule>("ContentBrowser"))
-	{
-		TArray<FContentBrowserMenuExtender_SelectedAssets>& CBMenuExtenderDelegates = ContentBrowserModule->GetAllAssetViewContextMenuExtenders();
-
-		CBMenuExtenderDelegates.RemoveAll(
-			[](const FContentBrowserMenuExtender_SelectedAssets& InElement)
-			{
-				return InElement.GetHandle() == ExtenderDelegateHandle;
-			}
-		);
-	}
-
-	ExtenderDelegateHandle.Reset();
-}
-
-TSharedRef<FExtender> FDMTextureSetContentBrowserIntegration::OnExtendContentBrowserAssetSelectionMenu(const TArray<FAssetData>& InSelectedAssets)
-{
-	TSharedRef<FExtender> Extender = MakeShared<FExtender>();
-	bool bHasTexture = false;
-
-	for (const FAssetData& SelectedAsset : InSelectedAssets)
-	{
-		if (UClass* AssetClass = SelectedAsset.GetClass(EResolveClass::Yes))
-		{
-			if (AssetClass->IsChildOf<UTexture>())
-			{
-				bHasTexture = true;
-				break;
-			}
-		}
-	}
-
-	if (!bHasTexture)
-	{
-		return Extender;
-	}
-
-	Extender->AddMenuExtension(
-		"GetAssetActions",
-		EExtensionHook::After,
-		nullptr,
-		FMenuExtensionDelegate::CreateLambda(
-			[InSelectedAssets](FMenuBuilder& InMenuBuilder) {
-				InMenuBuilder.AddMenuEntry(
-					LOCTEXT("CreateTextureSet", "Create Texture Set"),
-					LOCTEXT("CreateTextureSetTooltip", "Creates an asset listing a group of textures and the material properties they are associated with."),
-					FSlateIconFinder::FindIconForClass(UTexture::StaticClass()),
-					FUIAction(FExecuteAction::CreateStatic(&FDMTextureSetContentBrowserIntegration::CreateTextureSet, InSelectedAssets))
-				);
-			})
-	);
-
-#if 0
-	Extender->AddMenuExtension(
-		"GetAssetActions",
-		EExtensionHook::After,
-		nullptr,
-		FMenuExtensionDelegate::CreateLambda(
-			[InSelectedAssets](FMenuBuilder& InMenuBuilder) {
-				InMenuBuilder.AddMenuEntry(
-					LOCTEXT("CreateMaterialDesignerInstanceFromTextureSet", "Create Material Designer Instance"),
-					LOCTEXT("CreateMaterialDesignerInstanceFromTextureSetTooltip", "Creates a Material Designer Instance in the content browser using a Texture Set."),
-					FSlateIconFinder::FindIconForClass(UMaterial::StaticClass()),
-					FUIAction(FExecuteAction::CreateStatic(&FDMTextureSetContentBrowserIntegration::CreateMaterialDesignerInstanceFromTextureSet, InSelectedAssets))
-				);
-			})
-	);
-
-	Extender->AddMenuExtension(
-		"GetAssetActions",
-		EExtensionHook::After,
-		nullptr,
-		FMenuExtensionDelegate::CreateLambda(
-			[InSelectedAssets](FMenuBuilder& InMenuBuilder) {
-				InMenuBuilder.AddMenuEntry(
-					LOCTEXT("UpdateMaterialDesignerInstanceFromTextureSetAdd", "Update Material Designer Instance (Add)"),
-					LOCTEXT("UpdateMaterialDesignerInstanceFromTextureAddSetTooltip", "Updates the opened Material Designer Instance using a Texture Set, adding new layers to the Model."),
-					FSlateIconFinder::FindIconForClass(UMaterial::StaticClass()),
-					FUIAction(FExecuteAction::CreateStatic(&FDMTextureSetContentBrowserIntegration::UpdateMaterialDesignerInstanceFromTextureSet, InSelectedAssets, /* Replace */ false))
-				);
-			})
-	);
-
-	Extender->AddMenuExtension(
-		"GetAssetActions",
-		EExtensionHook::After,
-		nullptr,
-		FMenuExtensionDelegate::CreateLambda(
-			[InSelectedAssets](FMenuBuilder& InMenuBuilder) {
-				InMenuBuilder.AddMenuEntry(
-					LOCTEXT("UpdateMaterialDesignerInstanceFromTextureSetReplace", "Update Material Designer Instance (Replace)"),
-					LOCTEXT("UpdateMaterialDesignerInstanceFromTextureSetReplaceTooltip", "Updates the opened Material Designer Instance using a Texture Set, replacing slots in the Model."),
-					FSlateIconFinder::FindIconForClass(UMaterial::StaticClass()),
-					FUIAction(FExecuteAction::CreateStatic(&FDMTextureSetContentBrowserIntegration::UpdateMaterialDesignerInstanceFromTextureSet, InSelectedAssets, /* Replace */ true))
-				);
-			})
-	);
-#endif
-
-	return Extender;
-}
-
-void FDMTextureSetContentBrowserIntegration::CreateTextureSet(TArray<FAssetData> InSelectedAssets)
-{
-	if (InSelectedAssets.IsEmpty())
-	{
-		return;
-	}
-
-	UDMTextureSetBlueprintFunctionLibrary::CreateTextureSetFromAssetsInteractive(
-		InSelectedAssets,
-		FDMTextureSetBuilderOnComplete::CreateStatic(
-			&FDMTextureSetContentBrowserIntegration::OnCreateTextureSetComplete,
-			InSelectedAssets[0].PackagePath.ToString()
-		)
-	);
-}
-
-void FDMTextureSetContentBrowserIntegration::OnCreateTextureSetComplete(UDMTextureSet* InTextureSet, bool bInAccepted, FString InPath)
-{
-	if (!InTextureSet || !bInAccepted)
-	{
-		return;
-	}
-
-	const FAssetToolsModule& AssetToolsModule = FModuleManager::LoadModuleChecked<FAssetToolsModule>("AssetTools");
-
-	FString UniquePackageName;
-	FString UniqueAssetName;
-
-	const FString BasePackageName = InPath / TEXT("NewTextureSet");
-	AssetToolsModule.Get().CreateUniqueAssetName(BasePackageName, TEXT(""), UniquePackageName, UniqueAssetName);
-
-	UPackage* Package = CreatePackage(*UniquePackageName);
-
-	if (!Package)
-	{
-		return;
-	}
-
-	InTextureSet->SetFlags(RF_Standalone | RF_Public);
-	InTextureSet->Rename(*UniqueAssetName, Package, REN_DontCreateRedirectors);
-
-	FAssetRegistryModule::AssetCreated(InTextureSet);
-}
-
-void FDMTextureSetContentBrowserIntegration::CreateMaterialDesignerInstanceFromTextureSet(TArray<FAssetData> InSelectedAssets)
-{
-	if (InSelectedAssets.IsEmpty())
-	{
-		return;
-	}
-
-	UDMTextureSetBlueprintFunctionLibrary::CreateTextureSetFromAssetsInteractive(
-		InSelectedAssets,
-		FDMTextureSetBuilderOnComplete::CreateStatic(
-			&FDMTextureSetContentBrowserIntegration::OnCreateMaterialDesignerInstanceFromTextureSetComplete,
-			InSelectedAssets[0].PackagePath.ToString()
-		)
-	);
-}
-
-void FDMTextureSetContentBrowserIntegration::OnCreateMaterialDesignerInstanceFromTextureSetComplete(UDMTextureSet* InTextureSet, bool bInAccepted,
-	FString InPath)
-{
-	if (!InTextureSet || !bInAccepted)
-	{
-		return;
-	}
-
-	UDynamicMaterialInstance* Instance = Cast<UDynamicMaterialInstance>(GetMutableDefault<UDynamicMaterialInstanceFactory>()->FactoryCreateNew(
-		UDynamicMaterialInstance::StaticClass(),
-		GetTransientPackage(),
-		NAME_None,
-		RF_Transient,
-		/* Context */ nullptr,
-		GWarn
-	));
-
-	if (!Instance)
-	{
-		return;
-	}
-
-	UDynamicMaterialModelEditorOnlyData* EditorOnlyData = UDynamicMaterialModelEditorOnlyData::Get(Instance);
-
-	if (!EditorOnlyData)
-	{
-		return;
-	}
-
-	if (!EditorOnlyData->AddTextureSet(InTextureSet, /* Replace */ true))
-	{
-		return;
-	}
-
-	const FAssetToolsModule& AssetToolsModule = FModuleManager::LoadModuleChecked<FAssetToolsModule>("AssetTools");
-
-	FString UniquePackageName;
-	FString UniqueAssetName;
-
-	const FString BasePackageName = InPath / TEXT("NewMaterialDesignerInstance");
-	AssetToolsModule.Get().CreateUniqueAssetName(BasePackageName, TEXT(""), UniquePackageName, UniqueAssetName);
-
-	UPackage* Package = CreatePackage(*UniquePackageName);
-
-	if (!Package)
-	{
-		return;
-	}
-
-	Instance->SetFlags(RF_Standalone | RF_Public);
-	Instance->Rename(*UniqueAssetName, Package, REN_DontCreateRedirectors);
-
-	FAssetRegistryModule::AssetCreated(Instance);
-}
-
-void FDMTextureSetContentBrowserIntegration::UpdateMaterialDesignerInstanceFromTextureSet(TArray<FAssetData> InSelectedAssets, bool bInReplace)
+TMulticastDelegate<void(FMenuBuilder&, const TArray<FAssetData>&)>::RegistrationType& FDMTextureSetContentBrowserIntegration::GetPopulateExtenderDelegate()
 {
-	if (InSelectedAssets.IsEmpty())
-	{
-		return;
-	}
-
-	UDMTextureSetBlueprintFunctionLibrary::CreateTextureSetFromAssetsInteractive(
-		InSelectedAssets,
-		FDMTextureSetBuilderOnComplete::CreateStatic(
-			&FDMTextureSetContentBrowserIntegration::OnUpdateMaterialDesignerInstanceFromTextureSetComplete,
-			bInReplace
-		)
-	);
+	return PopulateMenuDelegate;
 }
-
-void FDMTextureSetContentBrowserIntegration::OnUpdateMaterialDesignerInstanceFromTextureSetComplete(UDMTextureSet* InTextureSet, bool bInAccepted,
-	bool bInReplace)
-{
-	if (!InTextureSet || !bInAccepted)
-	{
-		return;
-	}
-
-	IDynamicMaterialEditorModule& DynamicMaterialEditorModule = IDynamicMaterialEditorModule::Get();
-
-	UDynamicMaterialModel* Model = DynamicMaterialEditorModule.GetOpenedMaterialModel(nullptr);
-
-	if (!Model)
-	{
-		return;
-	}
-
-	UDynamicMaterialModelEditorOnlyData* EditorOnlyData = UDynamicMaterialModelEditorOnlyData::Get(Model);
-
-	if (!EditorOnlyData)
-	{
-		return;
-	}
-
-	FScopedTransaction Transaction(LOCTEXT("AddTextureSet", "Material Designer Add Texture Set"));
-	EditorOnlyData->Modify();
-
-	const bool bSuccess = EditorOnlyData->AddTextureSet(InTextureSet, bInReplace);
-
-	if (!bSuccess)
-	{
-		Transaction.Cancel();
-	}
-}
-
-#undef LOCTEXT_NAMESPACE
diff --git a/Engine/Plugins/Experimental/DynamicMaterial/Source/DynamicMaterialTextureSetEditor/Private/DMTextureSetContentBrowserIntegration.h b/Engine/Plugins/Experimental/DynamicMaterial/Source/DynamicMaterialTextureSetEditor/Private/DMTextureSetContentBrowserIntegration.h
deleted file mode 100644
diff --git a/Engine/Plugins/Experimental/DynamicMaterial/Source/DynamicMaterialTextureSetEditor/Private/DMTextureSetContentBrowserIntegrationPrivate.cpp b/Engine/Plugins/Experimental/DynamicMaterial/Source/DynamicMaterialTextureSetEditor/Private/DMTextureSetContentBrowserIntegrationPrivate.cpp
new file mode 100644
--- /dev/null
+++ b/Engine/Plugins/Experimental/DynamicMaterial/Source/DynamicMaterialTextureSetEditor/Private/DMTextureSetContentBrowserIntegrationPrivate.cpp
@@ -0,0 +1,151 @@
+// Copyright Epic Games, Inc. All Rights Reserved.
+
+#include "DMTextureSetContentBrowserIntegrationPrivate.h"
+
+#include "AssetRegistry/AssetRegistryModule.h"
+#include "AssetToolsModule.h"
+#include "ContentBrowserModule.h"
+#include "DMTextureSet.h"
+#include "DMTextureSetBlueprintFunctionLibrary.h"
+#include "DMTextureSetStyle.h"
+#include "Engine/Texture.h"
+#include "Framework/MultiBox/MultiBoxBuilder.h"
+#include "Materials/Material.h"
+#include "Styling/SlateIconFinder.h"
+
+#define LOCTEXT_NAMESPACE "FDMTextureSetContentBrowserIntegrationPrivate"
+
+FDelegateHandle FDMTextureSetContentBrowserIntegrationPrivate::ContentBrowserHandle;
+
+void FDMTextureSetContentBrowserIntegrationPrivate::Integrate()
+{
+	FDMTextureSetStyle::Get();
+
+	FContentBrowserModule& ContentBrowserModule = FModuleManager::LoadModuleChecked<FContentBrowserModule>("ContentBrowser");
+	TArray<FContentBrowserMenuExtender_SelectedAssets>& CBMenuExtenderDelegates = ContentBrowserModule.GetAllAssetViewContextMenuExtenders();
+	CBMenuExtenderDelegates.Add(FContentBrowserMenuExtender_SelectedAssets::CreateStatic(&FDMTextureSetContentBrowserIntegrationPrivate::OnExtendContentBrowserAssetSelectionMenu));
+	ContentBrowserHandle = CBMenuExtenderDelegates.Last().GetHandle();
+}
+
+void FDMTextureSetContentBrowserIntegrationPrivate::Disintegrate()
+{
+	if (!ContentBrowserHandle.IsValid())
+	{
+		return;
+	}
+
+	if (FContentBrowserModule* ContentBrowserModule = FModuleManager::GetModulePtr<FContentBrowserModule>("ContentBrowser"))
+	{
+		TArray<FContentBrowserMenuExtender_SelectedAssets>& CBMenuExtenderDelegates = ContentBrowserModule->GetAllAssetViewContextMenuExtenders();
+
+		CBMenuExtenderDelegates.RemoveAll(
+			[](const FContentBrowserMenuExtender_SelectedAssets& InElement)
+			{
+				return InElement.GetHandle() == ContentBrowserHandle;
+			}
+		);
+	}
+
+	ContentBrowserHandle.Reset();
+}
+
+TSharedRef<FExtender> FDMTextureSetContentBrowserIntegrationPrivate::OnExtendContentBrowserAssetSelectionMenu(const TArray<FAssetData>& InSelectedAssets)
+{
+	TSharedRef<FExtender> Extender = MakeShared<FExtender>();
+	bool bHasTexture = false;
+
+	for (const FAssetData& SelectedAsset : InSelectedAssets)
+	{
+		if (UClass* AssetClass = SelectedAsset.GetClass(EResolveClass::Yes))
+		{
+			if (AssetClass->IsChildOf<UTexture>())
+			{
+				bHasTexture = true;
+				break;
+			}
+		}
+	}
+
+	if (!bHasTexture)
+	{
+		return Extender;
+	}
+
+	Extender->AddMenuExtension(
+		"GetAssetActions",
+		EExtensionHook::After,
+		nullptr,
+		FMenuExtensionDelegate::CreateLambda(
+			[InSelectedAssets](FMenuBuilder& InMenuBuilder)
+			{
+				InMenuBuilder.AddSubMenu(
+					LOCTEXT("TextureSetMenu", "Material Designer Texture Set"),
+					LOCTEXT("TextureSetMenuTooltip", "Create and use Material Designer Texture Sets"),
+					FNewMenuDelegate::CreateStatic(&FDMTextureSetContentBrowserIntegrationPrivate::CreateMenu, InSelectedAssets),
+					/* Close on click */ false,
+					FSlateIconFinder::FindIconForClass(UTexture::StaticClass())
+				);
+			}
+		)
+	);	
+
+	return Extender;
+}
+
+void FDMTextureSetContentBrowserIntegrationPrivate::CreateMenu(FMenuBuilder& InMenuBuilder, TArray<FAssetData> InSelectedAssets)
+{
+	InMenuBuilder.AddMenuEntry(
+		LOCTEXT("CreateTextureSet", "Create From Selected Textures"),
+		LOCTEXT("CreateTextureSetTooltip", "Creates an asset listing a group of textures and the material properties they are associated with."),
+		FSlateIconFinder::FindIconForClass(UTexture::StaticClass()),
+		FUIAction(FExecuteAction::CreateStatic(&FDMTextureSetContentBrowserIntegrationPrivate::CreateTextureSet, InSelectedAssets))
+	);
+
+	PopulateMenuDelegate.Broadcast(InMenuBuilder, InSelectedAssets);
+}
+
+void FDMTextureSetContentBrowserIntegrationPrivate::CreateTextureSet(TArray<FAssetData> InSelectedAssets)
+{
+	if (InSelectedAssets.IsEmpty())
+	{
+		return;
+	}
+
+	UDMTextureSetBlueprintFunctionLibrary::CreateTextureSetFromAssetsInteractive(
+		InSelectedAssets,
+		FDMTextureSetBuilderOnComplete::CreateStatic(
+			&FDMTextureSetContentBrowserIntegrationPrivate::OnCreateTextureSetComplete,
+			InSelectedAssets[0].PackagePath.ToString()
+		)
+	);
+}
+
+void FDMTextureSetContentBrowserIntegrationPrivate::OnCreateTextureSetComplete(UDMTextureSet* InTextureSet, bool bInAccepted, FString InPath)
+{
+	if (!InTextureSet || !bInAccepted)
+	{
+		return;
+	}
+
+	const FAssetToolsModule& AssetToolsModule = FModuleManager::LoadModuleChecked<FAssetToolsModule>("AssetTools");
+
+	FString UniquePackageName;
+	FString UniqueAssetName;
+
+	const FString BasePackageName = InPath / TEXT("NewTextureSet");
+	AssetToolsModule.Get().CreateUniqueAssetName(BasePackageName, TEXT(""), UniquePackageName, UniqueAssetName);
+
+	UPackage* Package = CreatePackage(*UniquePackageName);
+
+	if (!Package)
+	{
+		return;
+	}
+
+	InTextureSet->SetFlags(RF_Standalone | RF_Public);
+	InTextureSet->Rename(*UniqueAssetName, Package, REN_DontCreateRedirectors);
+
+	FAssetRegistryModule::AssetCreated(InTextureSet);
+}
+
+#undef LOCTEXT_NAMESPACE
diff --git a/Engine/Plugins/Experimental/DynamicMaterial/Source/DynamicMaterialTextureSetEditor/Private/DMTextureSetContentBrowserIntegrationPrivate.h b/Engine/Plugins/Experimental/DynamicMaterial/Source/DynamicMaterialTextureSetEditor/Private/DMTextureSetContentBrowserIntegrationPrivate.h
new file mode 100644
--- /dev/null
+++ b/Engine/Plugins/Experimental/DynamicMaterial/Source/DynamicMaterialTextureSetEditor/Private/DMTextureSetContentBrowserIntegrationPrivate.h
@@ -0,0 +1,33 @@
+// Copyright Epic Games, Inc. All Rights Reserved.
+
+#pragma once
+
+#include "DMTextureSetContentBrowserIntegration.h"
+
+#include "AssetRegistry/AssetData.h"
+#include "Containers/ContainersFwd.h"
+#include "Delegates/Delegate.h"
+#include "Delegates/DelegateCombinations.h"
+
+class FExtender;
+class FMenuBuilder;
+class UDMTextureSet;
+
+class FDMTextureSetContentBrowserIntegrationPrivate : public FDMTextureSetContentBrowserIntegration
+{
+public:
+	static void Integrate();
+
+	static void Disintegrate();
+
+protected:
+	static TSharedRef<FExtender> OnExtendContentBrowserAssetSelectionMenu(const TArray<FAssetData>& InSelectedAssets);
+
+	static void CreateMenu(FMenuBuilder& InMenuBuilder, TArray<FAssetData> InSelectedAssets);
+
+	static void CreateTextureSet(TArray<FAssetData> InSelectedAssets);
+
+	static void OnCreateTextureSetComplete(UDMTextureSet* InTextureSet, bool bInAccepted, FString InPath);
+
+	static FDelegateHandle ContentBrowserHandle;
+};
diff --git a/Engine/Plugins/Experimental/DynamicMaterial/Source/DynamicMaterialTextureSetEditor/Private/DMTextureSetFactory.cpp b/Engine/Plugins/Experimental/DynamicMaterial/Source/DynamicMaterialTextureSetEditor/Private/DMTextureSetFactory.cpp
--- a/Engine/Plugins/Experimental/DynamicMaterial/Source/DynamicMaterialTextureSetEditor/Private/DMTextureSetFactory.cpp
+++ b/Engine/Plugins/Experimental/DynamicMaterial/Source/DynamicMaterialTextureSetEditor/Private/DMTextureSetFactory.cpp
@@ -30,7 +30,7 @@ UObject* UDMTextureSetFactory::FactoryCreateNew(UClass* Class, UObject* InParent
 
 FText UDMTextureSetFactory::GetDisplayName() const
 {
-	return LOCTEXT("UDMTextureSet", "Texture Set");
+	return LOCTEXT("UDMTextureSet", "Material Designer Texture Set");
 }
 
 FText UDMTextureSetFactory::GetToolTip() const
diff --git a/Engine/Plugins/Experimental/DynamicMaterial/Source/DynamicMaterialTextureSetEditor/Private/DMTextureSetFilter.cpp b/Engine/Plugins/Experimental/DynamicMaterial/Source/DynamicMaterialTextureSetEditor/Private/DMTextureSetFilter.cpp
--- a/Engine/Plugins/Experimental/DynamicMaterial/Source/DynamicMaterialTextureSetEditor/Private/DMTextureSetFilter.cpp
+++ b/Engine/Plugins/Experimental/DynamicMaterial/Source/DynamicMaterialTextureSetEditor/Private/DMTextureSetFilter.cpp
@@ -2,17 +2,26 @@
 
 #include "DMTextureSetFilter.h"
 
+#include "DMTextureSetSettings.h"
+
 FDMTextureSetFilter::FDMTextureSetFilter()
 	: FilterStrings({TEXT("_")})
-	, MaterialProperties({{EMaterialProperty::MP_BaseColor, EDMTextureChannelMask::RGBA}})
+	, MaterialProperties({{EDMTextureSetMaterialProperty::BaseColor, EDMTextureChannelMask::RGBA}})
 {
 }
 
 bool FDMTextureSetFilter::MatchesFilter(const FString& InAssetName) const
 {
+	bool bOnlyMatchEndOfAssetName = false;
+
+	if (UDMTextureSetSettings* Settings = UDMTextureSetSettings::Get())
+	{
+		bOnlyMatchEndOfAssetName = Settings->bOnlyMatchEndOfAssetName;
+	}
+
 	for (const FString& FilterString : FilterStrings)
 	{
-		if (FilterString.StartsWith(TEXT("_")))
+		if (bOnlyMatchEndOfAssetName || FilterString.StartsWith(TEXT("_")))
 		{
 			if (InAssetName.EndsWith(FilterString))
 			{
diff --git a/Engine/Plugins/Experimental/DynamicMaterial/Source/DynamicMaterialTextureSetEditor/Private/DMTextureSetFilter.h b/Engine/Plugins/Experimental/DynamicMaterial/Source/DynamicMaterialTextureSetEditor/Private/DMTextureSetFilter.h
--- a/Engine/Plugins/Experimental/DynamicMaterial/Source/DynamicMaterialTextureSetEditor/Private/DMTextureSetFilter.h
+++ b/Engine/Plugins/Experimental/DynamicMaterial/Source/DynamicMaterialTextureSetEditor/Private/DMTextureSetFilter.h
@@ -4,7 +4,7 @@
 
 #include "Containers/Map.h"
 #include "DMTextureChannelMask.h"
-#include "SceneTypes.h"
+#include "DMTextureSetMaterialProperty.h"
 
 #include "DMTextureSetFilter.generated.h"
 
@@ -27,7 +27,7 @@ struct FDMTextureSetFilter
 	 * Where the matched Texture assets should be placed into the Texture Set. Links to the channel for the given asset.
 	 */
 	UPROPERTY(EditAnywhere, Category = "Material Designer")
-	TMap<TEnumAsByte<EMaterialProperty>, EDMTextureChannelMask> MaterialProperties;
+	TMap<EDMTextureSetMaterialProperty, EDMTextureChannelMask> MaterialProperties;
 
 	/**
 	 * Checks the given asset name against the filter strings.
diff --git a/Engine/Plugins/Experimental/DynamicMaterial/Source/DynamicMaterialTextureSetEditor/Private/DMTextureSetSettings.cpp b/Engine/Plugins/Experimental/DynamicMaterial/Source/DynamicMaterialTextureSetEditor/Private/DMTextureSetSettings.cpp
--- a/Engine/Plugins/Experimental/DynamicMaterial/Source/DynamicMaterialTextureSetEditor/Private/DMTextureSetSettings.cpp
+++ b/Engine/Plugins/Experimental/DynamicMaterial/Source/DynamicMaterialTextureSetEditor/Private/DMTextureSetSettings.cpp
@@ -2,49 +2,119 @@
 
 #include "DMTextureSetSettings.h"
 
+#include "DMTextureSetMaterialProperty.h"
+
 UDMTextureSetSettings::UDMTextureSetSettings()
 {
+	int32 FilterCount = 0;
+
 	FDMTextureSetFilter BaseColor;
-	BaseColor.FilterStrings = {TEXT("Base_Color"), TEXT("BaseColor"), TEXT("Base_Colour"), TEXT("BaseColour"), TEXT("_BC"), TEXT("Diffuse"), TEXT("Albedo")};
-	BaseColor.MaterialProperties = {{EMaterialProperty::MP_BaseColor, EDMTextureChannelMask::RGBA}};
+	BaseColor.FilterStrings = {TEXT("Base_Color"), TEXT("BaseColor"), TEXT("Base_Colour"), TEXT("BaseColour"), TEXT("_BC"), TEXT("Diffuse"), TEXT("Albedo"), TEXT("_Diff")};
+	BaseColor.MaterialProperties = {{EDMTextureSetMaterialProperty::BaseColor, EDMTextureChannelMask::RGBA}};
+	++FilterCount;
 
 	FDMTextureSetFilter Roughness;
 	Roughness.FilterStrings = {TEXT("Roughness"), TEXT("Rough"), TEXT("_R")};
-	Roughness.MaterialProperties = {{EMaterialProperty::MP_Roughness, EDMTextureChannelMask::RGBA}};
+	Roughness.MaterialProperties = {{EDMTextureSetMaterialProperty::Roughness, EDMTextureChannelMask::RGBA}};
+	++FilterCount;
 
 	FDMTextureSetFilter Normal;
-	Normal.FilterStrings = {TEXT("Normal"), TEXT("Norm"), TEXT("_N")};
-	Normal.MaterialProperties = {{EMaterialProperty::MP_Normal, EDMTextureChannelMask::RGBA}};
+	Normal.FilterStrings = {TEXT("Normal"), TEXT("Norm"), TEXT("_N"), TEXT("_Nor")};
+	Normal.MaterialProperties = {{EDMTextureSetMaterialProperty::Normal, EDMTextureChannelMask::RGBA}};
+	++FilterCount;
 
 	FDMTextureSetFilter Metallic;
 	Metallic.FilterStrings = {TEXT("Metallic"), TEXT("Metal"), TEXT("_M")};
-	Metallic.MaterialProperties = {{EMaterialProperty::MP_Metallic, EDMTextureChannelMask::RGBA}};
+	Metallic.MaterialProperties = {{EDMTextureSetMaterialProperty::Metallic, EDMTextureChannelMask::RGBA}};
+	++FilterCount;
 
 	FDMTextureSetFilter AmbientOcclusion;
 	AmbientOcclusion.FilterStrings = {TEXT("AmbientOcclusion"), TEXT("Ambient_Occlusion"), TEXT("_AO")};
-	AmbientOcclusion.MaterialProperties = {{EMaterialProperty::MP_AmbientOcclusion, EDMTextureChannelMask::RGBA}};
+	AmbientOcclusion.MaterialProperties = {{EDMTextureSetMaterialProperty::AmbientOcclusion, EDMTextureChannelMask::RGBA}};
+	++FilterCount;
 
 	FDMTextureSetFilter Specular;
 	Specular.FilterStrings = {TEXT("Specular"), TEXT("_S")};
-	Specular.MaterialProperties = {{EMaterialProperty::MP_EmissiveColor, EDMTextureChannelMask::RGBA}};
+	Specular.MaterialProperties = {{EDMTextureSetMaterialProperty::Specular, EDMTextureChannelMask::RGBA}};
+	++FilterCount;
 
 	FDMTextureSetFilter Emissive;
 	Emissive.FilterStrings = {TEXT("Emissive"), TEXT("Emission"), TEXT("_E")};
-	Emissive.MaterialProperties = {{EMaterialProperty::MP_EmissiveColor, EDMTextureChannelMask::RGBA}};
+	Emissive.MaterialProperties = {{EDMTextureSetMaterialProperty::EmissiveColor, EDMTextureChannelMask::RGBA}};
+	++FilterCount;
 
 	FDMTextureSetFilter Opacity;
 	Opacity.FilterStrings = {TEXT("Opacity"), TEXT("_O"), TEXT("Alpha"), TEXT("_A")};
-	Opacity.MaterialProperties = {{EMaterialProperty::MP_EmissiveColor, EDMTextureChannelMask::RGBA}};
+	Opacity.MaterialProperties = {{EDMTextureSetMaterialProperty::Opacity, EDMTextureChannelMask::RGBA}};
+	++FilterCount;
+
+	FDMTextureSetFilter OpacityMask;
+	Opacity.FilterStrings = {TEXT("OpacityMask"), TEXT("Opacity_Mask"), TEXT("_OM"), TEXT("AlphaMask"), TEXT("Alpha_Mask"), TEXT("_AM")};
+	Opacity.MaterialProperties = {{EDMTextureSetMaterialProperty::OpacityMask, EDMTextureChannelMask::RGBA}};
+	++FilterCount;
+
+	FDMTextureSetFilter Displacement;
+	Displacement.FilterStrings = {TEXT("Displacement"), TEXT("_Disp"), TEXT("Height")};
+	Displacement.MaterialProperties = {{EDMTextureSetMaterialProperty::Displacement, EDMTextureChannelMask::RGBA}};
+	++FilterCount;
+
+	FDMTextureSetFilter WorldPositionOffset;
+	WorldPositionOffset.FilterStrings = {TEXT("WorldPositionOffset"), TEXT("World_Position_Offset"), TEXT("_WPO")};
+	WorldPositionOffset.MaterialProperties = {{EDMTextureSetMaterialProperty::WorldPositionOffset, EDMTextureChannelMask::RGBA}};
+	++FilterCount;
+
+	FDMTextureSetFilter PixelDepthOffset;
+	PixelDepthOffset.FilterStrings = {TEXT("PixelDepthOffset"), TEXT("Pixel_Depth_Offset"), TEXT("_PDO")};
+	PixelDepthOffset.MaterialProperties = {{EDMTextureSetMaterialProperty::PixelDepthOffset, EDMTextureChannelMask::RGBA}};
+	++FilterCount;
+
+	FDMTextureSetFilter SubsurfaceColor;
+	SubsurfaceColor.FilterStrings = {TEXT("SubsurfaceColor"), TEXT("SubsurfaceColour"), TEXT("Subsurface_Color"), TEXT("Subsurface_Coluor"), TEXT("_SC"), TEXT("_SS"), TEXT("_SSC")};
+	SubsurfaceColor.MaterialProperties = {{EDMTextureSetMaterialProperty::SubsurfaceColor, EDMTextureChannelMask::RGBA}};
+	++FilterCount;
+
+	FDMTextureSetFilter Anisotropy;
+	Anisotropy.FilterStrings = {TEXT("Anisotropy"), TEXT("_Ani")};
+	Anisotropy.MaterialProperties = {{EDMTextureSetMaterialProperty::Anisotropy, EDMTextureChannelMask::RGBA}};
+	++FilterCount;
+
+	FDMTextureSetFilter Refraction;
+	Refraction.FilterStrings = {TEXT("Refraction"), TEXT("_Ref")};
+	Refraction.MaterialProperties = {{EDMTextureSetMaterialProperty::Refraction, EDMTextureChannelMask::RGBA}};
+	++FilterCount;
+
+	FDMTextureSetFilter Tangent;
+	Tangent.FilterStrings = {TEXT("Tangent"), TEXT("_T")};
+	Tangent.MaterialProperties = {{EDMTextureSetMaterialProperty::Tangent, EDMTextureChannelMask::RGBA}};
+	++FilterCount;
 
 	FDMTextureSetFilter ORM;
 	ORM.FilterStrings = {TEXT("_ORM")};
 	ORM.MaterialProperties = {
-		{EMaterialProperty::MP_Opacity, EDMTextureChannelMask::Red},
-		{EMaterialProperty::MP_Roughness, EDMTextureChannelMask::Green},
-		{EMaterialProperty::MP_Metallic, EDMTextureChannelMask::Blue}
+		{EDMTextureSetMaterialProperty::Opacity, EDMTextureChannelMask::Red},
+		{EDMTextureSetMaterialProperty::Roughness, EDMTextureChannelMask::Green},
+		{EDMTextureSetMaterialProperty::Metallic, EDMTextureChannelMask::Blue}
+	};
+	++FilterCount;
+
+	FDMTextureSetFilter RM;
+	RM.FilterStrings = {TEXT("_RM")};
+	RM.MaterialProperties = {
+		{EDMTextureSetMaterialProperty::Roughness, EDMTextureChannelMask::Red},
+		{EDMTextureSetMaterialProperty::Metallic, EDMTextureChannelMask::Green}
+	};
+	++FilterCount;
+
+	FDMTextureSetFilter RMA;
+	RMA.FilterStrings = {TEXT("_RMA")};
+	RMA.MaterialProperties = {
+		{EDMTextureSetMaterialProperty::Roughness, EDMTextureChannelMask::Red},
+		{EDMTextureSetMaterialProperty::Metallic, EDMTextureChannelMask::Green},
+		{EDMTextureSetMaterialProperty::AmbientOcclusion, EDMTextureChannelMask::Blue}
 	};
+	++FilterCount;
 
-	Filters.Reserve(9);
+	Filters.Reserve(FilterCount);
 
 	Filters.Add(MoveTemp(BaseColor));
 	Filters.Add(MoveTemp(Roughness));
@@ -54,7 +124,17 @@ UDMTextureSetSettings::UDMTextureSetSettings()
 	Filters.Add(MoveTemp(Specular));
 	Filters.Add(MoveTemp(Emissive));
 	Filters.Add(MoveTemp(Opacity));
+	Filters.Add(MoveTemp(OpacityMask));
+	Filters.Add(MoveTemp(Displacement));
+	Filters.Add(MoveTemp(WorldPositionOffset));
+	Filters.Add(MoveTemp(PixelDepthOffset));
+	Filters.Add(MoveTemp(SubsurfaceColor));
+	Filters.Add(MoveTemp(Anisotropy));
+	Filters.Add(MoveTemp(Refraction));
+	Filters.Add(MoveTemp(Tangent));
 	Filters.Add(MoveTemp(ORM));
+	Filters.Add(MoveTemp(RM));
+	Filters.Add(MoveTemp(RMA));
 }
 
 UDMTextureSetSettings* UDMTextureSetSettings::Get()
diff --git a/Engine/Plugins/Experimental/DynamicMaterial/Source/DynamicMaterialTextureSetEditor/Private/DMTextureSetSettings.h b/Engine/Plugins/Experimental/DynamicMaterial/Source/DynamicMaterialTextureSetEditor/Private/DMTextureSetSettings.h
--- a/Engine/Plugins/Experimental/DynamicMaterial/Source/DynamicMaterialTextureSetEditor/Private/DMTextureSetSettings.h
+++ b/Engine/Plugins/Experimental/DynamicMaterial/Source/DynamicMaterialTextureSetEditor/Private/DMTextureSetSettings.h
@@ -22,6 +22,9 @@ public:
 
 	static UDMTextureSetSettings* Get();
 
+	UPROPERTY(EditAnywhere, Category = "Material Designer")
+	bool bOnlyMatchEndOfAssetName = false;
+
 	UPROPERTY(EditAnywhere, Category = "Material Designer", meta = (TitleProperty=FilterStrings))
 	TArray<FDMTextureSetFilter> Filters;
 };
diff --git a/Engine/Plugins/Experimental/DynamicMaterial/Source/DynamicMaterialTextureSetEditor/Private/DMTextureSetStyle.cpp b/Engine/Plugins/Experimental/DynamicMaterial/Source/DynamicMaterialTextureSetEditor/Private/DMTextureSetStyle.cpp
--- a/Engine/Plugins/Experimental/DynamicMaterial/Source/DynamicMaterialTextureSetEditor/Private/DMTextureSetStyle.cpp
+++ b/Engine/Plugins/Experimental/DynamicMaterial/Source/DynamicMaterialTextureSetEditor/Private/DMTextureSetStyle.cpp
@@ -2,6 +2,7 @@
 
 #include "DMTextureSetStyle.h"
 
+#include "Brushes/SlateColorBrush.h"
 #include "Brushes/SlateImageBrush.h"
 #include "Styling/StyleColors.h"
 
@@ -14,7 +15,10 @@ const ISlateStyle& FDMTextureSetStyle::Get()
 FDMTextureSetStyle::FDMTextureSetStyle()
 	: FSlateStyleSet(TEXT("DMTextureSetStyle"))
 {
-	Set("TextureSetConfig.Background", new FSlateRoundedBoxBrush(
+	Set("TextureSetConfig.Window.Background", new FSlateColorBrush(
+		FStyleColors::Panel.GetSpecifiedColor()));
+
+	Set("TextureSetConfig.Cell.Background", new FSlateRoundedBoxBrush(
 		FStyleColors::Recessed.GetSpecifiedColor(), 6.0f,
 		FStyleColors::Header.GetSpecifiedColor(), 2.0f));
 }
diff --git a/Engine/Plugins/Experimental/DynamicMaterial/Source/DynamicMaterialTextureSetEditor/Private/DynamicMaterialTextureSetEditorModule.cpp b/Engine/Plugins/Experimental/DynamicMaterial/Source/DynamicMaterialTextureSetEditor/Private/DynamicMaterialTextureSetEditorModule.cpp
--- a/Engine/Plugins/Experimental/DynamicMaterial/Source/DynamicMaterialTextureSetEditor/Private/DynamicMaterialTextureSetEditorModule.cpp
+++ b/Engine/Plugins/Experimental/DynamicMaterial/Source/DynamicMaterialTextureSetEditor/Private/DynamicMaterialTextureSetEditorModule.cpp
@@ -2,19 +2,19 @@
 
 #include "DynamicMaterialTextureSetEditorModule.h"
 
-#include "DMTextureSetContentBrowserIntegration.h"
+#include "DMTextureSetContentBrowserIntegrationPrivate.h"
 #include "DMTextureSetStyle.h"
 #include "Modules/ModuleManager.h"
 
 void FDynamicMaterialTextureSetEditorModule::StartupModule()
 {
 	FDMTextureSetStyle::Get();
-	FDMTextureSetContentBrowserIntegration::Integrate();
+	FDMTextureSetContentBrowserIntegrationPrivate::Integrate();
 }
 
 void FDynamicMaterialTextureSetEditorModule::ShutdownModule()
 {
-	FDMTextureSetContentBrowserIntegration::Disintegrate();
+	FDMTextureSetContentBrowserIntegrationPrivate::Disintegrate();
 }
 
 IMPLEMENT_MODULE(FDynamicMaterialTextureSetEditorModule, DynamicMaterialTextureSetEditor)
diff --git a/Engine/Plugins/Experimental/DynamicMaterial/Source/DynamicMaterialTextureSetEditor/Private/Widgets/DMTextureSetBuilderEntry.cpp b/Engine/Plugins/Experimental/DynamicMaterial/Source/DynamicMaterialTextureSetEditor/Private/Widgets/DMTextureSetBuilderEntry.cpp
--- a/Engine/Plugins/Experimental/DynamicMaterial/Source/DynamicMaterialTextureSetEditor/Private/Widgets/DMTextureSetBuilderEntry.cpp
+++ b/Engine/Plugins/Experimental/DynamicMaterial/Source/DynamicMaterialTextureSetEditor/Private/Widgets/DMTextureSetBuilderEntry.cpp
@@ -2,14 +2,14 @@
 
 #include "Widgets/DMTextureSetBuilderEntry.h"
 
-#include "SceneTypes.h"
+#include "DMTextureSetMaterialProperty.h"
 
 FDMTextureSetBuilderEntry::FDMTextureSetBuilderEntry()
-	: FDMTextureSetBuilderEntry(EMaterialProperty::MP_BaseColor, nullptr, EDMTextureChannelMask::RGBA)
+	: FDMTextureSetBuilderEntry(EDMTextureSetMaterialProperty::BaseColor, nullptr, EDMTextureChannelMask::RGBA)
 {
 }
 
-FDMTextureSetBuilderEntry::FDMTextureSetBuilderEntry(EMaterialProperty InMaterialProperty, UTexture* InTexture, EDMTextureChannelMask InMask)
+FDMTextureSetBuilderEntry::FDMTextureSetBuilderEntry(EDMTextureSetMaterialProperty InMaterialProperty, UTexture* InTexture, EDMTextureChannelMask InMask)
 	: MaterialProperty(InMaterialProperty)
 	, Texture(InTexture)
 	, ChannelMask(InMask)
diff --git a/Engine/Plugins/Experimental/DynamicMaterial/Source/DynamicMaterialTextureSetEditor/Private/Widgets/DMTextureSetBuilderEntry.h b/Engine/Plugins/Experimental/DynamicMaterial/Source/DynamicMaterialTextureSetEditor/Private/Widgets/DMTextureSetBuilderEntry.h
--- a/Engine/Plugins/Experimental/DynamicMaterial/Source/DynamicMaterialTextureSetEditor/Private/Widgets/DMTextureSetBuilderEntry.h
+++ b/Engine/Plugins/Experimental/DynamicMaterial/Source/DynamicMaterialTextureSetEditor/Private/Widgets/DMTextureSetBuilderEntry.h
@@ -5,7 +5,7 @@
 #include "IStructureDataProvider.h"
 
 #include "DMTextureChannelMask.h"
-#include "SceneTypes.h"
+#include "DMTextureSetMaterialProperty.h"
 #include "Templates/SharedPointer.h"
 #include "UObject/ObjectPtr.h"
 
@@ -20,10 +20,10 @@ struct FDMTextureSetBuilderEntry
 
 	FDMTextureSetBuilderEntry();
 
-	FDMTextureSetBuilderEntry(EMaterialProperty InMaterialProperty, UTexture* InTexture, EDMTextureChannelMask InMask);
+	FDMTextureSetBuilderEntry(EDMTextureSetMaterialProperty InMaterialProperty, UTexture* InTexture, EDMTextureChannelMask InMask);
 
 	UPROPERTY(EditAnywhere, Category = "Material Designer")
-	TEnumAsByte<EMaterialProperty> MaterialProperty = EMaterialProperty::MP_BaseColor;
+	EDMTextureSetMaterialProperty MaterialProperty = EDMTextureSetMaterialProperty::BaseColor;
 
 	UPROPERTY(EditAnywhere, Category = "Material Designer")
 	TObjectPtr<UTexture> Texture;
diff --git a/Engine/Plugins/Experimental/DynamicMaterial/Source/DynamicMaterialTextureSetEditor/Private/Widgets/SDMTextureSetBuilder.cpp b/Engine/Plugins/Experimental/DynamicMaterial/Source/DynamicMaterialTextureSetEditor/Private/Widgets/SDMTextureSetBuilder.cpp
--- a/Engine/Plugins/Experimental/DynamicMaterial/Source/DynamicMaterialTextureSetEditor/Private/Widgets/SDMTextureSetBuilder.cpp
+++ b/Engine/Plugins/Experimental/DynamicMaterial/Source/DynamicMaterialTextureSetEditor/Private/Widgets/SDMTextureSetBuilder.cpp
@@ -3,10 +3,10 @@
 #include "Widgets/SDMTextureSetBuilder.h"
 
 #include "DMTextureSet.h"
+#include "DMTextureSetMaterialProperty.h"
 #include "DMTextureSetStyle.h"
 #include "Engine/Texture.h"
 #include "SAssetDropTarget.h"
-#include "SceneTypes.h"
 #include "Widgets/DMTextureSetBuilderDragDropOperation.h"
 #include "Widgets/DMTextureSetBuilderEntry.h"
 #include "Widgets/Input/SButton.h"
@@ -22,7 +22,7 @@
 
 namespace UE::DynamicMaterialEditor::TextureSet::Private
 {
-	constexpr int32 WidthMax = 7;
+	constexpr int32 WidthMax = 8;
 }
 
 void SDMTextureSetBuilder::Construct(const FArguments& InArgs, UDMTextureSet* InTextureSet, const TArray<FAssetData>& InAssets, 
@@ -46,7 +46,7 @@ void SDMTextureSetBuilder::Construct(const FArguments& InArgs, UDMTextureSet* In
 
 	MaterialProperties.Reserve(InTextureSet->GetTextures().Num());
 
-	for (const TPair<TEnumAsByte<EMaterialProperty>, FDMMaterialTexture>& SetElement : InTextureSet->GetTextures())
+	for (const TPair<EDMTextureSetMaterialProperty, FDMMaterialTexture>& SetElement : InTextureSet->GetTextures())
 	{
 		TSharedRef<FDMTextureSetBuilderEntry> Entry = MakeShared<FDMTextureSetBuilderEntry>(
 			SetElement.Key,
@@ -77,8 +77,9 @@ void SDMTextureSetBuilder::Construct(const FArguments& InArgs, UDMTextureSet* In
 
 	ChildSlot
 	[
-		SNew(SBox)
-		.Padding(5.0f)
+		SNew(SBorder)
+		.Padding(5.f)
+		.BorderImage(FDMTextureSetStyle::Get().GetBrush("TextureSetConfig.Window.Background"))
 		[
 			SNew(SVerticalBox)
 
@@ -95,7 +96,7 @@ void SDMTextureSetBuilder::Construct(const FArguments& InArgs, UDMTextureSet* In
 			[
 				SNew(SBorder)
 				.Padding(5.f)
-				.BorderImage(FDMTextureSetStyle::Get().GetBrush("TextureSetConfig.Background"))
+				.BorderImage(FDMTextureSetStyle::Get().GetBrush("TextureSetConfig.Cell.Background"))
 				[
 					SNew(SAssetDropTarget)
 					.OnAreAssetsAcceptableForDrop(this, &SDMTextureSetBuilder::OnAssetDraggedOver)
@@ -213,6 +214,32 @@ void SDMTextureSetBuilder::SwapTexture(int32 InFromIndex, bool bInIsFromMaterial
 	}
 }
 
+void SDMTextureSetBuilder::SetTexture(int32 InIndex, bool bInIsMaterialProperty, UTexture* InTexture)
+{
+	if (!IsValid(InTexture))
+	{
+		return;
+	}
+
+	TSharedPtr<SDMTextureSetBuilderCellBase> Cell = GetCell(InIndex, bInIsMaterialProperty);
+
+	if (!Cell.IsValid())
+	{
+		return;
+	}
+
+	Cell->SetTexture(InTexture);
+
+	if (bInIsMaterialProperty)
+	{
+		MaterialProperties[InIndex]->Texture = InTexture;
+	}
+
+	Assets.AddUnique(FAssetData(InTexture));
+
+	CreateUnassignedTextureSlots();
+}
+
 TSharedRef<SWidget> SDMTextureSetBuilder::GenerateMaterialPropertyCell(const TSharedRef<FDMTextureSetBuilderEntry>& InListItem, int32 InIndex)
 {
 	return SNew(SDMTextureSetBuilderMaterialPropertyCell, SharedThis(this), InListItem, InIndex);
@@ -232,7 +259,7 @@ FReply SDMTextureSetBuilder::OnAcceptClicked()
 		return FReply::Handled();
 	}
 
-	for (const TPair<TEnumAsByte<EMaterialProperty>, FDMMaterialTexture>& SetElement : TextureSet->GetTextures())
+	for (const TPair<EDMTextureSetMaterialProperty, FDMMaterialTexture>& SetElement : TextureSet->GetTextures())
 	{
 		TextureSet->SetMaterialTexture(SetElement.Key, {nullptr, EDMTextureChannelMask::RGBA});
 	}
diff --git a/Engine/Plugins/Experimental/DynamicMaterial/Source/DynamicMaterialTextureSetEditor/Private/Widgets/SDMTextureSetBuilder.h b/Engine/Plugins/Experimental/DynamicMaterial/Source/DynamicMaterialTextureSetEditor/Private/Widgets/SDMTextureSetBuilder.h
--- a/Engine/Plugins/Experimental/DynamicMaterial/Source/DynamicMaterialTextureSetEditor/Private/Widgets/SDMTextureSetBuilder.h
+++ b/Engine/Plugins/Experimental/DynamicMaterial/Source/DynamicMaterialTextureSetEditor/Private/Widgets/SDMTextureSetBuilder.h
@@ -28,6 +28,8 @@ public:
 
 	void SwapTexture(int32 InFromIndex, bool bInIsFromMaterialProperty, int32 InToIndex, bool bInIsToMaterialProperty);
 
+	void SetTexture(int32 InIndex, bool bInIsMaterialProperty, UTexture* InTexture);
+
 protected:
 	TStrongObjectPtr<UDMTextureSet> TextureSet;
 	TArray<FAssetData> Assets;
diff --git a/Engine/Plugins/Experimental/DynamicMaterial/Source/DynamicMaterialTextureSetEditor/Private/Widgets/SDMTextureSetBuilderCellBase.cpp b/Engine/Plugins/Experimental/DynamicMaterial/Source/DynamicMaterialTextureSetEditor/Private/Widgets/SDMTextureSetBuilderCellBase.cpp
--- a/Engine/Plugins/Experimental/DynamicMaterial/Source/DynamicMaterialTextureSetEditor/Private/Widgets/SDMTextureSetBuilderCellBase.cpp
+++ b/Engine/Plugins/Experimental/DynamicMaterial/Source/DynamicMaterialTextureSetEditor/Private/Widgets/SDMTextureSetBuilderCellBase.cpp
@@ -3,6 +3,7 @@
 #include "Widgets/SDMTextureSetBuilderCellBase.h"
 
 #include "AssetRegistry/AssetData.h"
+#include "DragAndDrop/AssetDragDropOp.h"
 #include "Engine/Texture.h"
 #include "Widgets/DMTextureSetBuilderDragDropOperation.h"
 #include "Widgets/Images/SImage.h"
@@ -74,9 +75,9 @@ bool SDMTextureSetBuilderCellBase::OnAssetDraggedOver(TArrayView<FAssetData> InA
 
 void SDMTextureSetBuilderCellBase::OnAssetsDropped(const FDragDropEvent& InDragDropEvent, TArrayView<FAssetData> InAssets)
 {
-	if (TSharedPtr<FDMTextureSetBuilderDragDropOperation> BuilderOperation = InDragDropEvent.GetOperationAs<FDMTextureSetBuilderDragDropOperation>())
+	if (TSharedPtr<SDMTextureSetBuilder> TextureSetBuilder = TextureSetBuilderWeak.Pin())
 	{
-		if (TSharedPtr<SDMTextureSetBuilder> TextureSetBuilder = TextureSetBuilderWeak.Pin())
+		if (TSharedPtr<FDMTextureSetBuilderDragDropOperation> BuilderOperation = InDragDropEvent.GetOperationAs<FDMTextureSetBuilderDragDropOperation>())
 		{
 			TextureSetBuilder->SwapTexture(
 				BuilderOperation->GetIndex(),
@@ -85,6 +86,28 @@ void SDMTextureSetBuilderCellBase::OnAssetsDropped(const FDragDropEvent& InDragD
 				bIsMaterialProperty
 			);
 		}
+		else if (TSharedPtr<FAssetDragDropOp> AssetOperation = InDragDropEvent.GetOperationAs<FAssetDragDropOp>())
+		{
+			for (const FAssetData& Asset : AssetOperation->GetAssets())
+			{
+				UClass* AssetClass = Asset.GetClass(EResolveClass::Yes);
+
+				if (AssetClass && AssetClass->IsChildOf(UTexture::StaticClass()))
+				{
+					if (UTexture* AssetTexture = Cast<UTexture>(Asset.GetAsset()))
+					{
+						TextureSetBuilder->SetTexture(
+							Index,
+							bIsMaterialProperty,
+							AssetTexture
+						);
+
+						// Only set the first texture.
+						break;
+					}
+				}
+			}
+		}
 	}
 }
 
@@ -102,7 +125,9 @@ FText SDMTextureSetBuilderCellBase::GetToolTipText() const
 		return LOCTEXT("NoTexture", "Texture slot empty.");
 	}
 
-	return FText::FromString(TextureObject->GetPathName());
+	const FText Format = LOCTEXT("TextureTooltipFormat", "{0}\n\nDrag to another slot to swap textures. Hold shift when dropping to overwrite.");
+
+	return FText::Format(Format, FText::FromString(TextureObject->GetPathName()));
 }
 
 FText SDMTextureSetBuilderCellBase::GetTextureName() const
diff --git a/Engine/Plugins/Experimental/DynamicMaterial/Source/DynamicMaterialTextureSetEditor/Private/Widgets/SDMTextureSetBuilderMaterialPropertyCell.cpp b/Engine/Plugins/Experimental/DynamicMaterial/Source/DynamicMaterialTextureSetEditor/Private/Widgets/SDMTextureSetBuilderMaterialPropertyCell.cpp
--- a/Engine/Plugins/Experimental/DynamicMaterial/Source/DynamicMaterialTextureSetEditor/Private/Widgets/SDMTextureSetBuilderMaterialPropertyCell.cpp
+++ b/Engine/Plugins/Experimental/DynamicMaterial/Source/DynamicMaterialTextureSetEditor/Private/Widgets/SDMTextureSetBuilderMaterialPropertyCell.cpp
@@ -3,6 +3,8 @@
 #include "Widgets/SDMTextureSetBuilderMaterialPropertyCell.h"
 
 #include "DMTextureSetBuilderEntry.h"
+#include "DMTextureSetMaterialProperty.h"
+#include "DMTextureSetStyle.h"
 #include "Engine/Texture.h"
 #include "ISinglePropertyView.h"
 #include "Materials/Material.h"
@@ -10,7 +12,6 @@
 #include "Modules/ModuleManager.h"
 #include "PropertyEditorModule.h"
 #include "SAssetDropTarget.h"
-#include "SceneTypes.h"
 #include "SDMTextureSetBuilder.h"
 #include "Widgets/Images/SImage.h"
 #include "Widgets/SBoxPanel.h"
@@ -32,7 +33,7 @@ namespace UE::DynamicMaterial::TextureSet::Private
 }
 
 SDMTextureSetBuilderMaterialPropertyCell::SDMTextureSetBuilderMaterialPropertyCell()
-	: MaterialBrush(FSlateMaterialBrush(FVector2D(100.f)))
+	: MaterialBrush(FSlateMaterialBrush(FVector2D(120.f)))
 {
 }
 
@@ -49,7 +50,7 @@ void SDMTextureSetBuilderMaterialPropertyCell::Construct(const FArguments& InArg
 
 	Entry = InEntry;
 
-	UEnum* MaterialPropertyEnum = StaticEnum<EMaterialProperty>();
+	UEnum* MaterialPropertyEnum = StaticEnum<EDMTextureSetMaterialProperty>();
 
 	if (!MaterialPropertyEnum)
 	{
@@ -81,54 +82,59 @@ void SDMTextureSetBuilderMaterialPropertyCell::Construct(const FArguments& InArg
 
 	ChildSlot
 	[
-		SNew(SVerticalBox)
-
-		+ SVerticalBox::Slot()
-		.AutoHeight()
-		.HAlign(EHorizontalAlignment::HAlign_Center)
-		.Padding(0.f, 0.f, 0.f, 5.f)
+		SNew(SBorder)
+		.Padding(10.f)
+		.BorderImage(FDMTextureSetStyle::Get().GetBrush("TextureSetConfig.Cell.Background"))
 		[
-			SNew(STextBlock)
-			.Text(MaterialPropertyEnum->GetDisplayNameTextByValue(Entry->MaterialProperty))
-		]
+			SNew(SVerticalBox)
 
-		+ SVerticalBox::Slot()
-		.AutoHeight()
-		.HAlign(EHorizontalAlignment::HAlign_Center)
-		.Padding(0.f, 0.f, 0.f, 5.f)
-		[
-			SNew(SAssetDropTarget)
-			.OnAreAssetsAcceptableForDrop(this, &SDMTextureSetBuilderMaterialPropertyCell::OnAssetDraggedOver)
-			.OnAssetsDropped(this, &SDMTextureSetBuilderMaterialPropertyCell::OnAssetsDropped)
+			+ SVerticalBox::Slot()
+			.AutoHeight()
+			.HAlign(EHorizontalAlignment::HAlign_Center)
+			.Padding(0.f, 0.f, 0.f, 5.f)
 			[
-				SNew(SOverlay)
-				.ToolTipText(this, &SDMTextureSetBuilderMaterialPropertyCell::GetToolTipText)
-
-				+ SOverlay::Slot()
-				[
-					SNew(SImage)
-					.Image(&MaterialBrush)
-					.DesiredSizeOverride(FVector2D(100.f))
-					.Visibility(this, &SDMTextureSetBuilderMaterialPropertyCell::GetImageVisibility)
-				]
+				SNew(STextBlock)
+				.Text(MaterialPropertyEnum->GetDisplayNameTextByValue(static_cast<int64>(Entry->MaterialProperty)))
+			]
 
-				+ SOverlay::Slot()
-				.Padding(5.f)
+			+ SVerticalBox::Slot()
+			.AutoHeight()
+			.HAlign(EHorizontalAlignment::HAlign_Center)
+			.Padding(0.f, 0.f, 0.f, 5.f)
+			[
+				SNew(SAssetDropTarget)
+				.OnAreAssetsAcceptableForDrop(this, &SDMTextureSetBuilderMaterialPropertyCell::OnAssetDraggedOver)
+				.OnAssetsDropped(this, &SDMTextureSetBuilderMaterialPropertyCell::OnAssetsDropped)
 				[
-					SNew(STextBlock)
-					.Text(this, &SDMTextureSetBuilderMaterialPropertyCell::GetTextureName)
-					.WrappingPolicy(ETextWrappingPolicy::AllowPerCharacterWrapping)
-					.WrapTextAt(90.f)
+					SNew(SOverlay)
+					.ToolTipText(this, &SDMTextureSetBuilderMaterialPropertyCell::GetToolTipText)
+
+					+ SOverlay::Slot()
+					[
+						SNew(SImage)
+						.Image(&MaterialBrush)
+						.DesiredSizeOverride(FVector2D(120.f))
+						.Visibility(this, &SDMTextureSetBuilderMaterialPropertyCell::GetImageVisibility)
+					]
+
+					+ SOverlay::Slot()
+					.Padding(5.f)
+					[
+						SNew(STextBlock)
+						.Text(this, &SDMTextureSetBuilderMaterialPropertyCell::GetTextureName)
+						.WrappingPolicy(ETextWrappingPolicy::AllowPerCharacterWrapping)
+						.WrapTextAt(90.f)
+					]
 				]
 			]
-		]
 
-		+ SVerticalBox::Slot()
-		.AutoHeight()
-		.HAlign(EHorizontalAlignment::HAlign_Left)
-		.Padding(0.f, 0.f, 0.f, 0.f)
-		[
-			PropertyView.ToSharedRef()
+			+ SVerticalBox::Slot()
+			.AutoHeight()
+			.HAlign(EHorizontalAlignment::HAlign_Left)
+			.Padding(0.f, 0.f, 0.f, 0.f)
+			[
+				PropertyView.ToSharedRef()
+			]
 		]
 	];
 }
diff --git a/Engine/Plugins/Experimental/DynamicMaterial/Source/DynamicMaterialTextureSetEditor/Private/Widgets/SDMTextureSetBuilderUnassignedTextureCell.cpp b/Engine/Plugins/Experimental/DynamicMaterial/Source/DynamicMaterialTextureSetEditor/Private/Widgets/SDMTextureSetBuilderUnassignedTextureCell.cpp
--- a/Engine/Plugins/Experimental/DynamicMaterial/Source/DynamicMaterialTextureSetEditor/Private/Widgets/SDMTextureSetBuilderUnassignedTextureCell.cpp
+++ b/Engine/Plugins/Experimental/DynamicMaterial/Source/DynamicMaterialTextureSetEditor/Private/Widgets/SDMTextureSetBuilderUnassignedTextureCell.cpp
@@ -13,7 +13,7 @@
 #define LOCTEXT_NAMESPACE "SDMTextureSetBuilderUnassignedTextureCell"
 
 SDMTextureSetBuilderUnassignedTextureCell::SDMTextureSetBuilderUnassignedTextureCell()
-	: TextureBrush(FSlateImageBrush((UObject*)nullptr, FVector2D(100.f)))
+	: TextureBrush(FSlateImageBrush((UObject*)nullptr, FVector2D(120.f)))
 {
 }
 
@@ -59,7 +59,7 @@ void SDMTextureSetBuilderUnassignedTextureCell::Construct(const FArguments& InAr
 			[
 				SNew(SImage)
 				.Image(&TextureBrush)
-				.DesiredSizeOverride(FVector2D(100.f))
+				.DesiredSizeOverride(FVector2D(120.f))
 				.Visibility(this, &SDMTextureSetBuilderUnassignedTextureCell::GetImageVisibility)
 			]
 
diff --git a/Engine/Plugins/Experimental/DynamicMaterial/Source/DynamicMaterialTextureSetEditor/Public/DMTextureSetContentBrowserIntegration.h b/Engine/Plugins/Experimental/DynamicMaterial/Source/DynamicMaterialTextureSetEditor/Public/DMTextureSetContentBrowserIntegration.h
new file mode 100644
--- /dev/null
+++ b/Engine/Plugins/Experimental/DynamicMaterial/Source/DynamicMaterialTextureSetEditor/Public/DMTextureSetContentBrowserIntegration.h
@@ -0,0 +1,23 @@
+// Copyright Epic Games, Inc. All Rights Reserved.
+
+#pragma once
+
+#include "AssetRegistry/AssetData.h"
+#include "Containers/ContainersFwd.h"
+#include "Delegates/Delegate.h"
+#include "Delegates/DelegateCombinations.h"
+#include "Delegates/IDelegateInstance.h"
+
+class FMenuBuilder;
+
+class FDMTextureSetContentBrowserIntegration
+{
+public:
+	DECLARE_MULTICAST_DELEGATE_TwoParams(FOnPopulateMenu, FMenuBuilder& InMenuBuilder,
+		const TArray<FAssetData>& InSelectedAssets)
+
+	DYNAMICMATERIALTEXTURESETEDITOR_API static FOnPopulateMenu::RegistrationType& GetPopulateExtenderDelegate();
+
+protected:
+	static FOnPopulateMenu PopulateMenuDelegate;
+};
