Material Designer
- Changed the order of UI operations when dragging multiple textures into the Material Designer.

diff --git a/Engine/Plugins/Experimental/DynamicMaterial/Source/DynamicMaterialEditor/Private/Slate/Layers/SDMSlotLayerView.cpp b/Engine/Plugins/Experimental/DynamicMaterial/Source/DynamicMaterialEditor/Private/Slate/Layers/SDMSlotLayerView.cpp
--- a/Engine/Plugins/Experimental/DynamicMaterial/Source/DynamicMaterialEditor/Private/Slate/Layers/SDMSlotLayerView.cpp
+++ b/Engine/Plugins/Experimental/DynamicMaterial/Source/DynamicMaterialEditor/Private/Slate/Layers/SDMSlotLayerView.cpp
@@ -16,6 +16,7 @@
 #include "Components/MaterialValues/DMMaterialValueTexture.h"
 #include "DMContentBrowserIntegration.h"
 #include "DMTextureSet.h"
+#include "DMTextureSetBlueprintFunctionLibrary.h"
 #include "DragAndDrop/AssetDragDropOp.h"
 #include "DynamicMaterialEditorCommands.h"
 #include "DynamicMaterialEditorModule.h"
@@ -681,24 +682,45 @@ void SDMSlotLayerView::HandleAssetDragDropOperation(FAssetDragDropOp& InAssetDra
 
 	if (TexturesDropped.Num() > 1)
 	{
-		const EAppReturnType::Type Result = FMessageDialog::Open(
-			EAppMsgType::YesNoCancel,
-			LOCTEXT("ReplaceSlotsMultiTexture", "You are about to import multiple Textures via a Material Designer Texture Set.\n\nDo you want to replace the slot contents?\n- Yes: All layers are deleted in the matching slots.\n- No: New texture layers are added to the matching slots.\n- Cancel: Abort this operation.")
-		);
+		UDMTextureSetBlueprintFunctionLibrary::CreateTextureSetFromAssetsInteractive(
+			Assets,
+			FDMTextureSetBuilderOnComplete::CreateWeakLambda(
+				Slot,
+				[Slot](UDMTextureSet* InTextureSet, bool bInWasAccepted)
+				{
+					if (!InTextureSet || !bInWasAccepted)
+					{
+						return;
+					}
 
-		switch (Result)
-		{
-			case EAppReturnType::No:
-				FDMContentBrowserIntegration::UpdateMaterialDesignerInstanceFromTextureSet(Assets, /* Replace */ false);
-				break;
+					UDynamicMaterialModelEditorOnlyData* EditorOnlyData = Slot->GetMaterialModelEditorOnlyData();
 
-			case EAppReturnType::Yes:
-				FDMContentBrowserIntegration::UpdateMaterialDesignerInstanceFromTextureSet(Assets, /* Replace */ true);
-				break;
+					if (!EditorOnlyData)
+					{
+						return;
+					}
 
-			default:
-				break;
-		}
+					const EAppReturnType::Type Result = FMessageDialog::Open(
+						EAppMsgType::YesNoCancel,
+						LOCTEXT("ReplaceSlotsMultiTexture", "You are about to import multiple Textures via a Material Designer Texture Set.\n\nDo you want to replace the slot contents?\n- Yes: All layers are deleted in the matching slots.\n- No: New texture layers are added to the matching slots.\n- Cancel: Abort this operation.")
+					);
+
+					switch (Result)
+					{
+						case EAppReturnType::No:
+							EditorOnlyData->AddTextureSet(InTextureSet, /* Replace */ false);
+							break;
+
+						case EAppReturnType::Yes:
+							EditorOnlyData->AddTextureSet(InTextureSet, /* Replace */ true);
+							break;
+
+						default:
+							break;
+					}
+				}
+			)
+		);
 
 		return;
 	}
