Material Designer
- Removed unneeded FullyLoad from new packages.

diff --git a/Engine/Plugins/Experimental/DynamicMaterial/Source/DynamicMaterial/Private/Components/DMMaterialValue.cpp b/Engine/Plugins/Experimental/DynamicMaterial/Source/DynamicMaterial/Private/Components/DMMaterialValue.cpp
--- a/Engine/Plugins/Experimental/DynamicMaterial/Source/DynamicMaterial/Private/Components/DMMaterialValue.cpp
+++ b/Engine/Plugins/Experimental/DynamicMaterial/Source/DynamicMaterial/Private/Components/DMMaterialValue.cpp
@@ -523,7 +523,6 @@ void UDMMaterialValue::CreatePreviewMaterial()
 		));
  
 		FAssetRegistryModule::AssetCreated(PreviewMaterial);
-		Package->FullyLoad();
 	}
 }
  
diff --git a/Engine/Plugins/Experimental/DynamicMaterial/Source/DynamicMaterialEditor/Private/Components/DMMaterialStage.cpp b/Engine/Plugins/Experimental/DynamicMaterial/Source/DynamicMaterialEditor/Private/Components/DMMaterialStage.cpp
--- a/Engine/Plugins/Experimental/DynamicMaterial/Source/DynamicMaterialEditor/Private/Components/DMMaterialStage.cpp
+++ b/Engine/Plugins/Experimental/DynamicMaterial/Source/DynamicMaterialEditor/Private/Components/DMMaterialStage.cpp
@@ -1466,7 +1466,6 @@ void UDMMaterialStage::CreatePreviewMaterial()
 		));
 
 		FAssetRegistryModule::AssetCreated(PreviewMaterialBase);
-		Package->FullyLoad();
 	}
 }
 
diff --git a/Engine/Plugins/Experimental/DynamicMaterial/Source/DynamicMaterialEditor/Private/Components/DMMaterialStageSource.cpp b/Engine/Plugins/Experimental/DynamicMaterial/Source/DynamicMaterialEditor/Private/Components/DMMaterialStageSource.cpp
--- a/Engine/Plugins/Experimental/DynamicMaterial/Source/DynamicMaterialEditor/Private/Components/DMMaterialStageSource.cpp
+++ b/Engine/Plugins/Experimental/DynamicMaterial/Source/DynamicMaterialEditor/Private/Components/DMMaterialStageSource.cpp
@@ -180,7 +180,6 @@ void UDMMaterialStageSource::CreatePreviewMaterial()
 		));
 
 		FAssetRegistryModule::AssetCreated(PreviewMaterial);
-		Package->FullyLoad();
 	}
 }
 
diff --git a/Engine/Plugins/Experimental/DynamicMaterial/Source/DynamicMaterialEditor/Private/DetailsPanel/Slate/SDMDetailsPanelTabSpawner.cpp b/Engine/Plugins/Experimental/DynamicMaterial/Source/DynamicMaterialEditor/Private/DetailsPanel/Slate/SDMDetailsPanelTabSpawner.cpp
--- a/Engine/Plugins/Experimental/DynamicMaterial/Source/DynamicMaterialEditor/Private/DetailsPanel/Slate/SDMDetailsPanelTabSpawner.cpp
+++ b/Engine/Plugins/Experimental/DynamicMaterial/Source/DynamicMaterialEditor/Private/DetailsPanel/Slate/SDMDetailsPanelTabSpawner.cpp
@@ -163,7 +163,6 @@ FReply SDMDetailsPanelTabSpawner::CreateDynamicMaterialModel()
 	));
 
 	FAssetRegistryModule::AssetCreated(NewModel);
-	Package->FullyLoad();
 
 	PropertyHandle->SetValueFromFormattedString(NewModel->GetPathName());
 
diff --git a/Engine/Plugins/Experimental/DynamicMaterial/Source/DynamicMaterialEditor/Private/Model/DynamicMaterialModelEditorOnlyData.cpp b/Engine/Plugins/Experimental/DynamicMaterial/Source/DynamicMaterialEditor/Private/Model/DynamicMaterialModelEditorOnlyData.cpp
--- a/Engine/Plugins/Experimental/DynamicMaterial/Source/DynamicMaterialEditor/Private/Model/DynamicMaterialModelEditorOnlyData.cpp
+++ b/Engine/Plugins/Experimental/DynamicMaterial/Source/DynamicMaterialEditor/Private/Model/DynamicMaterialModelEditorOnlyData.cpp
@@ -470,7 +470,6 @@ void UDynamicMaterialModelEditorOnlyData::CreateMaterial()
 		));
 
 		FAssetRegistryModule::AssetCreated(MaterialModel->DynamicMaterial);
-		Package->FullyLoad();
 	}
 
 	MaterialModel->DynamicMaterial->bOutputTranslucentVelocity = true;
diff --git a/Engine/Plugins/Experimental/DynamicMaterial/Source/DynamicMaterialEditor/Private/Utils/DMBlueprintFunctionLibrary.cpp b/Engine/Plugins/Experimental/DynamicMaterial/Source/DynamicMaterialEditor/Private/Utils/DMBlueprintFunctionLibrary.cpp
--- a/Engine/Plugins/Experimental/DynamicMaterial/Source/DynamicMaterialEditor/Private/Utils/DMBlueprintFunctionLibrary.cpp
+++ b/Engine/Plugins/Experimental/DynamicMaterial/Source/DynamicMaterialEditor/Private/Utils/DMBlueprintFunctionLibrary.cpp
@@ -304,7 +304,6 @@ bool UDMBlueprintFunctionLibrary::ExportMaterialInstance(UDynamicMaterialModel*
 	}
 
 	FAssetRegistryModule::AssetCreated(NewAsset);
-	Package->FullyLoad();
 
 	return true;
 }
@@ -355,7 +354,6 @@ bool UDMBlueprintFunctionLibrary::ExportGeneratedMaterial(UDynamicMaterialModel*
 	NewAsset->SetFlags(RF_Public | RF_Standalone);
 
 	FAssetRegistryModule::AssetCreated(NewAsset);
-	Package->FullyLoad();
 
 	return true;
 }
