Material Designer
- Fixed error with Texture Set not applying material channel to builder UI.
- Fixed _ORDp filter

diff --git a/Engine/Plugins/Experimental/DynamicMaterial/Source/DynamicMaterialEditor/Private/DMContentBrowserIntegration.cpp b/Engine/Plugins/Experimental/DynamicMaterial/Source/DynamicMaterialEditor/Private/DMContentBrowserIntegration.cpp
--- a/Engine/Plugins/Experimental/DynamicMaterial/Source/DynamicMaterialEditor/Private/DMContentBrowserIntegration.cpp
+++ b/Engine/Plugins/Experimental/DynamicMaterial/Source/DynamicMaterialEditor/Private/DMContentBrowserIntegration.cpp
@@ -123,7 +123,7 @@ void FDMContentBrowserIntegration::OnCreateMaterialDesignerInstanceFromTextureSe
 	FString UniquePackageName;
 	FString UniqueAssetName;
 
-	const FString BasePackageName = InPath / TEXT("NewMaterialDesignerInstance");
+	const FString BasePackageName = InPath / TEXT("MDI_NewMaterial");
 	AssetToolsModule.Get().CreateUniqueAssetName(BasePackageName, TEXT(""), UniquePackageName, UniqueAssetName);
 
 	UPackage* Package = CreatePackage(*UniquePackageName);
diff --git a/Engine/Plugins/Experimental/DynamicMaterial/Source/DynamicMaterialTextureSetEditor/Private/DMTextureSetSettings.cpp b/Engine/Plugins/Experimental/DynamicMaterial/Source/DynamicMaterialTextureSetEditor/Private/DMTextureSetSettings.cpp
--- a/Engine/Plugins/Experimental/DynamicMaterial/Source/DynamicMaterialTextureSetEditor/Private/DMTextureSetSettings.cpp
+++ b/Engine/Plugins/Experimental/DynamicMaterial/Source/DynamicMaterialTextureSetEditor/Private/DMTextureSetSettings.cpp
@@ -9,7 +9,7 @@ UDMTextureSetSettings::UDMTextureSetSettings()
 	int32 FilterCount = 0;
 
 	FDMTextureSetFilter BaseColor;
-	BaseColor.FilterStrings = {TEXT("Base_Color"), TEXT("BaseColor"), TEXT("Base_Colour"), TEXT("BaseColour"), TEXT("_BC"), TEXT("Diffuse"), TEXT("Albedo"), TEXT("_Diff")};
+	BaseColor.FilterStrings = {TEXT("Base_Color"), TEXT("BaseColor"), TEXT("Base_Colour"), TEXT("BaseColour"), TEXT("_BC"), TEXT("Diffuse"), TEXT("Albedo"), TEXT("_Diff"), TEXT("_D")};
 	BaseColor.MaterialProperties = {{EDMTextureSetMaterialProperty::BaseColor, EDMTextureChannelMask::RGBA}};
 	++FilterCount;
 
@@ -97,6 +97,15 @@ UDMTextureSetSettings::UDMTextureSetSettings()
 	};
 	++FilterCount;
 
+	FDMTextureSetFilter ORDp;
+	ORDp.FilterStrings = {TEXT("_ORDp")};
+	ORDp.MaterialProperties = {
+		{EDMTextureSetMaterialProperty::AmbientOcclusion, EDMTextureChannelMask::Red},
+		{EDMTextureSetMaterialProperty::Roughness, EDMTextureChannelMask::Green},
+		{EDMTextureSetMaterialProperty::Displacement, EDMTextureChannelMask::Blue}
+	};
+	++FilterCount;
+
 	FDMTextureSetFilter RM;
 	RM.FilterStrings = {TEXT("_RM")};
 	RM.MaterialProperties = {
@@ -133,6 +142,7 @@ UDMTextureSetSettings::UDMTextureSetSettings()
 	Filters.Add(MoveTemp(Refraction));
 	Filters.Add(MoveTemp(Tangent));
 	Filters.Add(MoveTemp(ORM));
+	Filters.Add(MoveTemp(ORDp));
 	Filters.Add(MoveTemp(RM));
 	Filters.Add(MoveTemp(RMA));
 }
diff --git a/Engine/Plugins/Experimental/DynamicMaterial/Source/DynamicMaterialTextureSetEditor/Private/Widgets/SDMTextureSetBuilder.cpp b/Engine/Plugins/Experimental/DynamicMaterial/Source/DynamicMaterialTextureSetEditor/Private/Widgets/SDMTextureSetBuilder.cpp
--- a/Engine/Plugins/Experimental/DynamicMaterial/Source/DynamicMaterialTextureSetEditor/Private/Widgets/SDMTextureSetBuilder.cpp
+++ b/Engine/Plugins/Experimental/DynamicMaterial/Source/DynamicMaterialTextureSetEditor/Private/Widgets/SDMTextureSetBuilder.cpp
@@ -51,7 +51,7 @@ void SDMTextureSetBuilder::Construct(const FArguments& InArgs, UDMTextureSet* In
 		TSharedRef<FDMTextureSetBuilderEntry> Entry = MakeShared<FDMTextureSetBuilderEntry>(
 			SetElement.Key,
 			SetElement.Value.Texture.LoadSynchronous(),
-			EDMTextureChannelMask::RGBA
+			SetElement.Value.TextureChannel
 		);
 
 		MaterialPropertyGrid->AddSlot(WidthIndex, HeightIndex)
