Material Designer
- Added normal global strength multiplier.
- Fixed normal global strength value from not correctly being assigned to the global parameter array.

#virtualized

diff --git a/Engine/Plugins/Experimental/DynamicMaterial/Source/DynamicMaterial/Private/Model/DynamicMaterialModel.cpp b/Engine/Plugins/Experimental/DynamicMaterial/Source/DynamicMaterial/Private/Model/DynamicMaterialModel.cpp
--- a/Engine/Plugins/Experimental/DynamicMaterial/Source/DynamicMaterial/Private/Model/DynamicMaterialModel.cpp
+++ b/Engine/Plugins/Experimental/DynamicMaterial/Source/DynamicMaterial/Private/Model/DynamicMaterialModel.cpp
@@ -8,41 +8,43 @@
 #include "DMComponentPath.h"
 #include "DMDefs.h"
 #include "DMValueDefinition.h"
+#include "DynamicMaterialModule.h"
 #include "Material/DynamicMaterialInstance.h"
 #include "Model/IDynamicMaterialModelEditorOnlyDataInterface.h"
 
 #if WITH_EDITOR
 #include "Components/DMTextureUV.h"
-#include "DynamicMaterialModule.h"
 #include "Materials/Material.h"
 #endif
 
 const FString UDynamicMaterialModel::ValuesPathToken = FString(TEXT("Values"));
 const FString UDynamicMaterialModel::ParametersPathToken = FString(TEXT("Parameters"));
-const FName UDynamicMaterialModel::GlobalOpacityValueName = FName(TEXT("GlobalOpacityValue"));
-const FName UDynamicMaterialModel::GlobalOpacityParameterName = FName(TEXT("GlobalOpacity"));
-const FName UDynamicMaterialModel::GlobalMetallicValueName = FName(TEXT("GlobalMetallicValue"));
-const FName UDynamicMaterialModel::GlobalMetallicParameterName = FName(TEXT("GlobalMetallic"));
-const FName UDynamicMaterialModel::GlobalSpecularValueName = FName(TEXT("GlobalSpecularValue"));
-const FName UDynamicMaterialModel::GlobalSpecularParameterName = FName(TEXT("GlobalSpecular"));
-const FName UDynamicMaterialModel::GlobalRoughnessValueName = FName(TEXT("GlobalRoughnessValue"));
-const FName UDynamicMaterialModel::GlobalRoughnessParameterName = FName(TEXT("GlobalRoughness"));
-const FName UDynamicMaterialModel::GlobalAnisotropyValueName = FName(TEXT("GlobalAnisotropyValue"));
-const FName UDynamicMaterialModel::GlobalAnisotropyParameterName = FName(TEXT("GlobalAnisotropy"));
-const FName UDynamicMaterialModel::GlobalWorldPositionOffsetValueName = FName(TEXT("GlobalWorldPositionOffsetValue"));
-const FName UDynamicMaterialModel::GlobalWorldPositionOffsetParameterName = FName(TEXT("GlobalWorldPositionOffset"));
-const FName UDynamicMaterialModel::GlobalAmbientOcclusionValueName = FName(TEXT("GlobalAmbientOcclusionValue"));
-const FName UDynamicMaterialModel::GlobalAmbientOcclusionParameterName = FName(TEXT("GlobalAmbientOcclusion"));
-const FName UDynamicMaterialModel::GlobalRefractionValueName = FName(TEXT("GlobalRefractionValue"));
-const FName UDynamicMaterialModel::GlobalRefractionParameterName = FName(TEXT("GlobalRefraction"));
-const FName UDynamicMaterialModel::GlobalPixelDepthOffsetValueName = FName(TEXT("GlobalPixelDepthOffsetValue"));
-const FName UDynamicMaterialModel::GlobalPixelDepthOffsetParameterName = FName(TEXT("GlobalPixelDepthOffset"));
-const FName UDynamicMaterialModel::GlobalOffsetValueName = FName(TEXT("GlobalOffsetValue"));
-const FName UDynamicMaterialModel::GlobalOffsetParameterName = FName(TEXT("GlobalOffset"));
-const FName UDynamicMaterialModel::GlobalTilingValueName = FName(TEXT("GlobalTilingValue"));
-const FName UDynamicMaterialModel::GlobalTilingParameterName = FName(TEXT("GlobalTiling"));
-const FName UDynamicMaterialModel::GlobalRotationValueName = FName(TEXT("GlobalRotationValue"));
-const FName UDynamicMaterialModel::GlobalRotationParameterName = FName(TEXT("GlobalRotation"));
+const FLazyName UDynamicMaterialModel::GlobalOpacityValueName = FLazyName(TEXT("GlobalOpacityValue"));
+const FLazyName UDynamicMaterialModel::GlobalOpacityParameterName = FLazyName(TEXT("GlobalOpacity"));
+const FLazyName UDynamicMaterialModel::GlobalMetallicValueName = FLazyName(TEXT("GlobalMetallicValue"));
+const FLazyName UDynamicMaterialModel::GlobalMetallicParameterName = FLazyName(TEXT("GlobalMetallic"));
+const FLazyName UDynamicMaterialModel::GlobalSpecularValueName = FLazyName(TEXT("GlobalSpecularValue"));
+const FLazyName UDynamicMaterialModel::GlobalSpecularParameterName = FLazyName(TEXT("GlobalSpecular"));
+const FLazyName UDynamicMaterialModel::GlobalRoughnessValueName = FLazyName(TEXT("GlobalRoughnessValue"));
+const FLazyName UDynamicMaterialModel::GlobalRoughnessParameterName = FLazyName(TEXT("GlobalRoughness"));
+const FLazyName UDynamicMaterialModel::GlobalNormalValueName = FLazyName(TEXT("GlobalNormalValue"));
+const FLazyName UDynamicMaterialModel::GlobalNormalParameterName = FLazyName(TEXT("GlobalNormal"));
+const FLazyName UDynamicMaterialModel::GlobalAnisotropyValueName = FLazyName(TEXT("GlobalAnisotropyValue"));
+const FLazyName UDynamicMaterialModel::GlobalAnisotropyParameterName = FLazyName(TEXT("GlobalAnisotropy"));
+const FLazyName UDynamicMaterialModel::GlobalWorldPositionOffsetValueName = FLazyName(TEXT("GlobalWorldPositionOffsetValue"));
+const FLazyName UDynamicMaterialModel::GlobalWorldPositionOffsetParameterName = FLazyName(TEXT("GlobalWorldPositionOffset"));
+const FLazyName UDynamicMaterialModel::GlobalAmbientOcclusionValueName = FLazyName(TEXT("GlobalAmbientOcclusionValue"));
+const FLazyName UDynamicMaterialModel::GlobalAmbientOcclusionParameterName = FLazyName(TEXT("GlobalAmbientOcclusion"));
+const FLazyName UDynamicMaterialModel::GlobalRefractionValueName = FLazyName(TEXT("GlobalRefractionValue"));
+const FLazyName UDynamicMaterialModel::GlobalRefractionParameterName = FLazyName(TEXT("GlobalRefraction"));
+const FLazyName UDynamicMaterialModel::GlobalPixelDepthOffsetValueName = FLazyName(TEXT("GlobalPixelDepthOffsetValue"));
+const FLazyName UDynamicMaterialModel::GlobalPixelDepthOffsetParameterName = FLazyName(TEXT("GlobalPixelDepthOffset"));
+const FLazyName UDynamicMaterialModel::GlobalOffsetValueName = FLazyName(TEXT("GlobalOffsetValue"));
+const FLazyName UDynamicMaterialModel::GlobalOffsetParameterName = FLazyName(TEXT("GlobalOffset"));
+const FLazyName UDynamicMaterialModel::GlobalTilingValueName = FLazyName(TEXT("GlobalTilingValue"));
+const FLazyName UDynamicMaterialModel::GlobalTilingParameterName = FLazyName(TEXT("GlobalTiling"));
+const FLazyName UDynamicMaterialModel::GlobalRotationValueName = FLazyName(TEXT("GlobalRotationValue"));
+const FLazyName UDynamicMaterialModel::GlobalRotationParameterName = FLazyName(TEXT("GlobalRotation"));
 
 UDynamicMaterialModel::UDynamicMaterialModel()
 {
@@ -96,6 +98,7 @@ UDynamicMaterialModel::UDynamicMaterialModel()
 	AddFloatParameter(GlobalMetallicValueName, GlobalMetallicParameterName);
 	AddFloatParameter(GlobalSpecularValueName, GlobalSpecularParameterName);
 	AddFloatParameter(GlobalRoughnessValueName, GlobalRoughnessParameterName);
+	AddFloatParameter(GlobalNormalValueName, GlobalNormalParameterName);
 	AddFloatParameter(GlobalAnisotropyValueName, GlobalAnisotropyParameterName);
 	AddFloatParameter(GlobalWorldPositionOffsetValueName, GlobalWorldPositionOffsetParameterName);
 	AddFloatParameter(GlobalAmbientOcclusionValueName, GlobalAmbientOcclusionParameterName);
@@ -418,11 +421,15 @@ void UDynamicMaterialModel::ResetData()
 		ModelEditorOnlyData->ResetData();
 	}
 }
+#endif
 
 void UDynamicMaterialModel::PostLoad()
 {
 	Super::PostLoad();
 
+	FixGlobalParameterValues();
+
+#if WITH_EDITOR
 	SetFlags(RF_Transactional);
 
 	PRAGMA_DISABLE_DEPRECATION_WARNINGS
@@ -448,8 +455,10 @@ void UDynamicMaterialModel::PostLoad()
 	}
 
 	ReinitComponents();
+#endif
 }
 
+#if WITH_EDITOR
 void UDynamicMaterialModel::PostEditUndo()
 {
 	Super::PostEditUndo();
@@ -528,7 +537,32 @@ void UDynamicMaterialModel::PostEditorDuplicate()
 		ModelEditorOnlyData->PostEditorDuplicate();
 	}
 }
+#endif
+
+void UDynamicMaterialModel::FixGlobalParameterValues()
+{
+	for (TPair<FName, TObjectPtr<UDMMaterialValue>>& ParameterPair : GlobalParameterValues)
+	{
+		if (!ParameterPair.Value || !ParameterPair.Value->HasAnyFlags(RF_ArchetypeObject))
+		{
+			continue;
+		}
+
+		const FName Name = ParameterPair.Value->GetFName();
+
+		UDMMaterialValue* Local = FindObjectFast<UDMMaterialValue>(this, Name, /* Exact class */ false);
+
+		if (!Local)
+		{
+			UE_LOG(LogDynamicMaterial, Error, TEXT("Archtype global parameter found and could not find local one. [%s]"), *Name.ToString());
+			continue;
+		}
+
+		ParameterPair.Value = Local;
+	}
+}
 
+#if WITH_EDITOR
 void UDynamicMaterialModel::ReinitComponents()
 {
 	if (IsValid(DynamicMaterial))
@@ -636,6 +670,7 @@ void UDynamicMaterialModel::FixGlobalVars()
 
 	FixGlobalVar(GlobalOpacityValueName, GlobalOpacityParameterName);
 	FixGlobalVar(GlobalRoughnessValueName, GlobalRoughnessParameterName);
+	FixGlobalVar(GlobalNormalValueName, GlobalNormalParameterName);
 	FixGlobalVar(GlobalSpecularValueName, GlobalSpecularParameterName);
 	FixGlobalVar(GlobalMetallicValueName, GlobalMetallicParameterName);
 	FixGlobalVar(GlobalAnisotropyValueName, GlobalAnisotropyParameterName);
diff --git a/Engine/Plugins/Experimental/DynamicMaterial/Source/DynamicMaterial/Public/Model/DynamicMaterialModel.h b/Engine/Plugins/Experimental/DynamicMaterial/Source/DynamicMaterial/Public/Model/DynamicMaterialModel.h
--- a/Engine/Plugins/Experimental/DynamicMaterial/Source/DynamicMaterial/Public/Model/DynamicMaterialModel.h
+++ b/Engine/Plugins/Experimental/DynamicMaterial/Source/DynamicMaterial/Public/Model/DynamicMaterialModel.h
@@ -43,30 +43,32 @@ class DYNAMICMATERIAL_API UDynamicMaterialModel : public UObject
 public:
 	static const FString ValuesPathToken;
 	static const FString ParametersPathToken;
-	static const FName GlobalOpacityValueName;
-	static const FName GlobalOpacityParameterName;
-	static const FName GlobalMetallicValueName;
-	static const FName GlobalMetallicParameterName;
-	static const FName GlobalRoughnessValueName;
-	static const FName GlobalRoughnessParameterName;
-	static const FName GlobalSpecularValueName;
-	static const FName GlobalSpecularParameterName;
-	static const FName GlobalAnisotropyValueName;
-	static const FName GlobalAnisotropyParameterName;
-	static const FName GlobalWorldPositionOffsetValueName;
-	static const FName GlobalWorldPositionOffsetParameterName;
-	static const FName GlobalAmbientOcclusionValueName;
-	static const FName GlobalAmbientOcclusionParameterName;
-	static const FName GlobalRefractionValueName;
-	static const FName GlobalRefractionParameterName;
-	static const FName GlobalPixelDepthOffsetValueName;
-	static const FName GlobalPixelDepthOffsetParameterName;
-	static const FName GlobalOffsetValueName;
-	static const FName GlobalOffsetParameterName;
-	static const FName GlobalTilingValueName;
-	static const FName GlobalTilingParameterName;
-	static const FName GlobalRotationValueName;
-	static const FName GlobalRotationParameterName;
+	static const FLazyName GlobalOpacityValueName;
+	static const FLazyName GlobalOpacityParameterName;
+	static const FLazyName GlobalMetallicValueName;
+	static const FLazyName GlobalMetallicParameterName;
+	static const FLazyName GlobalRoughnessValueName;
+	static const FLazyName GlobalRoughnessParameterName;
+	static const FLazyName GlobalNormalValueName;
+	static const FLazyName GlobalNormalParameterName;
+	static const FLazyName GlobalSpecularValueName;
+	static const FLazyName GlobalSpecularParameterName;
+	static const FLazyName GlobalAnisotropyValueName;
+	static const FLazyName GlobalAnisotropyParameterName;
+	static const FLazyName GlobalWorldPositionOffsetValueName;
+	static const FLazyName GlobalWorldPositionOffsetParameterName;
+	static const FLazyName GlobalAmbientOcclusionValueName;
+	static const FLazyName GlobalAmbientOcclusionParameterName;
+	static const FLazyName GlobalRefractionValueName;
+	static const FLazyName GlobalRefractionParameterName;
+	static const FLazyName GlobalPixelDepthOffsetValueName;
+	static const FLazyName GlobalPixelDepthOffsetParameterName;
+	static const FLazyName GlobalOffsetValueName;
+	static const FLazyName GlobalOffsetParameterName;
+	static const FLazyName GlobalTilingValueName;
+	static const FLazyName GlobalTilingParameterName;
+	static const FLazyName GlobalRotationValueName;
+	static const FLazyName GlobalRotationParameterName;
 
 	UDynamicMaterialModel();
 
@@ -163,7 +165,9 @@ public:
 	void ResetData();
 
 	//~ Begin UObject
+#endif
 	virtual void PostLoad() override;
+#if WITH_EDITOR
 	virtual void PostEditUndo() override;
 	virtual void PostEditImport() override;
 	virtual void PostDuplicate(bool bDuplicateForPIE) override;
@@ -198,6 +202,12 @@ protected:
 	TScriptInterface<IDynamicMaterialModelEditorOnlyDataInterface> EditorOnlyDataSI;
 #endif
 
+	/**
+	 * Because they are not assigned to individual properties, the global parameter values end up being their archetype versions
+	 * for versions of the object from before the change. This needs to be fixed on post load.
+	 */
+	void FixGlobalParameterValues();
+
 #if WITH_EDITOR
 	FName CreateUniqueParameterName(FName InBaseName);
 
diff --git a/Engine/Plugins/Experimental/DynamicMaterial/Source/DynamicMaterialEditor/Private/Components/MaterialProperties/DMMPNormal.cpp b/Engine/Plugins/Experimental/DynamicMaterial/Source/DynamicMaterialEditor/Private/Components/MaterialProperties/DMMPNormal.cpp
--- a/Engine/Plugins/Experimental/DynamicMaterial/Source/DynamicMaterialEditor/Private/Components/MaterialProperties/DMMPNormal.cpp
+++ b/Engine/Plugins/Experimental/DynamicMaterial/Source/DynamicMaterialEditor/Private/Components/MaterialProperties/DMMPNormal.cpp
@@ -2,6 +2,7 @@
 
 #include "Components/MaterialProperties/DMMPNormal.h"
 
+#include "Components/MaterialValues/DMMaterialValueFloat1.h"
 #include "Materials/MaterialExpressionMaterialFunctionCall.h"
 #include "Model/DMMaterialBuildState.h"
 #include "Model/DynamicMaterialModelEditorOnlyData.h"
@@ -18,6 +19,16 @@ namespace UE::DynamicMaterialEditor::Private
 
 		return NormalizeBlend;
 	}
+
+	UMaterialFunctionInterface* GetNormalMagnitude()
+	{
+		static UMaterialFunctionInterface* NormalMagnitude = FDMMaterialFunctionLibrary::Get().GetFunction(
+			"MF_DM_Normal_Magnitude",
+			TEXT("/Script/Engine.MaterialFunction'/DynamicMaterial/MaterialFunctions/MF_DM_Normal_Magnitude.MF_DM_Normal_Magnitude'")
+		);
+
+		return NormalMagnitude;
+	}
 }
 
 UDMMaterialPropertyNormal::UDMMaterialPropertyNormal()
@@ -90,3 +101,62 @@ void UDMMaterialPropertyNormal::AddOutputProcessor(const TSharedRef<FDMMaterialB
 
 	MaterialPropertyPtr->OutputIndex = 0;
 }
+
+void UDMMaterialPropertyNormal::AddAlphaMultiplier(const TSharedRef<FDMMaterialBuildState>& InBuildState) const
+{
+	UDMMaterialValueFloat1* AlphaValue = GetTypedComponent<UDMMaterialValueFloat1>(UDynamicMaterialModelEditorOnlyData::AlphaValueName);
+
+	if (!AlphaValue)
+	{
+		return;
+	}
+
+	FExpressionInput* MaterialPropertyPtr = InBuildState->GetMaterialProperty(MaterialProperty);
+
+	if (!MaterialPropertyPtr)
+	{
+		return;
+	}
+
+	UMaterialExpression* LastPropertyExpression = MaterialPropertyPtr->Expression;
+
+	if (!LastPropertyExpression)
+	{
+		return;
+	}
+
+	AlphaValue->GenerateExpression(InBuildState);
+
+	UMaterialExpression* GlobalOpacityExpression = InBuildState->GetLastValueExpression(AlphaValue);
+
+	if (!GlobalOpacityExpression)
+	{
+		return;
+	}
+
+	UMaterialExpressionMaterialFunctionCall* MaterialFunctionCall = FDMMaterialFunctionLibrary::Get().MakeExpression(
+		InBuildState->GetDynamicMaterial(),
+		UE::DynamicMaterialEditor::Private::GetNormalMagnitude(),
+		UE_DM_NodeComment_Default
+	);
+
+	FExpressionInput* FirstInput = MaterialFunctionCall->GetInput(0);
+
+	if (!FirstInput)
+	{
+		return;
+	}
+
+	FExpressionInput* SecondInput = MaterialFunctionCall->GetInput(1);
+
+	if (!SecondInput)
+	{
+		return;
+	}
+
+	LastPropertyExpression->ConnectExpression(FirstInput, MaterialPropertyPtr->OutputIndex);
+	GlobalOpacityExpression->ConnectExpression(SecondInput, 0);
+	MaterialFunctionCall->ConnectExpression(MaterialPropertyPtr, 0);
+
+	MaterialPropertyPtr->OutputIndex = 0;
+}
diff --git a/Engine/Plugins/Experimental/DynamicMaterial/Source/DynamicMaterialEditor/Private/Components/MaterialProperties/DMMPNormal.h b/Engine/Plugins/Experimental/DynamicMaterial/Source/DynamicMaterialEditor/Private/Components/MaterialProperties/DMMPNormal.h
--- a/Engine/Plugins/Experimental/DynamicMaterial/Source/DynamicMaterialEditor/Private/Components/MaterialProperties/DMMPNormal.h
+++ b/Engine/Plugins/Experimental/DynamicMaterial/Source/DynamicMaterialEditor/Private/Components/MaterialProperties/DMMPNormal.h
@@ -24,5 +24,10 @@ public:
 	 * Adds the SafeNormalize function as an extra output processor.
 	 */
 	virtual void AddOutputProcessor(const TSharedRef<FDMMaterialBuildState>& InBuildState) const;
+
+	/**
+	 * Replaces the regular multiply node with a custom one that doesn't affect the Z axis.
+	 */
+	virtual void AddAlphaMultiplier(const TSharedRef<FDMMaterialBuildState>& InBuildState) const override;
 	//~ End UDMMaterialProperty
 };
