Material Designer
- A slightly upgraded 5.4 to 5.5 transition for material editor only data.

diff --git a/Engine/Plugins/Experimental/DynamicMaterial/Source/DynamicMaterialEditor/Private/Model/DynamicMaterialModelEditorOnlyData.cpp b/Engine/Plugins/Experimental/DynamicMaterial/Source/DynamicMaterialEditor/Private/Model/DynamicMaterialModelEditorOnlyData.cpp
--- a/Engine/Plugins/Experimental/DynamicMaterial/Source/DynamicMaterialEditor/Private/Model/DynamicMaterialModelEditorOnlyData.cpp
+++ b/Engine/Plugins/Experimental/DynamicMaterial/Source/DynamicMaterialEditor/Private/Model/DynamicMaterialModelEditorOnlyData.cpp
@@ -187,6 +187,7 @@ void UDynamicMaterialModelEditorOnlyData::AssignPropertyAlphaValues()
 	Properties[EDMMaterialPropertyType::Metallic           ]->AddComponent(AlphaValueName, MaterialModel->GetGlobalParameterValue(UDynamicMaterialModel::GlobalMetallicValueName));
 	Properties[EDMMaterialPropertyType::Specular           ]->AddComponent(AlphaValueName, MaterialModel->GetGlobalParameterValue(UDynamicMaterialModel::GlobalSpecularValueName));
 	Properties[EDMMaterialPropertyType::Roughness          ]->AddComponent(AlphaValueName, MaterialModel->GetGlobalParameterValue(UDynamicMaterialModel::GlobalRoughnessValueName));
+	Properties[EDMMaterialPropertyType::Normal             ]->AddComponent(AlphaValueName, MaterialModel->GetGlobalParameterValue(UDynamicMaterialModel::GlobalNormalValueName));
 	Properties[EDMMaterialPropertyType::Anisotropy         ]->AddComponent(AlphaValueName, MaterialModel->GetGlobalParameterValue(UDynamicMaterialModel::GlobalAnisotropyValueName));
 	Properties[EDMMaterialPropertyType::WorldPositionOffset]->AddComponent(AlphaValueName, MaterialModel->GetGlobalParameterValue(UDynamicMaterialModel::GlobalWorldPositionOffsetValueName));
 	Properties[EDMMaterialPropertyType::AmbientOcclusion   ]->AddComponent(AlphaValueName, MaterialModel->GetGlobalParameterValue(UDynamicMaterialModel::GlobalAmbientOcclusionValueName));
@@ -1573,12 +1574,13 @@ void UDynamicMaterialModelEditorOnlyData::Serialize(FArchive& Ar)
 
 	const int32 Version = Ar.CustomVer(FDynamicMaterialModelEditorOnlyDataVersion::GUID);
 
+	// The default blend mode was changed from translucent and to opaque. If we have a
+	// transparent channel, let's set it back to translucent.
 	if (Version < FDynamicMaterialModelEditorOnlyDataVersion::GlobalValueRename)
 	{
 		if (ChannelListPreset == NAME_None)
 		{
-			// The default blend mode was changed from translucent and to opaque. If we have a
-			// transparent channel, let's set it back to translucent.
+			// Try to guess from the available slots.
 			if (GetSlotForMaterialProperty(EDMMaterialPropertyType::Opacity))
 			{
 				ChannelListPreset = "Translucent";
@@ -1591,6 +1593,12 @@ void UDynamicMaterialModelEditorOnlyData::Serialize(FArchive& Ar)
 				BlendMode = BLEND_Masked;
 				ShadingModel = EDMMaterialShadingModel::Unlit;
 			}
+			if (GetSlotForMaterialProperty(EDMMaterialPropertyType::EmissiveColor))
+			{
+				ChannelListPreset = "Translucent";
+				BlendMode = BLEND_Translucent;
+				ShadingModel = EDMMaterialShadingModel::Unlit;
+			}
 			// Else let's stay on opaque
 			else
 			{
