Material Designer
- Sequencer fixes for 5.4 to 5.5 upgrade

diff --git a/Engine/Plugins/Experimental/DynamicMaterial/Source/DynamicMaterial/Private/Model/DynamicMaterialModel.cpp b/Engine/Plugins/Experimental/DynamicMaterial/Source/DynamicMaterial/Private/Model/DynamicMaterialModel.cpp
--- a/Engine/Plugins/Experimental/DynamicMaterial/Source/DynamicMaterial/Private/Model/DynamicMaterialModel.cpp
+++ b/Engine/Plugins/Experimental/DynamicMaterial/Source/DynamicMaterial/Private/Model/DynamicMaterialModel.cpp
@@ -430,22 +430,29 @@ void UDynamicMaterialModel::PostLoad()
 	FixGlobalParameterValues();
 
 #if WITH_EDITOR
-	SetFlags(RF_Transactional);
-
 	PRAGMA_DISABLE_DEPRECATION_WARNINGS
 
+	// This requires a parameter name change which is editor-only code. It cannot correct itself at runtime.
+	// Open assets in the editor first to fix version upgrades!
 	if (GlobalOpacityValue)
 	{
-		if (UDMMaterialValueFloat1* NewOpacityValue = Cast<UDMMaterialValueFloat1>(GetGlobalParameterValue(GlobalOpacityValueName)))
+		GlobalParameterValues.FindOrAdd(GlobalOpacityValueName) = GlobalOpacityValue;
+
+		if (IDynamicMaterialModelEditorOnlyDataInterface* EditorOnlyDataInterface = GetEditorOnlyData())
 		{
-			NewOpacityValue->SetValue(GlobalOpacityValue->GetValue());
+			EditorOnlyDataInterface->SetPropertyComponent(EDMMaterialPropertyType::Opacity, TEXT("AlphaValue"), GlobalOpacityValue);
+			EditorOnlyDataInterface->SetPropertyComponent(EDMMaterialPropertyType::OpacityMask, TEXT("AlphaValue"), GlobalOpacityValue);
 		}
 
+		GlobalOpacityValue->SetParameterName("VALUE_GlobalOpacity");
+
 		GlobalOpacityValue = nullptr;
 	}
 
 	PRAGMA_ENABLE_DEPRECATION_WARNINGS
 
+	SetFlags(RF_Transactional);
+
 	IDynamicMaterialModelEditorOnlyDataInterface* ModelEditorOnlyData = GetEditorOnlyData();
 
 	if (!ModelEditorOnlyData)
diff --git a/Engine/Plugins/Experimental/DynamicMaterial/Source/DynamicMaterial/Public/Model/IDynamicMaterialModelEditorOnlyDataInterface.h b/Engine/Plugins/Experimental/DynamicMaterial/Source/DynamicMaterial/Public/Model/IDynamicMaterialModelEditorOnlyDataInterface.h
--- a/Engine/Plugins/Experimental/DynamicMaterial/Source/DynamicMaterial/Public/Model/IDynamicMaterialModelEditorOnlyDataInterface.h
+++ b/Engine/Plugins/Experimental/DynamicMaterial/Source/DynamicMaterial/Public/Model/IDynamicMaterialModelEditorOnlyDataInterface.h
@@ -49,6 +49,8 @@ public:
 
 #if WITH_EDITOR
 	virtual TSharedRef<IDMMaterialBuildStateInterface> CreateBuildStateInterface(UMaterial* InMaterialToBuild) const = 0;
+
+	virtual void SetPropertyComponent(EDMMaterialPropertyType InPropertyType, FName InComponentName, UDMMaterialComponent* InComponent) = 0;
 #endif
 
 	virtual UDMMaterialComponent* GetSubComponentByPath(FDMComponentPath& InPath) const = 0;
diff --git a/Engine/Plugins/Experimental/DynamicMaterial/Source/DynamicMaterialEditor/Private/Model/DynamicMaterialModelEditorOnlyData.cpp b/Engine/Plugins/Experimental/DynamicMaterial/Source/DynamicMaterialEditor/Private/Model/DynamicMaterialModelEditorOnlyData.cpp
--- a/Engine/Plugins/Experimental/DynamicMaterial/Source/DynamicMaterialEditor/Private/Model/DynamicMaterialModelEditorOnlyData.cpp
+++ b/Engine/Plugins/Experimental/DynamicMaterial/Source/DynamicMaterialEditor/Private/Model/DynamicMaterialModelEditorOnlyData.cpp
@@ -1028,6 +1028,14 @@ TSharedRef<IDMMaterialBuildStateInterface> UDynamicMaterialModelEditorOnlyData::
 	return CreateBuildState(InMaterialToBuild);
 }
 
+void UDynamicMaterialModelEditorOnlyData::SetPropertyComponent(EDMMaterialPropertyType InPropertyType, FName InComponentName, UDMMaterialComponent* InComponent)
+{
+	if (TObjectPtr<UDMMaterialProperty>* Property = Properties.Find(InPropertyType))
+	{
+		(*Property)->AddComponent(InComponentName, InComponent);
+	}
+}
+
 UDMMaterialComponent* UDynamicMaterialModelEditorOnlyData::GetSubComponentByPath(FDMComponentPath& InPath) const
 {
 	if (InPath.IsLeaf())
diff --git a/Engine/Plugins/Experimental/DynamicMaterial/Source/DynamicMaterialEditor/Public/Model/DynamicMaterialModelEditorOnlyData.h b/Engine/Plugins/Experimental/DynamicMaterial/Source/DynamicMaterialEditor/Public/Model/DynamicMaterialModelEditorOnlyData.h
--- a/Engine/Plugins/Experimental/DynamicMaterial/Source/DynamicMaterialEditor/Public/Model/DynamicMaterialModelEditorOnlyData.h
+++ b/Engine/Plugins/Experimental/DynamicMaterial/Source/DynamicMaterialEditor/Public/Model/DynamicMaterialModelEditorOnlyData.h
@@ -226,6 +226,7 @@ public:
 	virtual void OnTextureUVUpdated(UDMTextureUV* InTextureUV) override;
 	virtual void LoadDeprecatedModelData(UDynamicMaterialModel* InMaterialModel) override;
 	virtual TSharedRef<IDMMaterialBuildStateInterface> CreateBuildStateInterface(UMaterial* InMaterialToBuild) const override;
+	virtual void SetPropertyComponent(EDMMaterialPropertyType InPropertyType, FName InComponentName, UDMMaterialComponent* InComponent) override;
 	virtual UDMMaterialComponent* GetSubComponentByPath(FDMComponentPath& InPath) const override;
 	virtual UDMMaterialComponent* GetSubComponentByPath(FDMComponentPath& InPath, const FDMComponentPathSegment& InPathSegment) const override;
 	//~ End IDynamicMaterialModelEditorOnlyDataInterface
