Material Designer
- Fix for creating material designer instance in the content browser directly from a texture set.

diff --git a/Engine/Plugins/Experimental/DynamicMaterial/Source/DynamicMaterialEditor/Private/DMContentBrowserIntegration.cpp b/Engine/Plugins/Experimental/DynamicMaterial/Source/DynamicMaterialEditor/Private/DMContentBrowserIntegration.cpp
--- a/Engine/Plugins/Experimental/DynamicMaterial/Source/DynamicMaterialEditor/Private/DMContentBrowserIntegration.cpp
+++ b/Engine/Plugins/Experimental/DynamicMaterial/Source/DynamicMaterialEditor/Private/DMContentBrowserIntegration.cpp
@@ -113,6 +113,8 @@ void FDMContentBrowserIntegration::OnCreateMaterialDesignerInstanceFromTextureSe
 		return;
 	}
 
+	EditorOnlyData->SetChannelListPreset("All");
+
 	if (!EditorOnlyData->AddTextureSet(InTextureSet, /* Replace */ true))
 	{
 		return;
