Material Designer
- Fix for non-slot-based properties (global/material settings) not being keyframeable.

diff --git a/Engine/Plugins/Experimental/DynamicMaterial/Source/DynamicMaterialEditor/Private/Slate/SDMComponentEdit.cpp b/Engine/Plugins/Experimental/DynamicMaterial/Source/DynamicMaterialEditor/Private/Slate/SDMComponentEdit.cpp
--- a/Engine/Plugins/Experimental/DynamicMaterial/Source/DynamicMaterialEditor/Private/Slate/SDMComponentEdit.cpp
+++ b/Engine/Plugins/Experimental/DynamicMaterial/Source/DynamicMaterialEditor/Private/Slate/SDMComponentEdit.cpp
@@ -256,9 +256,19 @@ void SDMComponentEdit::Construct(const FArguments& InArgs, UDMMaterialComponent*
 
 	KeyframeHandler = nullptr;
 
-	if (InComponent)
+	UObject* WorldContext = InComponent;
+
+	if (!WorldContext)
+	{
+		if (TSharedPtr<SDMEditor> EditorWidget = InEditorWidget.Pin())
+		{
+			WorldContext = EditorWidget->GetMaterialModel();
+		}
+	}
+
+	if (WorldContext)
 	{
-		if (const UWorld* const World = InComponent->GetWorld())
+		if (const UWorld* const World = WorldContext->GetWorld())
 		{
 			if (const UDMWorldSubsystem* const WorldSubsystem = World->GetSubsystem<UDMWorldSubsystem>())
 			{
