Material Designer
- Fixed order of operation in the selecting of actors.

diff --git a/Engine/Plugins/Experimental/DynamicMaterial/Source/DynamicMaterialEditor/Private/Slate/SDMEditor.cpp b/Engine/Plugins/Experimental/DynamicMaterial/Source/DynamicMaterialEditor/Private/Slate/SDMEditor.cpp
--- a/Engine/Plugins/Experimental/DynamicMaterial/Source/DynamicMaterialEditor/Private/Slate/SDMEditor.cpp
+++ b/Engine/Plugins/Experimental/DynamicMaterial/Source/DynamicMaterialEditor/Private/Slate/SDMEditor.cpp
@@ -624,10 +624,7 @@ void SDMEditor::OnActorSelected(AActor* InActor)
 		SetMaterialObjectProperty(ActorProperties[0]);
 	}
 
-	if (GetMaterialActor() != InActor)
-	{
-		SetMaterialActor(InActor);
-	}
+	SetMaterialActor(InActor);
 }
 
 TSharedPtr<SDMSlot> SDMEditor::GetActiveSlotWidget() const
@@ -635,9 +632,10 @@ TSharedPtr<SDMSlot> SDMEditor::GetActiveSlotWidget() const
 	return ActiveSlotWidget;
 }
 
-TSharedRef<SWidget> SDMEditor::CreateActorMaterialSlotSelector(const AActor* InActor)
+TSharedRef<SWidget> SDMEditor::CreateActorMaterialSlotSelector(AActor* InActor)
 {
-	TArray<TSharedPtr<FDMObjectMaterialProperty>> MaterialProperties = Toolbar->GetMaterialProperties();
+	TArray<FDMObjectMaterialProperty> MaterialProperties = UDMBlueprintFunctionLibrary::GetActorMaterialProperties(InActor);
+
 	if (MaterialProperties.IsEmpty())
 	{
 		return 
@@ -662,7 +660,7 @@ TSharedRef<SWidget> SDMEditor::CreateActorMaterialSlotSelector(const AActor* InA
 
 	const UObject* CurrentOuter = nullptr;
 
-	for (const TSharedPtr<FDMObjectMaterialProperty>& MaterialSlot : MaterialProperties)
+	for (const FDMObjectMaterialProperty& MaterialSlot : MaterialProperties)
 	{
 		if (!MaterialSlot.IsValid())
 		{
@@ -670,18 +668,18 @@ TSharedRef<SWidget> SDMEditor::CreateActorMaterialSlotSelector(const AActor* InA
 		}
 
 		// Only show material slots on the selector
-		if (MaterialSlot->Property)
+		if (MaterialSlot.Property)
 		{
 			continue;
 		}
 
-		const UObject* Outer = MaterialSlot->OuterWeak.Get();
+		const UObject* Outer = MaterialSlot.OuterWeak.Get();
 		if (!IsValid(Outer))
 		{
 			continue;
 		}
 
-		UPrimitiveComponent* PrimComponent = Cast<UPrimitiveComponent>(MaterialSlot->OuterWeak.Get());
+		UPrimitiveComponent* PrimComponent = Cast<UPrimitiveComponent>(MaterialSlot.OuterWeak.Get());
 
 		if (!PrimComponent)
 		{
@@ -702,12 +700,10 @@ TSharedRef<SWidget> SDMEditor::CreateActorMaterialSlotSelector(const AActor* InA
 			CurrentOuter = Outer;
 		}
 
-		TWeakPtr<FDMObjectMaterialProperty> MaterialSlotWeak = MaterialSlot;
-
 		constexpr int32 ThumbnailSize = 48;
 
 		TSharedRef<FAssetThumbnail> Thumbnail = MakeShared<FAssetThumbnail>(
-			PrimComponent->GetMaterial(MaterialSlot->Index),
+			PrimComponent->GetMaterial(MaterialSlot.Index),
 			ThumbnailSize,
 			ThumbnailSize, 
 			GetThumbnailPool()
@@ -742,7 +738,7 @@ TSharedRef<SWidget> SDMEditor::CreateActorMaterialSlotSelector(const AActor* InA
 					[
 						SNew(STextBlock)
 						.TextStyle(FDynamicMaterialEditorStyle::Get(), "RegularFont")
-						.Text(MaterialSlot->GetPropertyName(true))
+						.Text(MaterialSlot.GetPropertyName(true))
 					]
 
 					+ SVerticalBox::Slot()
@@ -751,7 +747,7 @@ TSharedRef<SWidget> SDMEditor::CreateActorMaterialSlotSelector(const AActor* InA
 					[
 						SNew(SButton)
 						.ContentPadding(FMargin(2.f, 2.f, 2.f, 2.f))
-						.OnClicked(this, &SDMEditor::OnCreateMaterialButtonClicked, MaterialSlotWeak)
+						.OnClicked(this, &SDMEditor::OnCreateMaterialButtonClicked, MaterialSlot)
 						[
 							SNew(STextBlock)
 							.TextStyle(FDynamicMaterialEditorStyle::Get(), "RegularFont")
@@ -1302,25 +1298,22 @@ void SDMEditor::OnSlotCheckStateChanged(ECheckBoxState InCheckState, EDMMaterial
 	return;
 }
 
-FReply SDMEditor::OnCreateMaterialButtonClicked(TWeakPtr<FDMObjectMaterialProperty> InMaterialProperty)
+FReply SDMEditor::OnCreateMaterialButtonClicked(FDMObjectMaterialProperty InMaterialProperty)
 {
-	if (TSharedPtr<FDMObjectMaterialProperty> MaterialProperty = InMaterialProperty.Pin())
-	{
-		UDynamicMaterialModel* NewModel = UDMBlueprintFunctionLibrary::CreateDynamicMaterialInObject(*MaterialProperty);
+	UDynamicMaterialModel* NewModel = UDMBlueprintFunctionLibrary::CreateDynamicMaterialInObject(InMaterialProperty);
 
-		if (NewModel)
+	if (NewModel)
+	{
+		if (Toolbar.IsValid())
 		{
-			if (Toolbar.IsValid())
-			{
-				Toolbar->SetMaterialModel(NewModel);
-			}
+			Toolbar->SetMaterialModel(NewModel);
+		}
 
-			if (AActor* Actor = MaterialProperty->GetTypedOuter<AActor>())
+		if (AActor* Actor = InMaterialProperty.GetTypedOuter<AActor>())
+		{
+			if (Toolbar->GetMaterialActor() != Actor)
 			{
-				if (Toolbar->GetMaterialActor() != Actor)
-				{
-					Toolbar->SetMaterialActor(Actor);
-				}
+				Toolbar->SetMaterialActor(Actor);
 			}
 		}
 	}
diff --git a/Engine/Plugins/Experimental/DynamicMaterial/Source/DynamicMaterialEditor/Private/Slate/SDMEditor.h b/Engine/Plugins/Experimental/DynamicMaterial/Source/DynamicMaterialEditor/Private/Slate/SDMEditor.h
--- a/Engine/Plugins/Experimental/DynamicMaterial/Source/DynamicMaterialEditor/Private/Slate/SDMEditor.h
+++ b/Engine/Plugins/Experimental/DynamicMaterial/Source/DynamicMaterialEditor/Private/Slate/SDMEditor.h
@@ -164,7 +164,7 @@ protected:
 
 	void OnComponentUpdated(UDMMaterialComponent* InComponent, EDMUpdateType InUpdateType);
 
-	TSharedRef<SWidget> CreateActorMaterialSlotSelector(const AActor* InActor);
+	TSharedRef<SWidget> CreateActorMaterialSlotSelector(AActor* InActor);
 
 	void OnMaterialBuilt(UDynamicMaterialModel* InMaterialModel);
 	void OnValuesUpdated(UDynamicMaterialModel* InMaterialModel);
@@ -175,7 +175,7 @@ protected:
 	FText GetToolTipForProperty(EDMMaterialPropertyType InProperty) const;
 	void OnSlotCheckStateChanged(ECheckBoxState InCheckState, EDMMaterialPropertyType InProperty);
 
-	FReply OnCreateMaterialButtonClicked(TWeakPtr<FDMObjectMaterialProperty> InMaterialProperty);
+	FReply OnCreateMaterialButtonClicked(FDMObjectMaterialProperty InMaterialProperty);
 
 	void OnToolBarPropertyChanged(TSharedPtr<FDMObjectMaterialProperty> InNewSelectedProperty);
 	TSharedRef<SWidget> MakeToolBarSettingsMenu();
