Material Designer
- Removed currently redundant material property display from the layer row.

diff --git a/Engine/Plugins/Experimental/DynamicMaterial/Source/DynamicMaterialEditor/Private/Slate/Layers/SDMSlotLayerItem.cpp b/Engine/Plugins/Experimental/DynamicMaterial/Source/DynamicMaterialEditor/Private/Slate/Layers/SDMSlotLayerItem.cpp
--- a/Engine/Plugins/Experimental/DynamicMaterial/Source/DynamicMaterialEditor/Private/Slate/Layers/SDMSlotLayerItem.cpp
+++ b/Engine/Plugins/Experimental/DynamicMaterial/Source/DynamicMaterialEditor/Private/Slate/Layers/SDMSlotLayerItem.cpp
@@ -238,16 +238,6 @@ TSharedRef<SWidget> SDMSlotLayerItem::CreateHeaderRowContent()
 			.FillHeight(1.0f)
 			.HAlign(HAlign_Left)
 			.VAlign(VAlign_Center)
-			.Padding(0.0f, 0.0f, 0.0f, 0.0f)
-			[
-				SNew(STextBlock)
-				.TextStyle(FDynamicMaterialEditorStyle::Get(), "SmallFont")
-				.Text(this, &SDMSlotLayerItem::GetBlendModeText)
-			]
-			+ SVerticalBox::Slot()
-			.FillHeight(1.0f)
-			.HAlign(HAlign_Left)
-			.VAlign(VAlign_Center)
 			.Padding(0.0f, 0.0f, 0.0f, 2.0f)
 			[
 				SNew(STextBlock)
