Material Designer
- Removed the limits on the global adjustment sliders.

diff --git a/Engine/Plugins/Experimental/DynamicMaterial/Source/DynamicMaterial/Private/Model/DynamicMaterialModel.cpp b/Engine/Plugins/Experimental/DynamicMaterial/Source/DynamicMaterial/Private/Model/DynamicMaterialModel.cpp
--- a/Engine/Plugins/Experimental/DynamicMaterial/Source/DynamicMaterial/Private/Model/DynamicMaterialModel.cpp
+++ b/Engine/Plugins/Experimental/DynamicMaterial/Source/DynamicMaterial/Private/Model/DynamicMaterialModel.cpp
@@ -54,17 +54,12 @@ UDynamicMaterialModel::UDynamicMaterialModel()
 	EditorOnlyDataSI.SetObject(nullptr);
 #endif
 
-	auto AddFloatParameter = [this](FName InPropertyName, FName InParameterName, float InDefaultValue = 1.f, bool bSetValueRange = true)
+	auto AddFloatParameter = [this](FName InPropertyName, FName InParameterName, float InDefaultValue = 1.f)
 		{
 			UDMMaterialValueFloat1* Property = CreateDefaultSubobject<UDMMaterialValueFloat1>(InPropertyName);
 			GlobalParameterValues.Add(InPropertyName, Property);
 
 #if WITH_EDITOR
-			if (bSetValueRange)
-			{
-				Property->SetValueRange({0.f, 1.f});
-			}
-
 			Property->SetDefaultValue(InDefaultValue);
 			Property->ApplyDefaultValue();
 #endif
@@ -107,7 +102,7 @@ UDynamicMaterialModel::UDynamicMaterialModel()
 
 	AddVector2Parameter(GlobalOffsetValueName, GlobalOffsetParameterName, FVector2D::ZeroVector);
 	AddVector2Parameter(GlobalTilingValueName, GlobalTilingParameterName, FVector2D::UnitVector);
-	AddFloatParameter(GlobalRotationValueName, GlobalRotationParameterName, /* Default Value */ 0.f, /* Set Value Range */ false);
+	AddFloatParameter(GlobalRotationValueName, GlobalRotationParameterName, /* Default Value */ 0.f);
 }
 
 void UDynamicMaterialModel::SetDynamicMaterialInstance(UDynamicMaterialInstance* InDynamicMaterialInstance)
