Material Designer
- Conversion of text renderer to use STextBlock.
- Fix for renderers not refreshing the component edit viewport when its structure changes.

diff --git a/Engine/Plugins/Experimental/DynamicMaterial/Source/DynamicMaterial/DynamicMaterial.Build.cs b/Engine/Plugins/Experimental/DynamicMaterial/Source/DynamicMaterial/DynamicMaterial.Build.cs
--- a/Engine/Plugins/Experimental/DynamicMaterial/Source/DynamicMaterial/DynamicMaterial.Build.cs
+++ b/Engine/Plugins/Experimental/DynamicMaterial/Source/DynamicMaterial/DynamicMaterial.Build.cs
@@ -16,6 +16,7 @@ public class DynamicMaterial : ModuleRules
 				"DynamicMaterialTextureSet",
 				"Slate",
 				"SlateCore",
+				"StructUtils",
 				"UMG"
 			}
 		);
diff --git a/Engine/Plugins/Experimental/DynamicMaterial/Source/DynamicMaterial/Private/Components/DMRenderTargetRenderer.cpp b/Engine/Plugins/Experimental/DynamicMaterial/Source/DynamicMaterial/Private/Components/DMRenderTargetRenderer.cpp
--- a/Engine/Plugins/Experimental/DynamicMaterial/Source/DynamicMaterial/Private/Components/DMRenderTargetRenderer.cpp
+++ b/Engine/Plugins/Experimental/DynamicMaterial/Source/DynamicMaterial/Private/Components/DMRenderTargetRenderer.cpp
@@ -76,6 +76,33 @@ bool UDMRenderTargetRenderer::JsonDeserialize(const TSharedPtr<FJsonValue>& InJs
 	return false;
 }
 
+void UDMRenderTargetRenderer::Update(EDMUpdateType InUpdateType)
+{
+	if (!IsComponentValid())
+	{
+		return;
+	}
+
+#if WITH_EDITOR
+	if (HasComponentBeenRemoved())
+	{
+		return;
+	}
+
+	if (InUpdateType == EDMUpdateType::Structure)
+	{
+		MarkComponentDirty();
+	}
+
+	if (UDMMaterialValueRenderTarget* RenderTarget = GetRenderTargetValue())
+	{
+		RenderTarget->Update(InUpdateType);
+	}
+#endif
+
+	Super::Update(InUpdateType);
+}
+
 void UDMRenderTargetRenderer::PostLoad()
 {
 	Super::PostLoad();
diff --git a/Engine/Plugins/Experimental/DynamicMaterial/Source/DynamicMaterial/Private/Components/RenderTargetRenderers/DMRenderTargetTextRenderer.cpp b/Engine/Plugins/Experimental/DynamicMaterial/Source/DynamicMaterial/Private/Components/RenderTargetRenderers/DMRenderTargetTextRenderer.cpp
--- a/Engine/Plugins/Experimental/DynamicMaterial/Source/DynamicMaterial/Private/Components/RenderTargetRenderers/DMRenderTargetTextRenderer.cpp
+++ b/Engine/Plugins/Experimental/DynamicMaterial/Source/DynamicMaterial/Private/Components/RenderTargetRenderers/DMRenderTargetTextRenderer.cpp
@@ -1,13 +1,13 @@
 // Copyright Epic Games, Inc. All Rights Reserved.
 
 #include "Components/RenderTargetRenderers/DMRenderTargetTextRenderer.h"
-#include "CanvasItem.h"
-#include "CanvasTypes.h"
 #include "Components/MaterialValues/DMMaterialValueRenderTarget.h"
-#include "Engine/Canvas.h"
+#include "Dom/JsonObject.h"
+#include "Dom/JsonValue.h"
 #include "Engine/Engine.h"
-#include "Engine/TextureRenderTarget2D.h"
 #include "TextureResource.h"
+#include "Widgets/SBoxPanel.h"
+#include "Widgets/Text/STextBlock.h"
 
 #if WITH_EDITOR
 #include "Dom/JsonValue.h"
@@ -24,32 +24,42 @@ namespace UE::DynamicMaterial::Private
 #if WITH_EDITOR
 struct FDMRenderTargetTextRenderer
 {
-	static const inline FName FontName = GET_MEMBER_NAME_CHECKED(UDMRenderTargetTextRenderer, Font);
+	static const inline FName FontInfoName = GET_MEMBER_NAME_CHECKED(UDMRenderTargetTextRenderer, FontInfo);
 	static const inline FName TextName = GET_MEMBER_NAME_CHECKED(UDMRenderTargetTextRenderer, Text);
 	static const inline FName TextColorName = GET_MEMBER_NAME_CHECKED(UDMRenderTargetTextRenderer, TextColor);
+	static const inline FName HasHighlightName = GET_MEMBER_NAME_CHECKED(UDMRenderTargetTextRenderer, bHasHighlight);
+	static const inline FName HighlightColorName = GET_MEMBER_NAME_CHECKED(UDMRenderTargetTextRenderer, HighlightColor);
+	static const inline FName HasShadowName = GET_MEMBER_NAME_CHECKED(UDMRenderTargetTextRenderer, bHasShadow);
+	static const inline FName ShadowColorName = GET_MEMBER_NAME_CHECKED(UDMRenderTargetTextRenderer, ShadowColor);
+	static const inline FName ShadowOffsetName = GET_MEMBER_NAME_CHECKED(UDMRenderTargetTextRenderer, ShadowOffset);
+	static const inline FName AutoWrapTextName = GET_MEMBER_NAME_CHECKED(UDMRenderTargetTextRenderer, bAutoWrapText);
+	static const inline FName WrapTextAtName = GET_MEMBER_NAME_CHECKED(UDMRenderTargetTextRenderer, WrapTextAt);
+	static const inline FName WrappingPolicyName = GET_MEMBER_NAME_CHECKED(UDMRenderTargetTextRenderer, WrappingPolicy);
 	static const inline FName JustifyName = GET_MEMBER_NAME_CHECKED(UDMRenderTargetTextRenderer, Justify);
-	static const inline FName KerningName = GET_MEMBER_NAME_CHECKED(UDMRenderTargetTextRenderer, Kerning);
+	static const inline FName TransformPolicyName = GET_MEMBER_NAME_CHECKED(UDMRenderTargetTextRenderer, TransformPolicy);
+	static const inline FName FlowDirectionName = GET_MEMBER_NAME_CHECKED(UDMRenderTargetTextRenderer, FlowDirection);
+	static const inline FName ShapingMethodName = GET_MEMBER_NAME_CHECKED(UDMRenderTargetTextRenderer, ShapingMethod);
+	static const inline FName StrikeBrushName = GET_MEMBER_NAME_CHECKED(UDMRenderTargetTextRenderer, StrikeBrush);
 	static const inline FName LineHeightName = GET_MEMBER_NAME_CHECKED(UDMRenderTargetTextRenderer, LineHeight);
 	static const inline FName PaddingLeftName = GET_MEMBER_NAME_CHECKED(UDMRenderTargetTextRenderer, PaddingLeft);
 	static const inline FName PaddingRightName = GET_MEMBER_NAME_CHECKED(UDMRenderTargetTextRenderer, PaddingRight);
 	static const inline FName PaddingTopName = GET_MEMBER_NAME_CHECKED(UDMRenderTargetTextRenderer, PaddingTop);
 	static const inline FName PaddingBottomName = GET_MEMBER_NAME_CHECKED(UDMRenderTargetTextRenderer, PaddingBottom);
-	static const inline FName TextScaleName = GET_MEMBER_NAME_CHECKED(UDMRenderTargetTextRenderer, TextScale);
-	static const inline FName TextureSizeOverrideName = GET_MEMBER_NAME_CHECKED(UDMRenderTargetTextRenderer, TextureSizeOverride);
-	static const inline FName bOutlineName = GET_MEMBER_NAME_CHECKED(UDMRenderTargetTextRenderer, bOutline);
-	static const inline FName OutlineColorName = GET_MEMBER_NAME_CHECKED(UDMRenderTargetTextRenderer, OutlineColor);
-	static const inline FName bShadowName = GET_MEMBER_NAME_CHECKED(UDMRenderTargetTextRenderer, bShadow);
-	static const inline FName ShadowColorName = GET_MEMBER_NAME_CHECKED(UDMRenderTargetTextRenderer, ShadowColor);
-	static const inline FName ShadowOffsetName = GET_MEMBER_NAME_CHECKED(UDMRenderTargetTextRenderer, ShadowOffset);
-	static const inline FName bGlowName = GET_MEMBER_NAME_CHECKED(UDMRenderTargetTextRenderer, bGlow);
-	static const inline FName GlowColorName = GET_MEMBER_NAME_CHECKED(UDMRenderTargetTextRenderer, GlowColor);
-	static const inline FName GlowInnerRadiusName = GET_MEMBER_NAME_CHECKED(UDMRenderTargetTextRenderer, GlowInnerRadius);
-	static const inline FName GlowOuterRadiusName = GET_MEMBER_NAME_CHECKED(UDMRenderTargetTextRenderer, GlowOuterRadius);
+	static const inline FName OverrideRenderTargetSizeName = GET_MEMBER_NAME_CHECKED(UDMRenderTargetTextRenderer, bOverrideRenderTargetSize);
 
 	static const inline TSet<FName> PropertyNames = {
-		FontName, TextName, TextColorName, JustifyName, KerningName, LineHeightName, PaddingLeftName, PaddingRightName,
-		PaddingTopName, PaddingBottomName, TextScaleName, TextureSizeOverrideName, bOutlineName, OutlineColorName, bShadowName,
-		ShadowColorName, ShadowOffsetName, bGlowName, GlowColorName, GlowInnerRadiusName, GlowOuterRadiusName
+		FontInfoName, TextName, TextColorName,
+		HasHighlightName, HighlightColorName,
+		HasShadowName, ShadowColorName, ShadowOffsetName,
+		AutoWrapTextName, WrapTextAtName, WrappingPolicyName,
+		JustifyName,
+		TransformPolicyName,
+		FlowDirectionName,
+		ShapingMethodName,
+		StrikeBrushName,
+		LineHeightName,
+		PaddingLeftName, PaddingTopName, PaddingRightName, PaddingBottomName,
+		OverrideRenderTargetSizeName
 	};
 };
 #endif
@@ -64,11 +74,166 @@ UDMRenderTargetTextRenderer::UDMRenderTargetTextRenderer()
 #if WITH_EDITOR
 TSharedPtr<FJsonValue> UDMRenderTargetTextRenderer::JsonSerialize() const
 {
-	return MakeShared<FJsonValueNull>();
+	TSharedRef<FJsonObject> Object = MakeShared<FJsonObject>();
+
+	Object->SetField(FDMRenderTargetTextRenderer::FontInfoName.GetPlainNameString(), FDMJsonUtils::Serialize<FSlateFontInfo>(FontInfo));
+	Object->SetField(FDMRenderTargetTextRenderer::TextName.GetPlainNameString(), FDMJsonUtils::Serialize(Text));
+	Object->SetField(FDMRenderTargetTextRenderer::TextColorName.GetPlainNameString(), FDMJsonUtils::Serialize(TextColor));
+	Object->SetField(FDMRenderTargetTextRenderer::HasHighlightName.GetPlainNameString(), FDMJsonUtils::Serialize(bHasHighlight));
+	Object->SetField(FDMRenderTargetTextRenderer::HighlightColorName.GetPlainNameString(), FDMJsonUtils::Serialize(HighlightColor));
+	Object->SetField(FDMRenderTargetTextRenderer::HasShadowName.GetPlainNameString(), FDMJsonUtils::Serialize(bHasShadow));
+	Object->SetField(FDMRenderTargetTextRenderer::ShadowColorName.GetPlainNameString(), FDMJsonUtils::Serialize(ShadowColor));
+	Object->SetField(FDMRenderTargetTextRenderer::ShadowOffsetName.GetPlainNameString(), FDMJsonUtils::Serialize(ShadowOffset));
+	Object->SetField(FDMRenderTargetTextRenderer::AutoWrapTextName.GetPlainNameString(), FDMJsonUtils::Serialize(bAutoWrapText));
+	Object->SetField(FDMRenderTargetTextRenderer::WrapTextAtName.GetPlainNameString(), FDMJsonUtils::Serialize(WrapTextAt));
+	Object->SetField(FDMRenderTargetTextRenderer::WrappingPolicyName.GetPlainNameString(), FDMJsonUtils::Serialize(WrappingPolicy));
+	Object->SetField(FDMRenderTargetTextRenderer::JustifyName.GetPlainNameString(), FDMJsonUtils::Serialize(Justify.GetValue()));
+	Object->SetField(FDMRenderTargetTextRenderer::TransformPolicyName.GetPlainNameString(), FDMJsonUtils::Serialize(TransformPolicy));
+	Object->SetField(FDMRenderTargetTextRenderer::FlowDirectionName.GetPlainNameString(), FDMJsonUtils::Serialize(FlowDirection));
+	Object->SetField(FDMRenderTargetTextRenderer::ShapingMethodName.GetPlainNameString(), FDMJsonUtils::Serialize(ShapingMethod));
+	Object->SetField(FDMRenderTargetTextRenderer::LineHeightName.GetPlainNameString(), FDMJsonUtils::Serialize(LineHeight));
+	Object->SetField(FDMRenderTargetTextRenderer::PaddingLeftName.GetPlainNameString(), FDMJsonUtils::Serialize(PaddingLeft));
+	Object->SetField(FDMRenderTargetTextRenderer::PaddingRightName.GetPlainNameString(), FDMJsonUtils::Serialize(PaddingRight));
+	Object->SetField(FDMRenderTargetTextRenderer::PaddingTopName.GetPlainNameString(), FDMJsonUtils::Serialize(PaddingTop));
+	Object->SetField(FDMRenderTargetTextRenderer::PaddingBottomName.GetPlainNameString(), FDMJsonUtils::Serialize(PaddingBottom));
+	Object->SetField(FDMRenderTargetTextRenderer::OverrideRenderTargetSizeName.GetPlainNameString(), FDMJsonUtils::Serialize(bOverrideRenderTargetSize));
+
+	if (StrikeBrush.IsValid())
+	{
+		Object->SetField(FDMRenderTargetTextRenderer::StrikeBrushName.GetPlainNameString(), FDMJsonUtils::Serialize<FSlateBrush>(StrikeBrush.Get<FSlateBrush>()));
+	}
+	else
+	{
+		Object->SetField(FDMRenderTargetTextRenderer::StrikeBrushName.GetPlainNameString(), MakeShared<FJsonValueNull>());
+	}
+
+	return MakeShared<FJsonValueObject>(Object);
 }
 
 bool UDMRenderTargetTextRenderer::JsonDeserialize(const TSharedPtr<FJsonValue>& InJsonValue)
 {
+	TSharedPtr<FJsonObject> Object = InJsonValue->AsObject();
+
+	if (!Object.IsValid())
+	{
+		return false;
+	}
+
+	if (const TSharedPtr<FJsonValue>* JsonValue = Object->Values.Find(FDMRenderTargetTextRenderer::FontInfoName.GetPlainNameString()))
+	{
+		FDMJsonUtils::Deserialize<FSlateFontInfo>(*JsonValue, FontInfo);
+	}
+
+	if (const TSharedPtr<FJsonValue>* JsonValue = Object->Values.Find(FDMRenderTargetTextRenderer::TextName.GetPlainNameString()))
+	{
+		FDMJsonUtils::Deserialize(*JsonValue, Text);
+	}
+
+	if (const TSharedPtr<FJsonValue>* JsonValue = Object->Values.Find(FDMRenderTargetTextRenderer::TextColorName.GetPlainNameString()))
+	{
+		FDMJsonUtils::Deserialize(*JsonValue, TextColor);
+	}
+
+	if (const TSharedPtr<FJsonValue>* JsonValue = Object->Values.Find(FDMRenderTargetTextRenderer::HasHighlightName.GetPlainNameString()))
+	{
+		FDMJsonUtils::Deserialize(*JsonValue, bHasHighlight);
+	}
+
+	if (const TSharedPtr<FJsonValue>* JsonValue = Object->Values.Find(FDMRenderTargetTextRenderer::HighlightColorName.GetPlainNameString()))
+	{
+		FDMJsonUtils::Deserialize(*JsonValue, HighlightColor);
+	}
+
+	if (const TSharedPtr<FJsonValue>* JsonValue = Object->Values.Find(FDMRenderTargetTextRenderer::HasShadowName.GetPlainNameString()))
+	{
+		FDMJsonUtils::Deserialize(*JsonValue, bHasShadow);
+	}
+
+	if (const TSharedPtr<FJsonValue>* JsonValue = Object->Values.Find(FDMRenderTargetTextRenderer::ShadowColorName.GetPlainNameString()))
+	{
+		FDMJsonUtils::Deserialize(*JsonValue, ShadowColor);
+	}
+
+	if (const TSharedPtr<FJsonValue>* JsonValue = Object->Values.Find(FDMRenderTargetTextRenderer::ShadowOffsetName.GetPlainNameString()))
+	{
+		FDMJsonUtils::Deserialize(*JsonValue, ShadowOffset);
+	}
+
+	if (const TSharedPtr<FJsonValue>* JsonValue = Object->Values.Find(FDMRenderTargetTextRenderer::AutoWrapTextName.GetPlainNameString()))
+	{
+		FDMJsonUtils::Deserialize(*JsonValue, bAutoWrapText);
+	}
+
+	if (const TSharedPtr<FJsonValue>* JsonValue = Object->Values.Find(FDMRenderTargetTextRenderer::WrapTextAtName.GetPlainNameString()))
+	{
+		FDMJsonUtils::Deserialize(*JsonValue, WrapTextAt);
+	}
+
+	if (const TSharedPtr<FJsonValue>* JsonValue = Object->Values.Find(FDMRenderTargetTextRenderer::WrappingPolicyName.GetPlainNameString()))
+	{
+		FDMJsonUtils::Deserialize(*JsonValue, WrappingPolicy);
+	}
+
+	if (const TSharedPtr<FJsonValue>* JsonValue = Object->Values.Find(FDMRenderTargetTextRenderer::JustifyName.GetPlainNameString()))
+	{
+		ETextJustify::Type EnumValue = ETextJustify::Left;
+		FDMJsonUtils::Deserialize(*JsonValue, EnumValue);
+		Justify = EnumValue;
+	}
+
+	if (const TSharedPtr<FJsonValue>* JsonValue = Object->Values.Find(FDMRenderTargetTextRenderer::TransformPolicyName.GetPlainNameString()))
+	{
+		FDMJsonUtils::Deserialize(*JsonValue, TransformPolicy);
+	}
+
+	if (const TSharedPtr<FJsonValue>* JsonValue = Object->Values.Find(FDMRenderTargetTextRenderer::FlowDirectionName.GetPlainNameString()))
+	{
+		FDMJsonUtils::Deserialize(*JsonValue, FlowDirection);
+	}
+
+	if (const TSharedPtr<FJsonValue>* JsonValue = Object->Values.Find(FDMRenderTargetTextRenderer::ShapingMethodName.GetPlainNameString()))
+	{
+		FDMJsonUtils::Deserialize(*JsonValue, ShapingMethod);
+	}
+
+	if (const TSharedPtr<FJsonValue>* JsonValue = Object->Values.Find(FDMRenderTargetTextRenderer::LineHeightName.GetPlainNameString()))
+	{
+		FDMJsonUtils::Deserialize(*JsonValue, LineHeight);
+	}
+
+	if (const TSharedPtr<FJsonValue>* JsonValue = Object->Values.Find(FDMRenderTargetTextRenderer::PaddingLeftName.GetPlainNameString()))
+	{
+		FDMJsonUtils::Deserialize(*JsonValue, PaddingLeft);
+	}
+
+	if (const TSharedPtr<FJsonValue>* JsonValue = Object->Values.Find(FDMRenderTargetTextRenderer::PaddingRightName.GetPlainNameString()))
+	{
+		FDMJsonUtils::Deserialize(*JsonValue, PaddingRight);
+	}
+
+	if (const TSharedPtr<FJsonValue>* JsonValue = Object->Values.Find(FDMRenderTargetTextRenderer::PaddingTopName.GetPlainNameString()))
+	{
+		FDMJsonUtils::Deserialize(*JsonValue, PaddingTop);
+	}
+
+	if (const TSharedPtr<FJsonValue>* JsonValue = Object->Values.Find(FDMRenderTargetTextRenderer::PaddingBottomName.GetPlainNameString()))
+	{
+		FDMJsonUtils::Deserialize(*JsonValue, PaddingBottom);
+	}
+
+	if (const TSharedPtr<FJsonValue>* JsonValue = Object->Values.Find(FDMRenderTargetTextRenderer::OverrideRenderTargetSizeName.GetPlainNameString()))
+	{
+		FDMJsonUtils::Deserialize(*JsonValue, bOverrideRenderTargetSize);
+	}
+
+	if (const TSharedPtr<FJsonValue>* JsonValue = Object->Values.Find(FDMRenderTargetTextRenderer::StrikeBrushName.GetPlainNameString()))
+	{
+		if (!(*JsonValue)->IsNull())
+		{
+			FDMJsonUtils::Deserialize<FSlateBrush>(*JsonValue, *StrikeBrush.GetMutablePtr<FSlateBrush>());
+		}
+	}
+
 	return true;
 }
 
@@ -78,55 +243,39 @@ void UDMRenderTargetTextRenderer::PostEditChangeProperty(FPropertyChangedEvent&
 
 	const FName PropertyName = InPropertyChangedEvent.GetMemberPropertyName();
 
-	if (PropertyName == FDMRenderTargetTextRenderer::TextName)
-	{
-		UpdateTextLines();
-	}
-	else if (PropertyName == FDMRenderTargetTextRenderer::FontName
-		|| PropertyName == FDMRenderTargetTextRenderer::KerningName
-		|| PropertyName == FDMRenderTargetTextRenderer::LineHeightName
-		|| PropertyName == FDMRenderTargetTextRenderer::PaddingLeftName
-		|| PropertyName == FDMRenderTargetTextRenderer::PaddingRightName
-		|| PropertyName == FDMRenderTargetTextRenderer::PaddingTopName
-		|| PropertyName == FDMRenderTargetTextRenderer::PaddingBottomName
-		|| PropertyName == FDMRenderTargetTextRenderer::TextScaleName
-		|| PropertyName == FDMRenderTargetTextRenderer::TextureSizeOverrideName)
-	{
-		TextScale.X = FMath::Clamp(TextScale.X, 1, UE::DynamicMaterial::Private::MaxTexScale.X);
-		TextScale.Y = FMath::Clamp(TextScale.Y, 1, UE::DynamicMaterial::Private::MaxTexScale.Y);
-		bRecalculateTextSize = true;
-		AsyncUpdateRenderTarget();
-	}
-	else if (PropertyName == FDMRenderTargetTextRenderer::bGlowName
-		|| PropertyName == FDMRenderTargetTextRenderer::bShadowName
-		|| PropertyName == FDMRenderTargetTextRenderer::bOutlineName)
+	if (PropertyName == FDMRenderTargetTextRenderer::HasHighlightName
+		|| PropertyName == FDMRenderTargetTextRenderer::HasShadowName
+		|| PropertyName == FDMRenderTargetTextRenderer::AutoWrapTextName)
 	{
 		// Cause details panel refresh.
 		Update(EDMUpdateType::Structure);
-		AsyncUpdateRenderTarget();
 	}
-	else if (FDMRenderTargetTextRenderer::PropertyNames.Contains(PropertyName))
+
+	if (FDMRenderTargetTextRenderer::PropertyNames.Contains(PropertyName))
 	{
+		Widget.Reset();
+		bRecalculateTextSize = true;
 		AsyncUpdateRenderTarget();
 	}
 }
 #endif
 
-UFont* UDMRenderTargetTextRenderer::GetFont() const
+const FSlateFontInfo& UDMRenderTargetTextRenderer::GetFontInfo() const
 {
-	return Font;
+	return FontInfo;
 }
 
-void UDMRenderTargetTextRenderer::SetFont(UFont* InFont)
+void UDMRenderTargetTextRenderer::SetFontInfo(const FSlateFontInfo& InFontInfo)
 {
-	if (Font == InFont)
+	if (FontInfo == InFontInfo)
 	{
 		return;
 	}
 
-	Font = InFont;
+	FontInfo = InFontInfo;
 
-	UpdateTextureSize();
+	bRecalculateTextSize = true;
+	AsyncUpdateRenderTarget();
 }
 
 const FText& UDMRenderTargetTextRenderer::GetText() const
@@ -143,8 +292,9 @@ void UDMRenderTargetTextRenderer::SetText(const FText& InText)
 
 	Text = InText;
 
-	UpdateTextLines();
-	UpdateTextureSize();
+	bRecalculateTextSize = true;
+	AsyncUpdateRenderTarget();
+
 }
 
 const FLinearColor& UDMRenderTargetTextRenderer::GetTextColor() const
@@ -182,53 +332,56 @@ void UDMRenderTargetTextRenderer::SetBackgroundColor(const FLinearColor& InBackg
 	}
 }
 
-ETextJustify::Type UDMRenderTargetTextRenderer::GetJustify() const
+bool UDMRenderTargetTextRenderer::GetHasHighlight() const
 {
-	return Justify;
+	return bHasHighlight;
 }
 
-void UDMRenderTargetTextRenderer::SetJustify(ETextJustify::Type InJustify)
+void UDMRenderTargetTextRenderer::SetHasHighlight(bool bInHasHighlight)
 {
-	if (Justify == InJustify)
+	if (bHasHighlight == bInHasHighlight)
 	{
 		return;
 	}
 
-	Justify = InJustify;
+	bHasHighlight = bInHasHighlight;
 
+	bRecalculateTextSize = true;
 	AsyncUpdateRenderTarget();
-}
 
-float UDMRenderTargetTextRenderer::GetKerning() const
-{
-	return Kerning;
 }
 
-void UDMRenderTargetTextRenderer::SetKerning(float InKerning)
+const FLinearColor& UDMRenderTargetTextRenderer::GetHighlightColor() const
 {
-	if (Kerning == InKerning)
+	if (UDMMaterialValueRenderTarget* RenderTargetValue = GetRenderTargetValue())
 	{
-		return;
+		return RenderTargetValue->GetClearColor();
 	}
 
-	Kerning = InKerning;
+	return FLinearColor::Black;
+}
 
-	UpdateTextureSize();
+void UDMRenderTargetTextRenderer::SetHighlightColor(const FLinearColor& InHighlightColor)
+{
+	if (UDMMaterialValueRenderTarget* RenderTargetValue = GetRenderTargetValue())
+	{
+		RenderTargetValue->SetClearColor(InHighlightColor);
+	}
 }
 
-const FLinearColor& UDMRenderTargetTextRenderer::GetOutlineColor() const
+ETextJustify::Type UDMRenderTargetTextRenderer::GetJustify() const
 {
-	return OutlineColor;
+	return Justify;
 }
 
-void UDMRenderTargetTextRenderer::SetOutlineColor(const FLinearColor& InOutlineColor)
+void UDMRenderTargetTextRenderer::SetJustify(ETextJustify::Type InJustify)
 {
-	if (OutlineColor == InOutlineColor)
+	if (Justify == InJustify)
 	{
 		return;
 	}
 
-	OutlineColor = InOutlineColor;
+	Justify = InJustify;
 
 	AsyncUpdateRenderTarget();
 }
@@ -247,7 +400,9 @@ void UDMRenderTargetTextRenderer::SetLineHeight(float InLineHeight)
 
 	LineHeight = InLineHeight;
 
-	UpdateTextureSize();
+	bRecalculateTextSize = true;
+	AsyncUpdateRenderTarget();
+
 }
 
 float UDMRenderTargetTextRenderer::GetPaddingLeft() const
@@ -264,7 +419,8 @@ void UDMRenderTargetTextRenderer::SetPaddingLeft(float InPaddingLeft)
 
 	PaddingLeft = InPaddingLeft;
 
-	UpdateTextureSize();
+	bRecalculateTextSize = true;
+	AsyncUpdateRenderTarget();
 }
 
 float UDMRenderTargetTextRenderer::GetPaddingRight() const
@@ -281,7 +437,8 @@ void UDMRenderTargetTextRenderer::SetPaddingRight(float InPaddingRight)
 
 	PaddingRight = InPaddingRight;
 
-	UpdateTextureSize();
+	bRecalculateTextSize = true;
+	AsyncUpdateRenderTarget();
 }
 
 float UDMRenderTargetTextRenderer::GetPaddingTop() const
@@ -298,7 +455,8 @@ void UDMRenderTargetTextRenderer::SetPaddingTop(float InPaddingTop)
 
 	PaddingTop = InPaddingTop;
 
-	UpdateTextureSize();
+	bRecalculateTextSize = true;
+	AsyncUpdateRenderTarget();
 }
 
 float UDMRenderTargetTextRenderer::GetPaddingBottom() const
@@ -315,78 +473,43 @@ void UDMRenderTargetTextRenderer::SetPaddingBottom(float InPaddingBottom)
 
 	PaddingBottom = InPaddingBottom;
 
-	UpdateTextureSize();
-}
-
-const FIntPoint& UDMRenderTargetTextRenderer::GetTextScale() const
-{
-	return TextScale;
-}
-
-void UDMRenderTargetTextRenderer::SetTextScale(const FIntPoint& InTextScale)
-{
-	FIntPoint NewTextScale = InTextScale;
-	NewTextScale.X = FMath::Clamp(InTextScale.X, 1, UE::DynamicMaterial::Private::MaxTexScale.X);
-	NewTextScale.Y = FMath::Clamp(InTextScale.Y, 1, UE::DynamicMaterial::Private::MaxTexScale.Y);
-
-	if (TextScale == NewTextScale)
-	{
-		return;
-	}
-
-	TextScale = NewTextScale;
-
-	UpdateTextureSize();
-}
-
-const FIntPoint& UDMRenderTargetTextRenderer::GetTextureSizeOverride() const
-{
-	return TextureSizeOverride;
-}
-
-void UDMRenderTargetTextRenderer::SetTextureSizeOverride(const FIntPoint& InTextureSizeOverride)
-{
-	if (TextureSizeOverride == InTextureSizeOverride)
-	{
-		return;
-	}
-
-	TextureSizeOverride = InTextureSizeOverride;
-
-	UpdateTextureSize();
+	bRecalculateTextSize = true;
+	AsyncUpdateRenderTarget();
 }
 
-bool UDMRenderTargetTextRenderer::GetHasOutline() const
+bool UDMRenderTargetTextRenderer::IsOverridingRenderTargetSize() const
 {
-	return bOutline;
+	return bOverrideRenderTargetSize;
 }
 
-void UDMRenderTargetTextRenderer::SetHasOutline(bool bInHasOutline)
+void UDMRenderTargetTextRenderer::SetOverrideRenderTargetSize(bool bInOverride)
 {
-	if (bOutline == bInHasOutline)
+	if (bOverrideRenderTargetSize == bInOverride)
 	{
 		return;
 	}
 
-	bOutline = bInHasOutline;
+	bOverrideRenderTargetSize = bInOverride;
 
+	bRecalculateTextSize = true;
 	AsyncUpdateRenderTarget();
 }
 
 bool UDMRenderTargetTextRenderer::GetHasShadow() const
 {
-	return bShadow;
+	return bHasShadow;
 }
 
 void UDMRenderTargetTextRenderer::SetHasShadow(bool bInHasShadow)
 {
-	if (bShadow == bInHasShadow)
+	if (bHasShadow == bInHasShadow)
 	{
 		return;
 	}
 
-	bShadow = bInHasShadow;
+	bHasShadow = bInHasShadow;
 
+	bRecalculateTextSize = true;
 	AsyncUpdateRenderTarget();
 }
 
@@ -421,134 +544,131 @@ void UDMRenderTargetTextRenderer::SetShadowOffset(const FVector2D& InShadowOffse
 
 	ShadowOffset = InShadowOffset;
 
+	bRecalculateTextSize = true;
 	AsyncUpdateRenderTarget();
 }
 
-bool UDMRenderTargetTextRenderer::GetHasGlow() const
+bool UDMRenderTargetTextRenderer::GetAutoWrapText() const
 {
-	return bGlow;
+	return bAutoWrapText;
 }
 
-void UDMRenderTargetTextRenderer::SetHasGlow(bool bInHasGlow)
+void UDMRenderTargetTextRenderer::SetAutoWrapText(bool bInAutoWrap)
 {
-	if (bGlow == bInHasGlow)
+	if (bAutoWrapText == bInAutoWrap)
 	{
 		return;
 	}
 
-	bGlow = bInHasGlow;
+	bAutoWrapText = bInAutoWrap;
 
+	bRecalculateTextSize = true;
 	AsyncUpdateRenderTarget();
 }
 
-const FLinearColor& UDMRenderTargetTextRenderer::GetGlowColor() const
+float UDMRenderTargetTextRenderer::GetWrapTextAt() const
 {
-	return GlowColor;
+	return WrapTextAt;
 }
 
-void UDMRenderTargetTextRenderer::SetGlowColor(const FLinearColor& InGlowColor)
+void UDMRenderTargetTextRenderer::SetWrapTextAt(float InWrapAt)
 {
-	if (GlowColor == InGlowColor)
+	if (WrapTextAt == InWrapAt)
 	{
 		return;
 	}
 
-	GlowColor = InGlowColor;
+	WrapTextAt = InWrapAt;
 
+	bRecalculateTextSize = true;
 	AsyncUpdateRenderTarget();
 }
 
-const FVector2D& UDMRenderTargetTextRenderer::GetGlowInnerRadius() const
+ETextWrappingPolicy UDMRenderTargetTextRenderer::GetWrappingPolicy() const
 {
-	return GlowInnerRadius;
+	return WrappingPolicy;
 }
 
-void UDMRenderTargetTextRenderer::SetGlowInnerRadius(const FVector2D& InGlowInnerRadius)
+void UDMRenderTargetTextRenderer::SetWrappingPolicy(ETextWrappingPolicy InWrappingPolicy)
 {
-	if (GlowInnerRadius == InGlowInnerRadius)
+	if (WrappingPolicy == InWrappingPolicy)
 	{
 		return;
 	}
 
-	GlowInnerRadius = InGlowInnerRadius;
+	WrappingPolicy = InWrappingPolicy;
 
+	bRecalculateTextSize = true;
 	AsyncUpdateRenderTarget();
 }
 
-const FVector2D& UDMRenderTargetTextRenderer::GetGlowOuterRadius() const
+ETextTransformPolicy UDMRenderTargetTextRenderer::GetTransformPolicy() const
 {
-	return GlowOuterRadius;
+	return TransformPolicy;
 }
 
-void UDMRenderTargetTextRenderer::SetGlowOuterRadius(const FVector2D& InGlowOuterRadius)
+void UDMRenderTargetTextRenderer::SetTransformPolicy(ETextTransformPolicy InTransformPolicy)
 {
-	if (GlowOuterRadius == InGlowOuterRadius)
+	if (TransformPolicy == InTransformPolicy)
 	{
 		return;
 	}
 
-	GlowOuterRadius = InGlowOuterRadius;
+	TransformPolicy = InTransformPolicy;
 
+	bRecalculateTextSize = true;
 	AsyncUpdateRenderTarget();
 }
 
-void UDMRenderTargetTextRenderer::CalculateLineSizes()
+ETextFlowDirection UDMRenderTargetTextRenderer::GetFlowDirection() const
 {
-	FTextSizingParameters Params(0, 0, 0, 0, Font);
-	Params.Scaling = FVector2D(1, 1);
+	return FlowDirection;
+}
 
-	for (FDMTextLine& Line : Lines)
+void UDMRenderTargetTextRenderer::SetFlowDirection(ETextFlowDirection InFlowDirection)
+{
+	if (FlowDirection == InFlowDirection)
 	{
-		Params.DrawXL = 0;
-		Params.DrawYL = 0;
-
-		UCanvas::CanvasStringSize(Params, *Line.Line);
+		return;
+	}
 
-		int32 RequiredWidth = Params.DrawXL;
-		RequiredWidth += Line.Line.Len() > 0 ? ((Line.Line.Len() - 1) * Kerning) : 0;
+	FlowDirection = InFlowDirection;
 
-		Line.Width = RequiredWidth;
-	}
+	AsyncUpdateRenderTarget();
 }
 
-FIntPoint UDMRenderTargetTextRenderer::GetRequiredTextureSize() const
+ETextShapingMethod UDMRenderTargetTextRenderer::GetShapingMethod() const
 {
-	if (TextureSizeOverride.X > 0 && TextureSizeOverride.Y > 0)
-	{
-		return TextureSizeOverride;
-	}
+	return ShapingMethod;
+}
 
-	if (!IsValid(Font) || Lines.IsEmpty())
+void UDMRenderTargetTextRenderer::SetShapingMethod(ETextShapingMethod InShapingMethod)
+{
+	if (ShapingMethod == InShapingMethod)
 	{
-		return UE::DynamicMaterial::Private::MinimumTextTextureSize;
+		return;
 	}
 
-	const float MaxCharHeight = Font->GetMaxCharHeight();
+	ShapingMethod = InShapingMethod;
 
-	if (MaxCharHeight <= 0)
-	{
-		return UE::DynamicMaterial::Private::MinimumTextTextureSize;
-	}
+	AsyncUpdateRenderTarget();
+}
 
-	FVector2f Size = FVector2f::ZeroVector;
+const TInstancedStruct<FSlateBrush>& UDMRenderTargetTextRenderer::GetStrikeBrush() const
+{
+	return StrikeBrush;
+}
 
-	for (const FDMTextLine& Line : Lines)
+void UDMRenderTargetTextRenderer::SetStrikeBrush(const TInstancedStruct<FSlateBrush>& InStrikeBrush)
+{
+	if (StrikeBrush == InStrikeBrush)
 	{
-		Size.X = FMath::Max(Size.X, Line.Width);
-		Size.Y += MaxCharHeight * LineHeight;
+		return;
 	}
 
-	Size.X += PaddingLeft + PaddingRight;
-	Size.Y += PaddingTop + PaddingBottom;
+	StrikeBrush = InStrikeBrush;
 
-	Size.X = FMath::Max(Size.X, UE::DynamicMaterial::Private::MinimumTextTextureSize.X);
-	Size.Y = FMath::Max(Size.Y, UE::DynamicMaterial::Private::MinimumTextTextureSize.Y);
-
-	FIntPoint RequiredSize = FIntPoint::ZeroValue;
-	RequiredSize.X = FMath::CeilToInt(Size.X * TextScale.X);
-	RequiredSize.Y = FMath::CeilToInt(Size.Y * TextScale.Y);
-
-	return RequiredSize;
+	AsyncUpdateRenderTarget();
 }
 
 void UDMRenderTargetTextRenderer::UpdateTextLines()
@@ -567,78 +687,60 @@ void UDMRenderTargetTextRenderer::UpdateTextLines()
 			NewLine.LeftChopInline(1);
 		}
 
-		Lines.Add({NewLine, -1.f});
-	}
+		TSharedRef<STextBlock> TextBlock = CreateTextWidget(FText::FromString(NewLine));
+
+		Lines.Emplace(NewLine, TextBlock->ComputeDesiredSize(1.f).X, TextBlock);
+	}	
 
 	bRecalculateTextSize = true;
 	AsyncUpdateRenderTarget();
 }
 
-FCanvasTextItem UDMRenderTargetTextRenderer::CreateTextItem(const FVector2D& InPosition, const FText& InText) const
+TSharedRef<STextBlock> UDMRenderTargetTextRenderer::CreateTextWidget(const FText& InText) const
 {
-	FCanvasTextItem TextItem(InPosition, InText, Font, TextColor);
-
-	TextItem.BlendMode = SE_BLEND_Translucent; // Must be translucent for fonts!
-	TextItem.HorizSpacingAdjust = Kerning;
-
-	if (bOutline)
-	{
-		TextItem.bOutlined = true;
-		TextItem.OutlineColor = OutlineColor;
-	}
-	else
-	{
-		TextItem.bOutlined = false;
-	}
-
-	if (bShadow)
-	{
-		TextItem.EnableShadow(ShadowColor, ShadowOffset);
-	}
-	else
-	{
-		TextItem.DisableShadow();
-	}
-
-	if (bGlow)
-	{
-		TextItem.FontRenderInfo.GlowInfo.bEnableGlow = true;
-		TextItem.FontRenderInfo.GlowInfo.GlowColor = GlowColor;
-		TextItem.FontRenderInfo.GlowInfo.GlowInnerRadius = GlowInnerRadius;
-		TextItem.FontRenderInfo.GlowInfo.GlowOuterRadius = GlowOuterRadius;
-	}
-	else
-	{
-		TextItem.FontRenderInfo.GlowInfo.bEnableGlow = false;
-	}
-
-	return TextItem;
+	return SNew(STextBlock)
+		.Font(FontInfo)
+		.Text(Text)
+		.LineHeightPercentage(LineHeight)
+		.ColorAndOpacity(TextColor)
+		.HighlightColor(bHasHighlight ? HighlightColor : TAttribute<FLinearColor>())
+		.ShadowColorAndOpacity(bHasShadow ? ShadowColor : TAttribute<FLinearColor>())
+		.ShadowOffset(bHasShadow ? ShadowOffset : TAttribute<FVector2D>())
+		.Justification(Justify)
+		.TransformPolicy(TransformPolicy)
+		.TextFlowDirection(FlowDirection)
+		.TextShapingMethod(ShapingMethod)
+		.StrikeBrush(StrikeBrush.GetPtr<FSlateBrush>())
+		.Margin(FMargin(PaddingLeft, PaddingTop, PaddingRight, PaddingBottom))
+		.AutoWrapText(bAutoWrapText)
+		.WrappingPolicy(bAutoWrapText ? WrappingPolicy : TAttribute<ETextWrappingPolicy>())
+		.WrapTextAt(bAutoWrapText ? WrapTextAt : TAttribute<float>());
 }
 
-void UDMRenderTargetTextRenderer::UpdateTextureSize()
+void UDMRenderTargetTextRenderer::CreateWidgetInstance()
 {
-	UDMMaterialValueRenderTarget* RenderTargetValue = GetRenderTargetValue();
+	TSharedRef<SVerticalBox> NewWidget = SNew(SVerticalBox);
 
-	if (!RenderTargetValue)
+	for (FDMTextLine& Line : Lines)
 	{
-		return;
-	}
+		if (!Line.Widget.IsValid())
+		{
+			Line.Widget = CreateTextWidget(FText::FromString(Line.Line));
+			Line.Width = Line.Widget->ComputeDesiredSize(1.f).X;
+		}
 
-	const FIntPoint RequiredSize = GetRequiredTextureSize();
+		NewWidget->AddSlot()
+			.AutoHeight()
+			[
+				Line.Widget.ToSharedRef()
+			];
+	}
 
-	RenderTargetValue->SetTextureSize(RequiredSize);
-	RenderTargetValue->FlushCreateRenderTarget();
+	Widget = NewWidget;
 }
 
-void UDMRenderTargetTextRenderer::UpdateRenderTarget_Internal()
+void UDMRenderTargetTextRenderer::SetCustomTextureSize()
 {
-	Super::UpdateRenderTarget();
-
-	if (!IsValid(Font) || Lines.IsEmpty())
-	{
-		return;
-	}
-
 	UDMMaterialValueRenderTarget* RenderTargetValue = GetRenderTargetValue();
 
 	if (!RenderTargetValue)
@@ -646,90 +748,34 @@ void UDMRenderTargetTextRenderer::UpdateRenderTarget_Internal()
 		return;
 	}
 
-	if (bRecalculateTextSize)
-	{
-		CalculateLineSizes();
-		UpdateTextureSize();
-	}
-
-	UTextureRenderTarget2D* RenderTarget = RenderTargetValue->GetRenderTarget();
-
-	if (!RenderTarget)
+	if (!Widget.IsValid())
 	{
 		return;
 	}
 
-	RenderTarget->UpdateResourceImmediate(true);
+	Widget->SlatePrepass(1.f);
+	const FVector2D Size = Widget->GetDesiredSize();
 
-	FTextureRenderTargetResource* RenderTargetResource = RenderTarget->GameThread_GetRenderTargetResource();
+	RenderTargetValue->SetTextureSize(FIntPoint(FMath::RoundToInt(Size.X), FMath::RoundToInt(Size.Y)));
+	RenderTargetValue->FlushCreateRenderTarget();
+}
 
-	if (!RenderTargetResource)
+void UDMRenderTargetTextRenderer::UpdateRenderTarget_Internal()
+{
+	if (Text.IsEmpty())
 	{
 		return;
 	}
 
-	static const FGameTime GameTime = FGameTime::CreateUndilated(0, 0);
-
-	FCanvas RenderCanvas(
-		RenderTargetResource,
-		nullptr,
-		GameTime,
-		GEngine->GetDefaultWorldFeatureLevel(),
-		FCanvas::ECanvasDrawMode::CDM_ImmediateDrawing
-	);
-
-	const FVector2D RenderTargetSize = FVector2D(RenderTarget->SizeX, RenderTarget->SizeY);
-
-	FCanvasTileItem TileItem = FCanvasTileItem(FVector2D::ZeroVector, RenderTargetSize, RenderTarget->ClearColor);
-	RenderCanvas.DrawItem(TileItem);
-
-	const float MaxCharHeight = Font->GetMaxCharHeight();
-
-	FCanvasTextItem TextItem = CreateTextItem(FVector2D::ZeroVector, FText::GetEmpty());
-
-	const float LeftOffset = PaddingLeft * static_cast<float>(TextScale.X);
-	const float RightOffset = PaddingRight * static_cast<float>(TextScale.X);
-
-	TextItem.Position.X = LeftOffset;
-	TextItem.Position.Y = PaddingTop * static_cast<float>(TextScale.Y);
-	TextItem.Scale.X = static_cast<float>(TextScale.X);
-	TextItem.Scale.Y = static_cast<float>(TextScale.Y);
-
-	if (!FMath::IsNearlyEqual(LineHeight, 1.f))
-	{
-		TextItem.Position.Y += MaxCharHeight * (LineHeight - 1.f) * 0.5f * static_cast<float>(TextScale.Y);
-	}
-
-	const float ActualLineHeight = MaxCharHeight * LineHeight * static_cast<float>(TextScale.Y);
+	UpdateTextLines();
+	CreateWidgetInstance();
 
-	for (const FDMTextLine& Line : Lines)
+	if (bOverrideRenderTargetSize && bRecalculateTextSize)
 	{
-		switch (Justify)
-		{
-			default:
-			case ETextJustify::Left:
-			case ETextJustify::InvariantLeft:
-				TextItem.Position.X = LeftOffset;
-				break;
-
-			case ETextJustify::Center:
-				TextItem.Position.X = (RenderTarget->SizeX - (Line.Width * TextItem.Scale.X)) * 0.5f;
-				break;
-
-			case ETextJustify::Right:
-			case ETextJustify::InvariantRight:
-				TextItem.Position.X = RenderTarget->SizeX - (Line.Width * TextItem.Scale.X) - RightOffset;
-				break;
-		}
-
-		TextItem.Text = FText::FromString(Line.Line);
-		RenderCanvas.DrawItem(TextItem);
-		TextItem.Position.Y += ActualLineHeight;
+		SetCustomTextureSize();
 	}
 
-	RenderCanvas.Flush_GameThread();
-
-	RenderTarget->UpdateResourceImmediate(false);
+	Super::UpdateRenderTarget_Internal();
 }
 
 #undef LOCTEXT_NAMESPACE
diff --git a/Engine/Plugins/Experimental/DynamicMaterial/Source/DynamicMaterial/Private/Components/RenderTargetRenderers/DMRenderTargetUMGWidgetRenderer.cpp b/Engine/Plugins/Experimental/DynamicMaterial/Source/DynamicMaterial/Private/Components/RenderTargetRenderers/DMRenderTargetUMGWidgetRenderer.cpp
new file mode 100644
--- /dev/null
+++ b/Engine/Plugins/Experimental/DynamicMaterial/Source/DynamicMaterial/Private/Components/RenderTargetRenderers/DMRenderTargetUMGWidgetRenderer.cpp
@@ -0,0 +1,76 @@
+// Copyright Epic Games, Inc. All Rights Reserved.
+
+#include "Components/RenderTargetRenderers/DMRenderTargetUMGWidgetRenderer.h"
+#include "Components/MaterialValues/DMMaterialValueRenderTarget.h"
+#include "Components/Widget.h"
+#include "Dom/JsonValue.h"
+#include "Engine/Engine.h"
+
+#define LOCTEXT_NAMESPACE "DMRenderTargetUWidgetRenderer"
+
+UDMRenderTargetUMGWidgetRenderer::UDMRenderTargetUMGWidgetRenderer()
+{
+#if WITH_EDITOR
+	EditableProperties.Add(GET_MEMBER_NAME_CHECKED(UDMRenderTargetUMGWidgetRenderer, WidgetClass));
+#endif
+}
+
+void UDMRenderTargetUMGWidgetRenderer::SetWidgetClass(TSubclassOf<UWidget> InWidgetClass)
+{
+	if (InWidgetClass == WidgetClass)
+	{
+		return;
+	}
+
+	WidgetClass = InWidgetClass;
+	CreateWidgetInstance();
+	AsyncUpdateRenderTarget();
+}
+
+#if WITH_EDITOR
+TSharedPtr<FJsonValue> UDMRenderTargetUMGWidgetRenderer::JsonSerialize() const
+{
+	return FDMJsonUtils::Serialize(WidgetClass.Get());
+}
+
+bool UDMRenderTargetUMGWidgetRenderer::JsonDeserialize(const TSharedPtr<FJsonValue>& InJsonValue)
+{
+	TSubclassOf<UWidget> WidgetClassJson;
+
+	if (FDMJsonUtils::Deserialize(InJsonValue, WidgetClassJson))
+	{
+		SetWidgetClass(WidgetClassJson);
+		return true;
+	}
+
+	return false;
+}
+
+void UDMRenderTargetUMGWidgetRenderer::PostEditChangeProperty(FPropertyChangedEvent& InPropertyChangedEvent)
+{
+	Super::PostEditChangeProperty(InPropertyChangedEvent);
+
+	const FName PropertyName = InPropertyChangedEvent.GetMemberPropertyName();
+
+	if (PropertyName == GET_MEMBER_NAME_CHECKED(UDMRenderTargetUMGWidgetRenderer, WidgetClass))
+	{
+		CreateWidgetInstance();
+		AsyncUpdateRenderTarget();
+	}
+}
+#endif
+
+void UDMRenderTargetUMGWidgetRenderer::CreateWidgetInstance()
+{
+	UClass* WidgetClassLocal = WidgetClass.Get();
+
+	if (!WidgetClassLocal)
+	{
+		return;
+	}
+
+	WidgetInstance = NewObject<UWidget>(this, WidgetClassLocal, NAME_None, RF_Transient);
+	Widget = WidgetInstance->TakeWidget();
+}
+
+#undef LOCTEXT_NAMESPACE
diff --git a/Engine/Plugins/Experimental/DynamicMaterial/Source/DynamicMaterial/Private/Components/RenderTargetRenderers/DMRenderTargetWidgetRenderer.cpp b/Engine/Plugins/Experimental/DynamicMaterial/Source/DynamicMaterial/Private/Components/RenderTargetRenderers/DMRenderTargetWidgetRenderer.cpp
deleted file mode 100644
diff --git a/Engine/Plugins/Experimental/DynamicMaterial/Source/DynamicMaterial/Private/Components/RenderTargetRenderers/DMRenderTargetWidgetRendererBase.cpp b/Engine/Plugins/Experimental/DynamicMaterial/Source/DynamicMaterial/Private/Components/RenderTargetRenderers/DMRenderTargetWidgetRendererBase.cpp
new file mode 100644
--- /dev/null
+++ b/Engine/Plugins/Experimental/DynamicMaterial/Source/DynamicMaterial/Private/Components/RenderTargetRenderers/DMRenderTargetWidgetRendererBase.cpp
@@ -0,0 +1,57 @@
+// Copyright Epic Games, Inc. All Rights Reserved.
+
+#include "Components/RenderTargetRenderers/DMRenderTargetWidgetRendererBase.h"
+#include "Components/MaterialValues/DMMaterialValueRenderTarget.h"
+#include "Engine/Engine.h"
+#include "Engine/TextureRenderTarget2D.h"
+#include "Slate/WidgetRenderer.h"
+#include "Widgets/SWidget.h"
+
+#define LOCTEXT_NAMESPACE "DMRenderTargetUMGWidgetRenderer"
+
+UDMRenderTargetWidgetRendererBase::UDMRenderTargetWidgetRendererBase()
+{
+	WidgetRenderer = MakeShared<FWidgetRenderer>(/* Gamma correction */ false);
+	WidgetRenderer->SetIsPrepassNeeded(true);
+	WidgetRenderer->SetShouldClearTarget(true);
+}
+
+void UDMRenderTargetWidgetRendererBase::UpdateRenderTarget_Internal()
+{
+	Super::UpdateRenderTarget_Internal();
+
+	UDMMaterialValueRenderTarget* RenderTargetValue = GetRenderTargetValue();
+
+	if (!RenderTargetValue)
+	{
+		return;
+	}
+
+	RenderTargetValue->EnsureRenderTarget(/* Async */ false);
+
+	UTextureRenderTarget2D* RenderTarget = RenderTargetValue->GetRenderTarget();
+
+	if (!RenderTarget)
+	{
+		return;
+	}
+
+	if (!Widget.IsValid())
+	{
+		CreateWidgetInstance();
+
+		if (!Widget.IsValid())
+		{
+			return;
+		}
+	}
+
+	WidgetRenderer->DrawWidget(
+		RenderTarget,
+		Widget.ToSharedRef(),
+		{(double)RenderTarget->SizeX, (double)RenderTarget->SizeX},
+		/* Delta Time */ 0.f
+	);
+}
+
+#undef LOCTEXT_NAMESPACE
diff --git a/Engine/Plugins/Experimental/DynamicMaterial/Source/DynamicMaterial/Private/Utils/DMJsonUtils.cpp b/Engine/Plugins/Experimental/DynamicMaterial/Source/DynamicMaterial/Private/Utils/DMJsonUtils.cpp
--- a/Engine/Plugins/Experimental/DynamicMaterial/Source/DynamicMaterial/Private/Utils/DMJsonUtils.cpp
+++ b/Engine/Plugins/Experimental/DynamicMaterial/Source/DynamicMaterial/Private/Utils/DMJsonUtils.cpp
@@ -60,7 +60,7 @@ TSharedPtr<FJsonValue> FDMJsonUtils::Serialize(const UScriptStruct* InScriptStru
 	if (!IsValid(InScriptStruct))
 	{
 		UE_LOG(LogDynamicMaterial, Error, TEXT("Invalid script struct."));
-		return nullptr;
+		return MakeShared<FJsonValueNull>();
 	}
 
 	TSharedRef<FJsonObject> JsonStruct = MakeShared<FJsonObject>();
@@ -71,7 +71,7 @@ TSharedPtr<FJsonValue> FDMJsonUtils::Serialize(const UScriptStruct* InScriptStru
 	}
 
 	UE_LOG(LogDynamicMaterial, Error,TEXT("Failed to convert struct to json. [%s]"), *InScriptStruct->GetName());
-	return nullptr;
+	return MakeShared<FJsonValueNull>();
 }
 
 TSharedPtr<FJsonValue> FDMJsonUtils::Serialize(const UObject* InObject)
@@ -116,7 +116,7 @@ TSharedPtr<FJsonValue> FDMJsonUtils::Serialize(const UObject* InObject)
 	UE_LOG(LogDynamicMaterial, Error, TEXT("Failed to convert object to json. [%s] [%s]"), 
 		*InObject->GetClass()->GetName(), *InObject->GetPathName());
 
-	return nullptr;
+	return MakeShared<FJsonValueNull>();
 }
 
 TSharedPtr<FJsonValue> FDMJsonUtils::Serialize(const FObjectPtr& InObject)
diff --git a/Engine/Plugins/Experimental/DynamicMaterial/Source/DynamicMaterial/Public/Components/DMRenderTargetRenderer.h b/Engine/Plugins/Experimental/DynamicMaterial/Source/DynamicMaterial/Public/Components/DMRenderTargetRenderer.h
--- a/Engine/Plugins/Experimental/DynamicMaterial/Source/DynamicMaterial/Public/Components/DMRenderTargetRenderer.h
+++ b/Engine/Plugins/Experimental/DynamicMaterial/Source/DynamicMaterial/Public/Components/DMRenderTargetRenderer.h
@@ -48,6 +48,10 @@ public:
 	virtual bool JsonDeserialize(const TSharedPtr<FJsonValue>& InJsonValue) override;
 	//~ End IDMJsonSerializable
 
+	//~ Begin UDMMaterialComponent
+	virtual void Update(EDMUpdateType InUpdateType = EDMUpdateType::Value) override;
+	//~ End UDMMaterialComponent
+
 	//~ Begin UObject
 	virtual void PostLoad() override;
 	//~ End UObject
diff --git a/Engine/Plugins/Experimental/DynamicMaterial/Source/DynamicMaterial/Public/Components/RenderTargetRenderers/DMRenderTargetTextRenderer.h b/Engine/Plugins/Experimental/DynamicMaterial/Source/DynamicMaterial/Public/Components/RenderTargetRenderers/DMRenderTargetTextRenderer.h
--- a/Engine/Plugins/Experimental/DynamicMaterial/Source/DynamicMaterial/Public/Components/RenderTargetRenderers/DMRenderTargetTextRenderer.h
+++ b/Engine/Plugins/Experimental/DynamicMaterial/Source/DynamicMaterial/Public/Components/RenderTargetRenderers/DMRenderTargetTextRenderer.h
@@ -2,15 +2,20 @@
 
 #pragma once
 
-#include "Components/DMRenderTargetRenderer.h"
+#include "Components/RenderTargetRenderers/DMRenderTargetWidgetRendererBase.h"
 #include "Delegates/IDelegateInstance.h"
+#include "Fonts/FontCache.h"
 #include "Framework/Text/TextLayout.h"
+#include "InstancedStruct.h"
 #include "Internationalization/Text.h"
 #include "Math/Color.h"
 #include "Math/Vector2D.h"
+#include "Templates/SharedPointer.h"
 #include "DMRenderTargetTextRenderer.generated.h"
 
 class FCanvasTextItem;
+class STextBlock;
+class SWidget;
 class UCanvas;
 class UFont;
 
@@ -23,11 +28,13 @@ struct FDMTextLine
 	FString Line;
 
 	UPROPERTY()
-	float Width = 0.f;
+	float Width;
+
+	TSharedPtr<STextBlock> Widget;
 };
 
 UCLASS(BlueprintType, ClassGroup = "Material Designer", meta = (DisplayName = "Material Designer Render Target Text Renderer"))
-class DYNAMICMATERIAL_API UDMRenderTargetTextRenderer : public UDMRenderTargetRenderer
+class DYNAMICMATERIAL_API UDMRenderTargetTextRenderer : public UDMRenderTargetWidgetRendererBase
 {
 	GENERATED_BODY()
 
@@ -37,10 +44,10 @@ public:
 	UDMRenderTargetTextRenderer();
 
 	UFUNCTION(BlueprintPure, Category = "Material Designer")
-	UFont* GetFont() const;
+	const FSlateFontInfo& GetFontInfo() const;
 
 	UFUNCTION(BlueprintCallable, Category = "Material Designer")
-	void SetFont(UFont* InFont);
+	void SetFontInfo(const FSlateFontInfo& InFontInfo);
 
 	UFUNCTION(BlueprintPure, Category = "Material Designer")
 	const FText& GetText() const;
@@ -61,112 +68,118 @@ public:
 	void SetBackgroundColor(const FLinearColor& InBackgroundColor);
 
 	UFUNCTION(BlueprintPure, Category = "Material Designer")
-	ETextJustify::Type GetJustify() const;
+	bool GetHasHighlight() const;
 
 	UFUNCTION(BlueprintCallable, Category = "Material Designer")
-	void SetJustify(ETextJustify::Type InJustify);
+	void SetHasHighlight(bool bInHasHighlight);
 
 	UFUNCTION(BlueprintPure, Category = "Material Designer")
-	float GetKerning() const;
+	const FLinearColor& GetHighlightColor() const;
 
 	UFUNCTION(BlueprintCallable, Category = "Material Designer")
-	void SetKerning(float InKerning);
+	void SetHighlightColor(const FLinearColor& InHighlightColor);
 
 	UFUNCTION(BlueprintPure, Category = "Material Designer")
-	float GetLineHeight() const;
+	bool GetHasShadow() const;
 
 	UFUNCTION(BlueprintCallable, Category = "Material Designer")
-	void SetLineHeight(float InLineHeight);
+	void SetHasShadow(bool bInHasShadow);
 
 	UFUNCTION(BlueprintPure, Category = "Material Designer")
-	float GetPaddingLeft() const;
+	const FLinearColor& GetShadowColor() const;
 
 	UFUNCTION(BlueprintCallable, Category = "Material Designer")
-	void SetPaddingLeft(float InPadding);
+	void SetShadowColor(const FLinearColor& InShadowColor);
 
 	UFUNCTION(BlueprintPure, Category = "Material Designer")
-	float GetPaddingRight() const;
+	const FVector2D& GetShadowOffset() const;
 
 	UFUNCTION(BlueprintCallable, Category = "Material Designer")
-	void SetPaddingRight(float InPadding);
+	void SetShadowOffset(const FVector2D& InShadowOffset);
 
 	UFUNCTION(BlueprintPure, Category = "Material Designer")
-	float GetPaddingTop() const;
+	bool GetAutoWrapText() const;
 
 	UFUNCTION(BlueprintCallable, Category = "Material Designer")
-	void SetPaddingTop(float InPadding);
+	void SetAutoWrapText(bool bInAutoWrap);
 
 	UFUNCTION(BlueprintPure, Category = "Material Designer")
-	float GetPaddingBottom() const;
+	float GetWrapTextAt() const;
 
 	UFUNCTION(BlueprintCallable, Category = "Material Designer")
-	void SetPaddingBottom(float InPadding);
+	void SetWrapTextAt(float InWrapAt);
 
 	UFUNCTION(BlueprintPure, Category = "Material Designer")
-	const FIntPoint& GetTextScale() const;
+	ETextWrappingPolicy GetWrappingPolicy() const;
 
 	UFUNCTION(BlueprintCallable, Category = "Material Designer")
-	void SetTextScale(const FIntPoint& InTextScale);
+	void SetWrappingPolicy(ETextWrappingPolicy InWrappingPolicy);
 
 	UFUNCTION(BlueprintPure, Category = "Material Designer")
-	const FIntPoint& GetTextureSizeOverride() const;
+	ETextJustify::Type GetJustify() const;
 
 	UFUNCTION(BlueprintCallable, Category = "Material Designer")
-	void SetTextureSizeOverride(const FIntPoint& InTextureSizeOverride);
+	void SetJustify(ETextJustify::Type InJustify);
 
 	UFUNCTION(BlueprintPure, Category = "Material Designer")
-	bool GetHasOutline() const;
+	ETextTransformPolicy GetTransformPolicy() const;
 
 	UFUNCTION(BlueprintCallable, Category = "Material Designer")
-	void SetHasOutline(bool bInHasOutline);
+	void SetTransformPolicy(ETextTransformPolicy InTransformPolicy);
 
 	UFUNCTION(BlueprintPure, Category = "Material Designer")
-	const FLinearColor& GetOutlineColor() const;
+	ETextFlowDirection GetFlowDirection() const;
 
 	UFUNCTION(BlueprintCallable, Category = "Material Designer")
-	void SetOutlineColor(const FLinearColor& InOutlineColor);
+	void SetFlowDirection(ETextFlowDirection InFlowDirection);
 
 	UFUNCTION(BlueprintPure, Category = "Material Designer")
-	bool GetHasShadow() const;
+	ETextShapingMethod GetShapingMethod() const;
 
 	UFUNCTION(BlueprintCallable, Category = "Material Designer")
-	void SetHasShadow(bool bInHasShadow);
+	void SetShapingMethod(ETextShapingMethod InShapingMethod);
 
 	UFUNCTION(BlueprintPure, Category = "Material Designer")
-	const FLinearColor& GetShadowColor() const;
+	const TInstancedStruct<FSlateBrush>& GetStrikeBrush() const;
 
 	UFUNCTION(BlueprintCallable, Category = "Material Designer")
-	void SetShadowColor(const FLinearColor& InShadowColor);
+	void SetStrikeBrush(const TInstancedStruct<FSlateBrush>& InStrikeBrush);
 
 	UFUNCTION(BlueprintPure, Category = "Material Designer")
-	const FVector2D& GetShadowOffset() const;
+	float GetLineHeight() const;
 
 	UFUNCTION(BlueprintCallable, Category = "Material Designer")
-	void SetShadowOffset(const FVector2D& InShadowOffset);
+	void SetLineHeight(float InLineHeight);
 
 	UFUNCTION(BlueprintPure, Category = "Material Designer")
-	bool GetHasGlow() const;
+	float GetPaddingLeft() const;
 
 	UFUNCTION(BlueprintCallable, Category = "Material Designer")
-	void SetHasGlow(bool bInHasGlow);
+	void SetPaddingLeft(float InPadding);
 
 	UFUNCTION(BlueprintPure, Category = "Material Designer")
-	const FLinearColor& GetGlowColor() const;
+	float GetPaddingRight() const;
 
 	UFUNCTION(BlueprintCallable, Category = "Material Designer")
-	void SetGlowColor(const FLinearColor& InGlowColor);
+	void SetPaddingRight(float InPadding);
 
 	UFUNCTION(BlueprintPure, Category = "Material Designer")
-	const FVector2D& GetGlowInnerRadius() const;
+	float GetPaddingTop() const;
+
+	UFUNCTION(BlueprintCallable, Category = "Material Designer")
+	void SetPaddingTop(float InPadding);
+
+	UFUNCTION(BlueprintPure, Category = "Material Designer")
+	float GetPaddingBottom() const;
 
 	UFUNCTION(BlueprintCallable, Category = "Material Designer")
-	void SetGlowInnerRadius(const FVector2D& InGlowInnerRadius);
+	void SetPaddingBottom(float InPadding);
 
 	UFUNCTION(BlueprintPure, Category = "Material Designer")
-	const FVector2D& GetGlowOuterRadius() const;
+	bool IsOverridingRenderTargetSize() const;
 
 	UFUNCTION(BlueprintCallable, Category = "Material Designer")
-	void SetGlowOuterRadius(const FVector2D& InGlowOuterRadius);
+	void SetOverrideRenderTargetSize(bool bInOverride);
 
 #if WITH_EDITOR
 	//~ Begin IDMJsonSerializable
@@ -180,9 +193,9 @@ public:
 #endif
 
 protected:
-	UPROPERTY(EditInstanceOnly, BlueprintReadWrite, Setter = SetFont, BlueprintSetter = SetFont, Category = "Material Designer|Text",
+	UPROPERTY(EditInstanceOnly, BlueprintReadWrite, Setter = SetFontInfo, BlueprintSetter = SetFontInfo, Category = "Material Designer|Text",
 		meta = (NotKeyframeable, AllowPrivateAccess = "true"))
-	UFont* Font = nullptr;
+	FSlateFontInfo FontInfo;
 
 	UPROPERTY(EditInstanceOnly, BlueprintReadWrite, Setter = SetText, BlueprintSetter = SetText, Category = "Material Designer|Text",
 		meta = (MultiLine, NotKeyframeable, AllowPrivateAccess = "true"))
@@ -195,13 +208,57 @@ protected:
 		meta = (AllowPrivateAccess = "true"))
 	FLinearColor TextColor = FLinearColor::White;
 
+	UPROPERTY(EditInstanceOnly, BlueprintReadWrite, Setter = SetHasHighlight, BlueprintSetter = SetHasHighlight, Category = "Material Designer|Text",
+		meta = (NotKeyframeable, AllowPrivateAccess = "true"))
+	bool bHasHighlight = false;
+
+	UPROPERTY(EditInstanceOnly, BlueprintReadWrite, Setter = SetHighlightColor, BlueprintSetter = SetHighlightColor, Category = "Material Designer|Text",
+		meta = (AllowPrivateAccess = "true", EditCondition = "bHasHighlight", EditConditionHides = true))
+	FLinearColor HighlightColor = FLinearColor::Black;
+
+	UPROPERTY(EditInstanceOnly, BlueprintReadWrite, Setter = SetHasShadow, BlueprintSetter = SetHasShadow, Category = "Material Designer|Text",
+		meta = (NotKeyframeable, AllowPrivateAccess = "true"))
+	bool bHasShadow = false;
+
+	UPROPERTY(EditInstanceOnly, BlueprintReadWrite, Setter = SetShadowColor, BlueprintSetter = SetShadowColor, Category = "Material Designer|Text",
+		meta = (AllowPrivateAccess = "true", EditCondition = "bHasShadow", EditConditionHides = true))
+	FLinearColor ShadowColor = FLinearColor::Black;
+
+	UPROPERTY(EditInstanceOnly, BlueprintReadWrite, Setter = SetShadowOffset, BlueprintSetter = SetShadowOffset, Category = "Material Designer|Text",
+		meta = (AllowPrivateAccess = "true", EditCondition = "bHasShadow", EditConditionHides = true))
+	FVector2D ShadowOffset = FVector2D(1.0, 1.0);
+
+	UPROPERTY(EditInstanceOnly, BlueprintReadWrite, Setter = SetAutoWrapText, BlueprintSetter = SetAutoWrapText, Category = "Material Designer|Text",
+		meta = (NotKeyframeable, AllowPrivateAccess = "true"))
+	bool bAutoWrapText = false;
+
+	UPROPERTY(EditInstanceOnly, BlueprintReadWrite, Setter = SetWrapTextAt, BlueprintSetter = SetWrapTextAt, Category = "Material Designer|Text",
+		meta = (NotKeyframeable, AllowPrivateAccess = "true", EditCondition="bAutoWrapText", EditConditionHides))
+	float WrapTextAt = 0.f;
+
+	UPROPERTY(EditInstanceOnly, BlueprintReadWrite, Setter = SetWrappingPolicy, BlueprintSetter = SetWrappingPolicy, Category = "Material Designer|Text",
+		meta = (NotKeyframeable, AllowPrivateAccess = "true", EditCondition = "bAutoWrapText", EditConditionHides))
+	ETextWrappingPolicy WrappingPolicy = ETextWrappingPolicy::DefaultWrapping;
+
 	UPROPERTY(EditInstanceOnly, BlueprintReadWrite, Setter = SetJustify, BlueprintSetter = SetJustify, Category = "Material Designer|Text",
 		meta = (NotKeyframeable, AllowPrivateAccess = "true"))
 	TEnumAsByte<ETextJustify::Type> Justify = ETextJustify::Left;
 
-	UPROPERTY(EditInstanceOnly, BlueprintReadWrite, Setter = SetKerning, BlueprintSetter = SetKerning, Category = "Material Designer|Text",
-		meta = (AllowPrivateAccess = "true"))
-	float Kerning = 0.f;
+	UPROPERTY(EditInstanceOnly, BlueprintReadWrite, Setter = SetTransformPolicy, BlueprintSetter = SetTransformPolicy, Category = "Material Designer|Text",
+		meta = (NotKeyframeable, AllowPrivateAccess = "true"))
+	ETextTransformPolicy TransformPolicy = ETextTransformPolicy::None;
+
+	UPROPERTY(EditInstanceOnly, BlueprintReadWrite, Setter = SetFlowDirection, BlueprintSetter = SetFlowDirection, Category = "Material Designer|Text",
+		meta = (NotKeyframeable, AllowPrivateAccess = "true"))
+	ETextFlowDirection FlowDirection = ETextFlowDirection::Auto;
+
+	UPROPERTY(EditInstanceOnly, BlueprintReadWrite, Setter = SetShapingMethod, BlueprintSetter = SetShapingMethod, Category = "Material Designer|Text",
+		meta = (NotKeyframeable, AllowPrivateAccess = "true"))
+	ETextShapingMethod ShapingMethod = ETextShapingMethod::Auto;
+
+	UPROPERTY(EditInstanceOnly, BlueprintReadWrite, Setter = SetStrikeBrush, BlueprintSetter = SetStrikeBrush, Category = "Material Designer|Text",
+		meta = (NotKeyframeable, AllowPrivateAccess = "true"))
+	TInstancedStruct<FSlateBrush> StrikeBrush;
 
 	/** Multiplier on the base font height. */
 	UPROPERTY(EditInstanceOnly, BlueprintReadWrite, Setter = SetLineHeight, BlueprintSetter = SetLineHeight, Category = "Material Designer|Text",
@@ -228,62 +285,21 @@ protected:
 		meta = (NotKeyframeable, AllowPrivateAccess = "true"))
 	float PaddingBottom = 0.f;
 
-	UPROPERTY(EditInstanceOnly, BlueprintReadWrite, Setter = SetTextScale, BlueprintSetter = SetTextScale, Category = "Material Designer|Text",
-		meta = (NotKeyframeable, AllowPrivateAccess = "true", ClampMin = 1, UIMin = 1, ClampMax = 64, UIMax = 64,
-			ToolTip = "Increasing this will increase the resolution of the text."))
-	FIntPoint TextScale = FIntPoint(10, 10);
-
-	UPROPERTY(EditInstanceOnly, BlueprintReadWrite, Setter = SetTextureSizeOverride, BlueprintSetter = SetTextureSizeOverride, Category = "Material Designer|Text",
-		meta = (NotKeyframeable, AllowPrivateAccess = "true", ToolTip = "Override the size of the render target. Set to 0 to disable."))
-	FIntPoint TextureSizeOverride = FIntPoint(0, 0);
-
-	UPROPERTY(EditInstanceOnly, BlueprintReadWrite, Setter = SetHasOutline, BlueprintSetter = SetHasOutline, Category = "Material Designer|Text",
-		meta = (AllowPrivateAccess = "true"))
-	bool bOutline = false;
-
-	UPROPERTY(EditInstanceOnly, BlueprintReadWrite, Setter = SetOutlineColor, BlueprintSetter = SetOutlineColor, Category = "Material Designer|Text",
-		meta = (AllowPrivateAccess = "true", EditCondition = "bOutline", EditConditionHides = true))
-	FLinearColor OutlineColor = FLinearColor::Black;
-
-	UPROPERTY(EditInstanceOnly, BlueprintReadWrite, Setter = SetHasShadow, BlueprintSetter = SetHasShadow, Category = "Material Designer|Text",
-		meta = (AllowPrivateAccess = "true"))
-	bool bShadow = false;
-
-	UPROPERTY(EditInstanceOnly, BlueprintReadWrite, Setter = SetShadowColor, BlueprintSetter = SetShadowColor, Category = "Material Designer|Text",
-		meta = (AllowPrivateAccess = "true", EditCondition = "bShadow", EditConditionHides = true))
-	FLinearColor ShadowColor = FLinearColor::Black;
-
-	UPROPERTY(EditInstanceOnly, BlueprintReadWrite, Setter = SetShadowOffset, BlueprintSetter = SetShadowOffset, Category = "Material Designer|Text",
-		meta = (AllowPrivateAccess = "true", EditCondition = "bShadow", EditConditionHides = true))
-	FVector2D ShadowOffset = FVector2D(1.0, 1.0);
-
-	UPROPERTY(EditInstanceOnly, BlueprintReadWrite, Setter = SetHasGlow, BlueprintSetter = SetHasGlow, Category = "Material Designer|Text",
-		meta = (AllowPrivateAccess = "true"))
-	bool bGlow = false;
-
-	UPROPERTY(EditInstanceOnly, BlueprintReadWrite, Setter = SetGlowColor, BlueprintSetter = SetGlowColor, Category = "Material Designer|Text",
-		meta = (AllowPrivateAccess = "true", EditCondition = "bGlow", EditConditionHides = true))
-	FLinearColor GlowColor = FLinearColor::White;
-
-	UPROPERTY(EditInstanceOnly, BlueprintReadWrite, Setter = SetGlowInnerRadius, BlueprintSetter = SetGlowInnerRadius, Category = "Material Designer|Text",
-		meta = (AllowPrivateAccess = "true", EditCondition = "bGlow", EditConditionHides = true, ClampMin = 0, UIMin = 0, ClampMax = 1, UIMax = 1))
-	FVector2D GlowInnerRadius = FVector2D(0.5, 0.5);
-
-	UPROPERTY(EditInstanceOnly, BlueprintReadWrite, Setter = SetGlowOuterRadius, BlueprintSetter = SetGlowOuterRadius, Category = "Material Designer|Text",
-		meta = (AllowPrivateAccess = "true", EditCondition = "bGlow", EditConditionHides = true, ClampMin = 0, UIMin = 0, ClampMax = 1, UIMax = 1))
-	FVector2D GlowOuterRadius = FVector2D(0.5, 0.5);
+	UPROPERTY(EditInstanceOnly, BlueprintReadWrite, Setter = SetOverrideRenderTargetSize, BlueprintSetter = SetOverrideRenderTargetSize, Category = "Material Designer|Text",
+		meta = (NotKeyframeable, AllowPrivateAccess = "true", ToolTip = "When true, will change the size of the render target to fit the text."))
+	bool bOverrideRenderTargetSize = true;
 
 	bool bRecalculateTextSize = false;
 
-	void CalculateLineSizes();
-
-	FIntPoint GetRequiredTextureSize() const;
-
 	void UpdateTextLines();
 
-	FCanvasTextItem CreateTextItem(const FVector2D& InPosition, const FText& InText) const;
+	TSharedRef<STextBlock> CreateTextWidget(const FText& InText) const;
+
+	void SetCustomTextureSize();
 
-	void UpdateTextureSize();
+	//~ Begin UDMRenderTargetWidgetRendererBase
+	virtual void CreateWidgetInstance() override;
+	//~ End UDMRenderTargetWidgetRendererBase
 
 	//~ Begin UDMRenderTargetRenderer
 	virtual void UpdateRenderTarget_Internal() override;
diff --git a/Engine/Plugins/Experimental/DynamicMaterial/Source/DynamicMaterial/Public/Components/RenderTargetRenderers/DMRenderTargetUMGWidgetRenderer.h b/Engine/Plugins/Experimental/DynamicMaterial/Source/DynamicMaterial/Public/Components/RenderTargetRenderers/DMRenderTargetUMGWidgetRenderer.h
new file mode 100644
--- /dev/null
+++ b/Engine/Plugins/Experimental/DynamicMaterial/Source/DynamicMaterial/Public/Components/RenderTargetRenderers/DMRenderTargetUMGWidgetRenderer.h
@@ -0,0 +1,47 @@
+// Copyright Epic Games, Inc. All Rights Reserved.
+
+#pragma once
+
+#include "Components/RenderTargetRenderers/DMRenderTargetWidgetRendererBase.h"
+#include "DMRenderTargetUMGWidgetRenderer.generated.h"
+
+class FWidgetRenderer;
+class UWidget;
+
+UCLASS(BlueprintType, ClassGroup = "Material Designer", meta = (DisplayName = "Material Designer Render Target Widget Renderer"))
+class DYNAMICMATERIAL_API UDMRenderTargetUMGWidgetRenderer : public UDMRenderTargetWidgetRendererBase
+{
+	GENERATED_BODY()
+
+public:
+	UDMRenderTargetUMGWidgetRenderer();
+
+	UFUNCTION(BlueprintPure, Category = "Material Designer")
+	TSubclassOf<UWidget> GetWidgetClass() const { return WidgetClass; }
+
+	UFUNCTION(BlueprintCallable, Category = "Material Designer")
+	void SetWidgetClass(TSubclassOf<UWidget> InWidgetClass);
+
+#if WITH_EDITOR
+	//~ Begin IDMJsonSerializable
+	virtual TSharedPtr<FJsonValue> JsonSerialize() const override;
+	virtual bool JsonDeserialize(const TSharedPtr<FJsonValue>& InJsonValue) override;
+	//~ End IDMJsonSerializable
+
+	//~ Begin UObject
+	virtual void PostEditChangeProperty(FPropertyChangedEvent& InPropertyChangedEvent) override;
+	//~ End UObject
+#endif
+
+protected:
+	UPROPERTY(EditInstanceOnly, BlueprintReadWrite, Setter = SetWidgetClass, BlueprintSetter = SetWidgetClass, Category = "Material Designer",
+		meta = (NotKeyframeable, AllowPrivateAccess = "true"))
+	TSubclassOf<UWidget> WidgetClass;
+
+	UPROPERTY()
+	TObjectPtr<UWidget> WidgetInstance;
+
+	//~ Begin UDMRenderTargetWidgetRendererBase
+	virtual void CreateWidgetInstance() override;
+	//~ End UDMRenderTargetWidgetRendererBase
+};
diff --git a/Engine/Plugins/Experimental/DynamicMaterial/Source/DynamicMaterial/Public/Components/RenderTargetRenderers/DMRenderTargetWidgetRenderer.h b/Engine/Plugins/Experimental/DynamicMaterial/Source/DynamicMaterial/Public/Components/RenderTargetRenderers/DMRenderTargetWidgetRenderer.h
deleted file mode 100644
diff --git a/Engine/Plugins/Experimental/DynamicMaterial/Source/DynamicMaterial/Public/Components/RenderTargetRenderers/DMRenderTargetWidgetRendererBase.h b/Engine/Plugins/Experimental/DynamicMaterial/Source/DynamicMaterial/Public/Components/RenderTargetRenderers/DMRenderTargetWidgetRendererBase.h
new file mode 100644
--- /dev/null
+++ b/Engine/Plugins/Experimental/DynamicMaterial/Source/DynamicMaterial/Public/Components/RenderTargetRenderers/DMRenderTargetWidgetRendererBase.h
@@ -0,0 +1,29 @@
+// Copyright Epic Games, Inc. All Rights Reserved.
+
+#pragma once
+
+#include "Components/DMRenderTargetRenderer.h"
+#include "DMRenderTargetWidgetRendererBase.generated.h"
+
+class FWidgetRenderer;
+class SWidget;
+class UWidget;
+
+UCLASS(Abstract, BlueprintType, ClassGroup = "Material Designer", meta = (DisplayName = "Material Designer Render Target Widget Renderer"))
+class DYNAMICMATERIAL_API UDMRenderTargetWidgetRendererBase : public UDMRenderTargetRenderer
+{
+	GENERATED_BODY()
+
+public:
+	UDMRenderTargetWidgetRendererBase();
+
+protected:
+	TSharedPtr<SWidget> Widget;
+	TSharedPtr<FWidgetRenderer> WidgetRenderer;
+
+	virtual void CreateWidgetInstance() PURE_VIRTUAL(UDMRenderTargetUMGWidgetRenderer::CreateWidgetInstance);
+
+	//~ Begin UDMRenderTargetRenderer
+	virtual void UpdateRenderTarget_Internal() override;
+	//~ End UDMRenderTargetRenderer
+};
diff --git a/Engine/Plugins/Experimental/DynamicMaterial/Source/DynamicMaterialEditor/DynamicMaterialEditor.Build.cs b/Engine/Plugins/Experimental/DynamicMaterial/Source/DynamicMaterialEditor/DynamicMaterialEditor.Build.cs
--- a/Engine/Plugins/Experimental/DynamicMaterial/Source/DynamicMaterialEditor/DynamicMaterialEditor.Build.cs
+++ b/Engine/Plugins/Experimental/DynamicMaterial/Source/DynamicMaterialEditor/DynamicMaterialEditor.Build.cs
@@ -44,6 +44,7 @@ public class DynamicMaterialEditor : ModuleRules
 				"RHI",
 				"Slate",
 				"SlateCore",
+				"StructUtils",
 				"ToolMenus",
 				"TypedElementRuntime",
 				"UMG",
diff --git a/Engine/Plugins/Experimental/DynamicMaterial/Source/DynamicMaterialEditor/Private/Menus/DMMaterialSlotLayerMenus.cpp b/Engine/Plugins/Experimental/DynamicMaterial/Source/DynamicMaterialEditor/Private/Menus/DMMaterialSlotLayerMenus.cpp
--- a/Engine/Plugins/Experimental/DynamicMaterial/Source/DynamicMaterialEditor/Private/Menus/DMMaterialSlotLayerMenus.cpp
+++ b/Engine/Plugins/Experimental/DynamicMaterial/Source/DynamicMaterialEditor/Private/Menus/DMMaterialSlotLayerMenus.cpp
@@ -10,7 +10,7 @@
 #include "Components/MaterialStageExpressions/DMMSETextureSample.h"
 #include "Components/MaterialStageExpressions/DMMSETextureSampleEdgeColor.h"
 #include "Components/RenderTargetRenderers/DMRenderTargetTextRenderer.h"
-#include "Components/RenderTargetRenderers/DMRenderTargetWidgetRenderer.h"
+#include "Components/RenderTargetRenderers/DMRenderTargetUMGWidgetRenderer.h"
 #include "DMDefs.h"
 #include "DMMaterialSlotLayerAddEffectMenus.h"
 #include "DMValueDefinition.h"
@@ -573,7 +573,7 @@ void FDMMaterialSlotLayerMenus::AddAddLayerSection(UToolMenu* InMenu)
 		FUIAction(FExecuteAction::CreateSP(
 			EditorWidget->GetActiveSlotWidget().Get(),
 			&SDMSlot::AddNewLayer_Renderer,
-			TSubclassOf<UDMRenderTargetRenderer>(UDMRenderTargetWidgetRenderer::StaticClass())
+			TSubclassOf<UDMRenderTargetRenderer>(UDMRenderTargetUMGWidgetRenderer::StaticClass())
 		))
 	);
 
diff --git a/Engine/Plugins/Experimental/DynamicMaterial/Source/DynamicMaterialEditor/Private/Menus/DMMaterialStageSourceMenus.cpp b/Engine/Plugins/Experimental/DynamicMaterial/Source/DynamicMaterialEditor/Private/Menus/DMMaterialStageSourceMenus.cpp
--- a/Engine/Plugins/Experimental/DynamicMaterial/Source/DynamicMaterialEditor/Private/Menus/DMMaterialStageSourceMenus.cpp
+++ b/Engine/Plugins/Experimental/DynamicMaterial/Source/DynamicMaterialEditor/Private/Menus/DMMaterialStageSourceMenus.cpp
@@ -23,7 +23,7 @@
 #include "Components/MaterialStageInputs/DMMSISlot.h"
 #include "Components/MaterialStageInputs/DMMSIValue.h"
 #include "Components/RenderTargetRenderers/DMRenderTargetTextRenderer.h"
-#include "Components/RenderTargetRenderers/DMRenderTargetWidgetRenderer.h"
+#include "Components/RenderTargetRenderers/DMRenderTargetUMGWidgetRenderer.h"
 #include "DMDefs.h"
 #include "DMValueDefinition.h"
 #include "DynamicMaterialEditorModule.h"
@@ -1219,14 +1219,14 @@ namespace UE::DynamicMaterialEditor::Private
 			FScopedTransaction Transaction(LOCTEXT("SetStageInputBase", "Set Material Designer Base Source"));
 			Stage->Modify();
 
-			UDMBlueprintFunctionLibrary::SetStageInputToRenderer(Stage, UDMRenderTargetWidgetRenderer::StaticClass(), UDMMaterialStageBlend::InputB);
+			UDMBlueprintFunctionLibrary::SetStageInputToRenderer(Stage, UDMRenderTargetUMGWidgetRenderer::StaticClass(), UDMMaterialStageBlend::InputB);
 		}
 		else if (StageSource->IsA<UDMMaterialStageThroughputLayerBlend>())
 		{
 			FScopedTransaction Transaction(LOCTEXT("SetStageInputMask", "Set Material Designer Mask Source"));
 			Stage->Modify();
 
-			UDMBlueprintFunctionLibrary::SetStageInputToRenderer(Stage, UDMRenderTargetWidgetRenderer::StaticClass(), UDMMaterialStageThroughputLayerBlend::InputMaskSource);
+			UDMBlueprintFunctionLibrary::SetStageInputToRenderer(Stage, UDMRenderTargetUMGWidgetRenderer::StaticClass(), UDMMaterialStageThroughputLayerBlend::InputMaskSource);
 		}
 		else
 		{
diff --git a/Engine/Plugins/Experimental/DynamicMaterial/Source/DynamicMaterialEditor/Private/Slate/SDMComponentEdit.cpp b/Engine/Plugins/Experimental/DynamicMaterial/Source/DynamicMaterialEditor/Private/Slate/SDMComponentEdit.cpp
--- a/Engine/Plugins/Experimental/DynamicMaterial/Source/DynamicMaterialEditor/Private/Slate/SDMComponentEdit.cpp
+++ b/Engine/Plugins/Experimental/DynamicMaterial/Source/DynamicMaterialEditor/Private/Slate/SDMComponentEdit.cpp
@@ -238,15 +238,6 @@ SDMComponentEdit::~SDMComponentEdit()
 	{
 		Stage->SetBeingEdited(false);
 	}
-
-	if (UDMMaterialStage* Stage = Cast<UDMMaterialStage>(ComponentWeak.Get()))
-	{
-		Stage->GetOnUpdate().RemoveAll(this);
-	}
-	else if (UDMMaterialEffect* Effect = Cast<UDMMaterialEffect>(ComponentWeak.Get()))
-	{
-		Effect->GetOnUpdate().RemoveAll(this);
-	}
 }
 
 void SDMComponentEdit::Construct(const FArguments& InArgs, UDMMaterialComponent* InComponent, const TWeakPtr<SDMEditor>& InEditorWidget)
