Material Designer
- Horde issue 486620

diff --git a/Engine/Plugins/Experimental/DynamicMaterial/Source/DynamicMaterial/Private/Components/RenderTargetRenderers/DMRenderTargetTextRenderer.cpp b/Engine/Plugins/Experimental/DynamicMaterial/Source/DynamicMaterial/Private/Components/RenderTargetRenderers/DMRenderTargetTextRenderer.cpp
--- a/Engine/Plugins/Experimental/DynamicMaterial/Source/DynamicMaterial/Private/Components/RenderTargetRenderers/DMRenderTargetTextRenderer.cpp
+++ b/Engine/Plugins/Experimental/DynamicMaterial/Source/DynamicMaterial/Private/Components/RenderTargetRenderers/DMRenderTargetTextRenderer.cpp
@@ -689,7 +689,11 @@ void UDMRenderTargetTextRenderer::UpdateTextLines()
 
 		TSharedRef<STextBlock> TextBlock = CreateTextWidget(FText::FromString(NewLine));
 
-		Lines.Emplace(NewLine, TextBlock->ComputeDesiredSize(1.f).X, TextBlock);
+		FDMTextLine Line;
+		Line.Line = NewLine;
+		Line.Widget = TextBlock;
+
+		Lines.Add(Line);
 	}	
 
 	bRecalculateTextSize = true;
@@ -726,7 +730,6 @@ void UDMRenderTargetTextRenderer::CreateWidgetInstance()
 		if (!Line.Widget.IsValid())
 		{
 			Line.Widget = CreateTextWidget(FText::FromString(Line.Line));
-			Line.Width = Line.Widget->ComputeDesiredSize(1.f).X;
 		}
 
 		NewWidget->AddSlot()
diff --git a/Engine/Plugins/Experimental/DynamicMaterial/Source/DynamicMaterial/Public/Components/RenderTargetRenderers/DMRenderTargetTextRenderer.h b/Engine/Plugins/Experimental/DynamicMaterial/Source/DynamicMaterial/Public/Components/RenderTargetRenderers/DMRenderTargetTextRenderer.h
--- a/Engine/Plugins/Experimental/DynamicMaterial/Source/DynamicMaterial/Public/Components/RenderTargetRenderers/DMRenderTargetTextRenderer.h
+++ b/Engine/Plugins/Experimental/DynamicMaterial/Source/DynamicMaterial/Public/Components/RenderTargetRenderers/DMRenderTargetTextRenderer.h
@@ -27,9 +27,6 @@ struct FDMTextLine
 	UPROPERTY()
 	FString Line;
 
-	UPROPERTY()
-	float Width;
-
 	TSharedPtr<STextBlock> Widget;
 };
 
