Material Designer
- Material settings in the global settings do not respect the default values.

diff --git a/Engine/Plugins/Experimental/DynamicMaterial/Source/DynamicMaterialEditor/Private/Slate/SDMComponentEdit.cpp b/Engine/Plugins/Experimental/DynamicMaterial/Source/DynamicMaterialEditor/Private/Slate/SDMComponentEdit.cpp
--- a/Engine/Plugins/Experimental/DynamicMaterial/Source/DynamicMaterialEditor/Private/Slate/SDMComponentEdit.cpp
+++ b/Engine/Plugins/Experimental/DynamicMaterial/Source/DynamicMaterialEditor/Private/Slate/SDMComponentEdit.cpp
@@ -851,6 +851,11 @@ void SDMComponentEdit::GenerateMaterialModelPropertyRows(const TSharedRef<SDMEdi
 
 			ValueHandle.CategoryOverrideName = MaterialSettingsCategory;
 			ValueHandle.NameOverride = InNameOverride;
+			ValueHandle.ResetToDefaultOverride = FResetToDefaultOverride::Create(
+				FIsResetToDefaultVisible::CreateUObject(InValue, &UDMMaterialValue::CanResetToDefault),
+				FResetToDefaultHandler::CreateUObject(InValue, &UDMMaterialValue::ResetToDefault),
+				/* Propagate to children */ false
+			);
 
 			InOutProcessedObjects.Add(InValue);
 		};
