Material Designer
- 5.4 to 5.5 global opacity data loss issue

diff --git a/Engine/Plugins/Experimental/DynamicMaterial/Source/DynamicMaterial/Private/Model/DynamicMaterialModel.cpp b/Engine/Plugins/Experimental/DynamicMaterial/Source/DynamicMaterial/Private/Model/DynamicMaterialModel.cpp
--- a/Engine/Plugins/Experimental/DynamicMaterial/Source/DynamicMaterial/Private/Model/DynamicMaterialModel.cpp
+++ b/Engine/Plugins/Experimental/DynamicMaterial/Source/DynamicMaterial/Private/Model/DynamicMaterialModel.cpp
@@ -441,6 +441,16 @@ void UDynamicMaterialModel::PostLoad()
 
 		GlobalOpacityValue->SetParameterName("VALUE_GlobalOpacity");
 
+		if (FMath::IsNearlyZero(GlobalOpacityValue->GetDefaultValue()))
+		{
+			GlobalOpacityValue->SetDefaultValue(1.f);
+
+			if (FMath::IsNearlyZero(GlobalOpacityValue->GetValue()))
+			{
+				GlobalOpacityValue->SetValue(1.f);
+			}
+		}
+
 		GlobalOpacityValue = nullptr;
 	}
 
