Material Designer
- Storing the expansion state of categories and effects lists.

diff --git a/Engine/Plugins/Experimental/DynamicMaterial/Source/DynamicMaterialEditor/Private/Slate/Layers/SDMSlotLayerItem.cpp b/Engine/Plugins/Experimental/DynamicMaterial/Source/DynamicMaterialEditor/Private/Slate/Layers/SDMSlotLayerItem.cpp
--- a/Engine/Plugins/Experimental/DynamicMaterial/Source/DynamicMaterialEditor/Private/Slate/Layers/SDMSlotLayerItem.cpp
+++ b/Engine/Plugins/Experimental/DynamicMaterial/Source/DynamicMaterialEditor/Private/Slate/Layers/SDMSlotLayerItem.cpp
@@ -34,6 +34,8 @@
 
 #define LOCTEXT_NAMESPACE "SDMSlotLayer"
 
+const FLazyName SDMSlotLayerItem::EffectsListName = "EffectsList";
+
 void SDMSlotLayerItem::Construct(const FArguments& InArgs, const TSharedPtr<SDMSlot>& InSlotWidget, const TSharedRef<STableViewBase>& InLayerView, 
 	const TSharedPtr<FDMMaterialLayerReference>& InLayerItem)
 {
@@ -49,8 +51,6 @@ void SDMSlotLayerItem::Construct(const FArguments& InArgs, const TSharedPtr<SDMS
 	OnStageSelected = InArgs._OnStageSelected;
 	OnLayerLinkToggled = InArgs._OnLayerLinkToggled;
 
-	bDisplayEffectsList = true;
-
 	STableRow<TSharedPtr<FDMMaterialLayerReference>>::Construct(
 		STableRow<TSharedPtr<FDMMaterialLayerReference>>::FArguments()
 		.Padding(2.0f)
@@ -1009,7 +1009,7 @@ void SDMSlotLayerItem::SaveLayerPreviewSize(const float InNewSize, const TShared
 
 EVisibility SDMSlotLayerItem::GetEffectsListVisibility() const
 {
-	return EffectsList->GetLayerItemCount() > 0 && bDisplayEffectsList ? EVisibility::Visible : EVisibility::Collapsed;
+	return EffectsList->GetLayerItemCount() > 0 && AreEffectsExpanded() ? EVisibility::Visible : EVisibility::Collapsed;
 }
 
 TSharedRef<SWidget> SDMSlotLayerItem::CreateNewEffectMenu()
@@ -1181,7 +1181,7 @@ EVisibility SDMSlotLayerItem::GetLayerLinkToggleButtonVisibility() const
 
 FReply SDMSlotLayerItem::OnEffectsToggleButtonClicked()
 {
-	bDisplayEffectsList = !bDisplayEffectsList;
+	SetEffectsExpanded(!AreEffectsExpanded());
 
 	return FReply::Handled();
 }
@@ -1191,7 +1191,7 @@ const FSlateBrush* SDMSlotLayerItem::GetEffectsToggleButtonImage() const
 	static const FSlateBrush* Displayed = FDynamicMaterialEditorStyle::GetBrush("EffectsView.Row.Fx.Opened");
 	static const FSlateBrush* Hidden = FDynamicMaterialEditorStyle::GetBrush("EffectsView.Row.Fx.Closed");
 
-	if (bDisplayEffectsList)
+	if (AreEffectsExpanded())
 	{
 		return Displayed;
 	}
@@ -1293,6 +1293,31 @@ TSharedRef<SWidget> SDMSlotLayerItem::CreateLayerHeaderEditableText() const
 			}));
 }
 
+bool SDMSlotLayerItem::AreEffectsExpanded() const
+{
+	// Default to expanded
+	bool bExpanded = true;
+
+	if (UDMMaterialSlot* Slot = GetSlot())
+	{
+		SDMEditor::GetExpansionState(Slot, EffectsListName, bExpanded);
+	}
+
+	return bExpanded;
+}
+
+void SDMSlotLayerItem::SetEffectsExpanded(bool bInExpanded)
+{
+	UDMMaterialSlot* Slot = GetSlot();
+
+	if (!Slot)
+	{
+		return;
+	}
+
+	SDMEditor::SetExpansionState(Slot, EffectsListName, bInExpanded);
+}
+
 void SDMSlotLayerItem::DeselectAllEffects()
 {
 	EffectsList->ClearSelection();
diff --git a/Engine/Plugins/Experimental/DynamicMaterial/Source/DynamicMaterialEditor/Private/Slate/Layers/SDMSlotLayerItem.h b/Engine/Plugins/Experimental/DynamicMaterial/Source/DynamicMaterialEditor/Private/Slate/Layers/SDMSlotLayerItem.h
--- a/Engine/Plugins/Experimental/DynamicMaterial/Source/DynamicMaterialEditor/Private/Slate/Layers/SDMSlotLayerItem.h
+++ b/Engine/Plugins/Experimental/DynamicMaterial/Source/DynamicMaterialEditor/Private/Slate/Layers/SDMSlotLayerItem.h
@@ -52,14 +52,14 @@ public:
 	int32 GetLayerItemIndex() const;
 	bool IsRootLayer() const;
 
-	FORCEINLINE const TSharedPtr<SDMStage>& GetBaseStageWidget() const { return BaseStageWidget; }
-	FORCEINLINE const TSharedPtr<SDMStage>& GetMaskStageWidget() const { return MaskStageWidget; }
+	const TSharedPtr<SDMStage>& GetBaseStageWidget() const { return BaseStageWidget; }
+	const TSharedPtr<SDMStage>& GetMaskStageWidget() const { return MaskStageWidget; }
 
-	FORCEINLINE const TSharedPtr<SDMStage>& GetSelectedStageWidget() const { return SelectedStageWidget; }
+	const TSharedPtr<SDMStage>& GetSelectedStageWidget() const { return SelectedStageWidget; }
 
-	FORCEINLINE bool AreEffectsExpanded() const { return bDisplayEffectsList; }
+	bool AreEffectsExpanded() const;
 
-	void SetEffectsExpanded(bool bInExpanded) { bDisplayEffectsList = bInExpanded; }
+	void SetEffectsExpanded(bool bInExpanded);
 
 	void DeselectAllEffects();
 
@@ -68,6 +68,8 @@ public:
 	TSharedPtr<SDMLayerEffectsView> GetEffectsList() const { return EffectsList; }
 
 protected:
+	static const FLazyName EffectsListName;
+
 	TWeakPtr<SDMSlot> SlotWidgetWeak;
 	TWeakPtr<STableViewBase> LayerViewWeak;
 	TSharedPtr<FDMMaterialLayerReference> LayerItem;
@@ -86,7 +88,6 @@ protected:
 	TSharedPtr<SDMStage> SelectedStageWidget;
 
 	TSharedPtr<SDMLayerEffectsView> EffectsList;
-	bool bDisplayEffectsList;
 
 	TSharedPtr<SBox> LayerHeaderTextContainer;
 
diff --git a/Engine/Plugins/Experimental/DynamicMaterial/Source/DynamicMaterialEditor/Private/Slate/SDMComponentEdit.cpp b/Engine/Plugins/Experimental/DynamicMaterial/Source/DynamicMaterialEditor/Private/Slate/SDMComponentEdit.cpp
--- a/Engine/Plugins/Experimental/DynamicMaterial/Source/DynamicMaterialEditor/Private/Slate/SDMComponentEdit.cpp
+++ b/Engine/Plugins/Experimental/DynamicMaterial/Source/DynamicMaterialEditor/Private/Slate/SDMComponentEdit.cpp
@@ -245,6 +245,8 @@ void SDMComponentEdit::Construct(const FArguments& InArgs, UDMMaterialComponent*
 	ComponentWeak = InComponent;
 	EditorWidgetWeak = InEditorWidget;
 
+	TGuardValue<bool> Constructing(bConstructing, true);
+
 	KeyframeHandler = nullptr;
 
 	UObject* WorldContext = InComponent;
@@ -288,29 +290,43 @@ TSharedRef<SWidget> SDMComponentEdit::CreateEditWidget()
 {
 	SDMEditor::ClearPropertyHandles(this);
 
+	constexpr bool bDefaultCategoryExpansionState = true;
+
+	UDMMaterialComponent* Component = ComponentWeak.Get();
+
 	FCustomDetailsViewArgs Args;
 	Args.KeyframeHandler = KeyframeHandler;
 	Args.bAllowGlobalExtensions = true;
 	Args.bAllowResetToDefault = true;
 	Args.bShowCategories = false;
+	Args.OnExpansionStateChanged.AddSP(this, &SDMComponentEdit::OnExpansionStateChanged);
 
 	TSharedRef<ICustomDetailsView> DetailsView = ICustomDetailsViewModule::Get().CreateCustomDetailsView(Args);
 	FCustomDetailsViewItemId RootId = DetailsView->GetRootItem()->GetItemId();
 
 	TSharedPtr<ICustomDetailsViewItem> DefaultCategoryItem;
 
-	auto GetDefaultCategory = [&DefaultCategoryItem, &DetailsView, &RootId]()
+	auto GetDefaultCategory = [this, Component, &DefaultCategoryItem, &DetailsView, &RootId]()
 		{
 			if (!DefaultCategoryItem.IsValid())
 			{
-				DefaultCategoryItem = DetailsView->CreateCustomCategoryItem("General", LOCTEXT("General", "General"))->AsItem();
+				constexpr const TCHAR* DefaultCategoryName = TEXT("General");
+
+				DefaultCategoryItem = DetailsView->CreateCustomCategoryItem(DefaultCategoryName, LOCTEXT("General", "General"))->AsItem();
+				DefaultCategoryItem->RefreshItemId();
 				DetailsView->ExtendTree(RootId, ECustomDetailsTreeInsertPosition::Child, DefaultCategoryItem.ToSharedRef());
+
+				bool bExpansionState = true;
+				SDMEditor::GetExpansionState(Component, DefaultCategoryName, bExpansionState);
+
+				DetailsView->SetItemExpansionState(DefaultCategoryItem->GetItemId(), bExpansionState);
+				Categories.Add(DefaultCategoryName);
 			}
 
 			return DefaultCategoryItem;
 		};
 
-	if (UDMMaterialStage* Stage = Cast<UDMMaterialStage>(ComponentWeak.Get()))
+	if (UDMMaterialStage* Stage = Cast<UDMMaterialStage>(Component))
 	{
 		TSharedPtr<ICustomDetailsViewCustomItem> TypeSelectorItem = DetailsView->CreateCustomItem(
 			FName(TEXT("SamplerType")),
@@ -389,7 +405,15 @@ TSharedRef<SWidget> SDMComponentEdit::CreateEditWidget()
 			if (!CategoryItem.IsValid())
 			{
 				CategoryItem = DetailsView->CreateCustomCategoryItem(CategoryName, FText::FromName(CategoryName))->AsItem();
+				CategoryItem->RefreshItemId();
 				DetailsView->ExtendTree(RootId, ECustomDetailsTreeInsertPosition::Child, CategoryItem.ToSharedRef());
+
+				bool bExpansionState = true;
+				SDMEditor::GetExpansionState(Component, CategoryName, bExpansionState);
+
+				DetailsView->SetItemExpansionState(CategoryItem->GetItemId(), bExpansionState);
+
+				Categories.Add(CategoryName);
 			}
 		}
 
@@ -484,7 +508,7 @@ TArray<FDMPropertyHandle> SDMComponentEdit::GetEditRows()
 	}
 	else
 	{
-		if (TSharedPtr<SDMEditor> EditorWidget = EditorWidgetWeak.Pin())
+		if (TSharedPtr<SDMEditor> EditorWidget = GetEditorWidget())
 		{
 			if (UDynamicMaterialModel* MaterialModel = EditorWidget->GetMaterialModel())
 			{
@@ -838,6 +862,30 @@ void SDMComponentEdit::OnUndo()
 	}
 }
 
+void SDMComponentEdit::OnExpansionStateChanged(const TSharedRef<ICustomDetailsViewItem>& InItem, bool bInExpansionState)
+{
+	if (bConstructing)
+	{
+		return;
+	}
+
+	const FCustomDetailsViewItemId& ItemId = InItem->GetItemId();
+
+	if (ItemId.GetItemType() != static_cast<uint32>(EDetailNodeType::Category))
+	{
+		return;
+	}
+
+	TSharedPtr<SDMEditor> EditorWidget = GetEditorWidget();
+
+	if (!EditorWidget.IsValid())
+	{
+		return;
+	}
+
+	EditorWidget->SetExpansionState(ComponentWeak.Get(), *ItemId.GetItemName(), bInExpansionState);
+}
+
 void SDMComponentEdit::GenerateMaterialModelPropertyRows(const TSharedRef<SDMEditor> InEditorWidget, UDynamicMaterialModel* InMaterialModel,
 	TArray<FDMPropertyHandle>& InOutPropertyRows, TSet<UDMMaterialComponent*>& InOutProcessedObjects)
 {
diff --git a/Engine/Plugins/Experimental/DynamicMaterial/Source/DynamicMaterialEditor/Private/Slate/SDMComponentEdit.h b/Engine/Plugins/Experimental/DynamicMaterial/Source/DynamicMaterialEditor/Private/Slate/SDMComponentEdit.h
new file mode 100644
--- /dev/null
+++ b/Engine/Plugins/Experimental/DynamicMaterial/Source/DynamicMaterialEditor/Private/Slate/SDMComponentEdit.h
@@ -0,0 +1,102 @@
+// Copyright Epic Games, Inc. All Rights Reserved.
+
+#pragma once
+
+#include "DMEDefs.h"
+#include "Delegates/IDelegateInstance.h"
+#include "EditorUndoClient.h"
+#include "Widgets/SBoxPanel.h"
+#include "Widgets/SCompoundWidget.h"
+
+class FAssetThumbnailPool;
+class FScopedTransaction;
+class ICustomDetailsViewItem;
+class IDetailKeyframeHandler;
+class IDetailTreeNode;
+class IPropertyHandle;
+class IPropertyRowGenerator;
+class SBox;
+class SDMEditor;
+class SDMEditor;
+class SDMStage;
+class SWidget;
+class UDMMaterialComponent;
+class UDMMaterialEffect;
+class UDMMaterialEffectStack;
+class UDMMaterialSlot;
+class UDMMaterialStageExpressionTextureSample;
+class UDMMaterialStageInputTextureUV;
+class UDMMaterialStageInputThroughput;
+class UDMMaterialStageInputValue;
+class UDMMaterialStageSource;
+class UDMMaterialStageThroughput;
+class UDMMaterialValueFloat1;
+class UDMTextureUV;
+class UDynamicMaterialModelEditorOnlyData;
+class UMaterial;
+enum class ECheckBoxState : uint8;
+struct FDMPropertyHandle;
+struct FSlateIcon;
+
+class DYNAMICMATERIALEDITOR_API SDMComponentEdit : public SCompoundWidget, public FSelfRegisteringEditorUndoClient
+{
+public:
+	static ECheckBoxState IsInputPerChannelMapped(UDMMaterialStageThroughput* InThroughput, int32 InInputIdx);
+	static FText GetInputChannelMapDescription(UDMMaterialStageThroughput* InThroughput, int32 InInputIdx, int32 InChannelIdx);
+	static void CreateKeyFrame(TSharedPtr<IPropertyHandle> InPropertyHandle);
+	static TSharedRef<SWidget> CreateExtensionButtons(const TSharedPtr<SWidget>& InPropertyOwner, UDMMaterialComponent* InComponent,
+		const FName& InPropertyName, bool bInAllowKeyframe, FSimpleDelegate InOnResetDelegate);
+	static TSharedRef<SWidget> CreateExtensionButtons(const TSharedPtr<SDMComponentEdit>& InComponentEditWidget, UDMMaterialComponent* InComponent,
+		const FName& InPropertyName, bool bInAllowKeyframe);
+
+	SLATE_BEGIN_ARGS(SDMComponentEdit)
+		{}
+	SLATE_END_ARGS()
+
+	SDMComponentEdit() = default;
+	virtual ~SDMComponentEdit() override;
+
+	void Construct(const FArguments& InArgs, UDMMaterialComponent* InComponent, const TWeakPtr<SDMEditor>& InEditorWidget);
+
+	FORCEINLINE UDMMaterialComponent* GetComponent() const { return ComponentWeak.Get(); }
+	FORCEINLINE TSharedPtr<SDMEditor> GetEditorWidget() const { return EditorWidgetWeak.Pin(); }
+
+	TSharedPtr<SWidget> CreateSinglePropertyEditWidget(UDMMaterialComponent* InComponent, const FName& InPropertyName);
+
+	//~ Begin FUndoClient
+	virtual void PostUndo(bool bInSuccess) override { OnUndo(); }
+	virtual void PostRedo(bool bInSuccess) override { OnUndo(); }
+	//~ End FUndoClient
+
+protected:
+	static void GenerateMaterialModelPropertyRows(const TSharedRef<SDMEditor> InEditorWidget, UDynamicMaterialModel* InMaterialModel, 
+		TArray<FDMPropertyHandle>& InOutPropertyRows, TSet<UDMMaterialComponent*>& InOutProcessedObjects);
+
+	TWeakObjectPtr<UDMMaterialComponent> ComponentWeak;
+	TWeakPtr<SDMEditor> EditorWidgetWeak;
+
+	bool bConstructing = false;
+
+	TSharedPtr<IDetailKeyframeHandler> KeyframeHandler;
+
+	TSharedPtr<SBox> Container;
+
+	bool bCreatedWithLinkedUVs;
+
+	FDelegateHandle UpdateHandle;
+
+	TSet<FName> Categories;
+
+	TSharedRef<SWidget> CreateEditWidget();
+
+	TArray<FDMPropertyHandle> GetEditRows();
+
+	TSharedRef<SWidget> CreateSourceTypeEditWidget();
+
+	TSharedRef<SWidget> MakeSourceTypeEditWidgetMenuContent();
+	FText GetSourceTypeEditWidgetText() const;
+
+	void OnUndo();
+
+	void OnExpansionStateChanged(const TSharedRef<ICustomDetailsViewItem>& InItem, bool bInExpansionState);
+};
diff --git a/Engine/Plugins/Experimental/DynamicMaterial/Source/DynamicMaterialEditor/Private/Slate/SDMEditor.cpp b/Engine/Plugins/Experimental/DynamicMaterial/Source/DynamicMaterialEditor/Private/Slate/SDMEditor.cpp
--- a/Engine/Plugins/Experimental/DynamicMaterial/Source/DynamicMaterialEditor/Private/Slate/SDMEditor.cpp
+++ b/Engine/Plugins/Experimental/DynamicMaterial/Source/DynamicMaterialEditor/Private/Slate/SDMEditor.cpp
@@ -51,6 +51,8 @@
 
 #define LOCTEXT_NAMESPACE "SDMEditor"
 
+TMap<SDMEditor::FExpansionItem, bool> SDMEditor::ExpansionStates;
+
 namespace UE::DynamicMaterialEditor::Private
 {
 	TSharedPtr<IDetailTreeNode> SearchNodesForProperty(const TArray<TSharedRef<IDetailTreeNode>>& InNodes, FName InPropertyName)
@@ -1491,6 +1493,25 @@ void SDMEditor::DeleteSelectedLayer()
 	}
 }
 
+bool SDMEditor::GetExpansionState(UObject* InOwner, FName InName, bool& bOutExpanded)
+{
+	const FExpansionItem ExpansionItem = {InOwner, InName};
+
+	if (const bool* State = ExpansionStates.Find(ExpansionItem))
+	{
+		bOutExpanded = *State;
+		return true;
+	}
+
+	return false;
+}
+
+void SDMEditor::SetExpansionState(UObject* InOwner, FName InName, bool bInIsExpanded)
+{
+	const FExpansionItem ExpansionItem = {InOwner, InName};
+	ExpansionStates.FindOrAdd(ExpansionItem) = bInIsExpanded;
+}
+
 FReply SDMEditor::OnKeyDown(const FGeometry& MyGeometry, const FKeyEvent& InKeyEvent)
 {
 	if (CommandList.IsValid())
diff --git a/Engine/Plugins/Experimental/DynamicMaterial/Source/DynamicMaterialEditor/Private/Slate/SDMEditor.h b/Engine/Plugins/Experimental/DynamicMaterial/Source/DynamicMaterialEditor/Private/Slate/SDMEditor.h
--- a/Engine/Plugins/Experimental/DynamicMaterial/Source/DynamicMaterialEditor/Private/Slate/SDMEditor.h
+++ b/Engine/Plugins/Experimental/DynamicMaterial/Source/DynamicMaterialEditor/Private/Slate/SDMEditor.h
@@ -7,6 +7,7 @@
 #include "DMEDefs.h"
 #include "DMObjectMaterialProperty.h"
 #include "EditorUndoClient.h"
+#include "UObject/ObjectKey.h"
 #include "Widgets/Layout/SSplitter.h"
 
 class FAssetThumbnailPool;
@@ -46,6 +47,9 @@ public:
 
 	static TSharedRef<SWidget> GetEmptyContent();
 
+	static bool GetExpansionState(UObject* InOwner, FName InName, bool& bOutExpanded);
+	static void SetExpansionState(UObject* InOwner, FName InName, bool bInIsExpanded);
+
 	static FDMPropertyHandle GetPropertyHandle(const SWidget* InOwningWidget, UObject* InObject, const FName& InPropertyName);
 
 	static void ClearPropertyHandles(const SWidget* InOwningWidget);
@@ -127,6 +131,28 @@ public:
 	//~ End FUndoClient
 
 protected:
+	struct FExpansionItem
+	{
+		TObjectKey<UObject> Owner;
+		FName Name;
+
+		friend uint32 GetTypeHash(const FExpansionItem& InItem)
+		{
+			return HashCombineFast(
+				GetTypeHash(InItem.Owner),
+				GetTypeHash(InItem.Name)
+			);
+		}
+
+		bool operator==(const FExpansionItem& InOther) const
+		{
+			return Owner == InOther.Owner
+				&& Name == InOther.Name;
+		}
+	};
+
+	static TMap<FExpansionItem, bool> ExpansionStates;
+
 	static TSharedPtr<FAssetThumbnailPool> ThumbnailPool;
 
 	static TMap<const SWidget*, TArray<FDMPropertyHandle>> PropertyHandleMap;
diff --git a/Engine/Plugins/Experimental/DynamicMaterial/Source/DynamicMaterialEditor/Public/Slate/SDMComponentEdit.h b/Engine/Plugins/Experimental/DynamicMaterial/Source/DynamicMaterialEditor/Public/Slate/SDMComponentEdit.h
deleted file mode 100644
