Material Designer
- Changed the event flow for the value updates to start from the material model, not the editor only data.

diff --git a/Engine/Plugins/Experimental/DynamicMaterial/Source/DynamicMaterial/Private/Components/DMMaterialValue.cpp b/Engine/Plugins/Experimental/DynamicMaterial/Source/DynamicMaterial/Private/Components/DMMaterialValue.cpp
--- a/Engine/Plugins/Experimental/DynamicMaterial/Source/DynamicMaterial/Private/Components/DMMaterialValue.cpp
+++ b/Engine/Plugins/Experimental/DynamicMaterial/Source/DynamicMaterial/Private/Components/DMMaterialValue.cpp
@@ -469,15 +469,10 @@ void UDMMaterialValue::Update(EDMUpdateType InUpdateType)
 
 	Super::Update(InUpdateType);
 
-#if WITH_EDITOR
-	UDynamicMaterialModel* MaterialModel = GetMaterialModel();
-	check(MaterialModel);
-
-	if (IDynamicMaterialModelEditorOnlyDataInterface* ModelEditorOnlyData = MaterialModel->GetEditorOnlyData())
+	if (UDynamicMaterialModel* MaterialModel = GetMaterialModel())
 	{
-		ModelEditorOnlyData->OnValueUpdated(this, InUpdateType);
+		MaterialModel->OnValueUpdated(this, InUpdateType);
 	}
-#endif
 }
 
 #if WITH_EDITOR
diff --git a/Engine/Plugins/Experimental/DynamicMaterial/Source/DynamicMaterial/Private/Model/DynamicMaterialModel.cpp b/Engine/Plugins/Experimental/DynamicMaterial/Source/DynamicMaterial/Private/Model/DynamicMaterialModel.cpp
--- a/Engine/Plugins/Experimental/DynamicMaterial/Source/DynamicMaterial/Private/Model/DynamicMaterialModel.cpp
+++ b/Engine/Plugins/Experimental/DynamicMaterial/Source/DynamicMaterial/Private/Model/DynamicMaterialModel.cpp
@@ -5,6 +5,7 @@
 #include "Components/DMMaterialValue.h"
 #include "Components/MaterialValues/DMMaterialValueFloat1.h"
 #include "Components/MaterialValues/DMMaterialValueFloat2.h"
+#include "Components/DMTextureUV.h"
 #include "DMComponentPath.h"
 #include "DMDefs.h"
 #include "DMValueDefinition.h"
@@ -418,6 +419,54 @@ void UDynamicMaterialModel::ResetData()
 }
 #endif
 
+void UDynamicMaterialModel::OnValueUpdated(UDMMaterialValue* InValue, EDMUpdateType InUpdateType)
+{
+	check(InValue);
+
+	if (InValue->GetMaterialModel() != this)
+	{
+		return;
+	}
+
+	if (InUpdateType == EDMUpdateType::Value && IsValid(DynamicMaterialInstance))
+	{
+		InValue->SetMIDParameter(DynamicMaterialInstance);
+	}
+
+	OnValueUpdateDelegate.Broadcast(this, InValue);
+
+#if WITH_EDITOR
+	if (IDynamicMaterialModelEditorOnlyDataInterface* EditorOnlyData = GetEditorOnlyData())
+	{
+		EditorOnlyData->OnValueUpdated(InValue, InUpdateType);
+	}
+#endif
+}
+
+void UDynamicMaterialModel::OnTextureUVUpdated(UDMTextureUV* InTextureUV)
+{
+	check(InTextureUV);
+
+	if (InTextureUV->GetMaterialModel() != this)
+	{
+		return;
+	}
+
+	if (IsValid(this) && IsValid(DynamicMaterialInstance))
+	{
+		InTextureUV->SetMIDParameters(DynamicMaterialInstance);
+	}
+
+	OnTextureUVUpdateDelegate.Broadcast(this, InTextureUV);
+
+#if WITH_EDITOR
+	if (IDynamicMaterialModelEditorOnlyDataInterface* EditorOnlyData = GetEditorOnlyData())
+	{
+		EditorOnlyData->OnTextureUVUpdated(InTextureUV);
+	}
+#endif
+}
+
 void UDynamicMaterialModel::PostLoad()
 {
 	Super::PostLoad();
diff --git a/Engine/Plugins/Experimental/DynamicMaterial/Source/DynamicMaterial/Public/Model/DynamicMaterialModel.h b/Engine/Plugins/Experimental/DynamicMaterial/Source/DynamicMaterial/Public/Model/DynamicMaterialModel.h
--- a/Engine/Plugins/Experimental/DynamicMaterial/Source/DynamicMaterial/Public/Model/DynamicMaterialModel.h
+++ b/Engine/Plugins/Experimental/DynamicMaterial/Source/DynamicMaterial/Public/Model/DynamicMaterialModel.h
@@ -35,6 +35,9 @@ class UMaterialExpression;
 enum EBlendMode : int;
 #endif
 
+DECLARE_MULTICAST_DELEGATE_TwoParams(FDMOnValueUpdated, UDynamicMaterialModel*, UDMMaterialValue*);
+DECLARE_MULTICAST_DELEGATE_TwoParams(FDMOnTextureUVUpdated, UDynamicMaterialModel*, UDMTextureUV*);
+
 UCLASS(ClassGroup = "Material Designer", DefaultToInstanced, BlueprintType, meta = (DisplayThumbnail = "true"))
 class DYNAMICMATERIAL_API UDynamicMaterialModel : public UObject
 {
@@ -163,9 +166,16 @@ public:
 
 	UFUNCTION(BlueprintCallable, Category = "Material Designer")
 	void ResetData();
+#endif
+
+	virtual void OnValueUpdated(UDMMaterialValue* InValue, EDMUpdateType InUpdateType);
+
+	virtual void OnTextureUVUpdated(UDMTextureUV* InTextureUV);
+
+	FDMOnValueUpdated::RegistrationType& GetOnValueUpdateDelegate() { return OnValueUpdateDelegate; }
+	FDMOnTextureUVUpdated::RegistrationType& GetOnTextureUVUpdateDelegate() { return OnTextureUVUpdateDelegate; }
 
 	//~ Begin UObject
-#endif
 	virtual void PostLoad() override;
 #if WITH_EDITOR
 	virtual void PostEditUndo() override;
@@ -202,6 +212,9 @@ protected:
 	TScriptInterface<IDynamicMaterialModelEditorOnlyDataInterface> EditorOnlyDataSI;
 #endif
 
+	FDMOnValueUpdated OnValueUpdateDelegate;
+	FDMOnTextureUVUpdated OnTextureUVUpdateDelegate;
+
 	/**
 	 * Because they are not assigned to individual properties, the global parameter values end up being their archetype versions
 	 * for versions of the object from before the change. This needs to be fixed on post load.
diff --git a/Engine/Plugins/Experimental/DynamicMaterial/Source/DynamicMaterialEditor/Private/Components/DMMaterialStage.cpp b/Engine/Plugins/Experimental/DynamicMaterial/Source/DynamicMaterialEditor/Private/Components/DMMaterialStage.cpp
--- a/Engine/Plugins/Experimental/DynamicMaterial/Source/DynamicMaterialEditor/Private/Components/DMMaterialStage.cpp
+++ b/Engine/Plugins/Experimental/DynamicMaterial/Source/DynamicMaterialEditor/Private/Components/DMMaterialStage.cpp
@@ -252,8 +252,11 @@ void UDMMaterialStage::AddDelegates()
 		{
 			if (UDynamicMaterialModelEditorOnlyData* ModelEditorOnlyData = Slot->GetMaterialModelEditorOnlyData())
 			{
-				ModelEditorOnlyData->GetOnValueUpdateDelegate().AddUObject(this, &UDMMaterialStage::OnValueUpdated);
-				ModelEditorOnlyData->GetOnTextureUVUpdateDelegate().AddUObject(this, &UDMMaterialStage::OnTextureUVUpdated);
+				if (UDynamicMaterialModel* MaterialModel = ModelEditorOnlyData->GetMaterialModel())
+				{
+					MaterialModel->GetOnValueUpdateDelegate().AddUObject(this, &UDMMaterialStage::OnValueUpdated);
+					MaterialModel->GetOnTextureUVUpdateDelegate().AddUObject(this, &UDMMaterialStage::OnTextureUVUpdated);
+				}
 			}
 		}
 	}
@@ -267,8 +270,11 @@ void UDMMaterialStage::RemoveDelegates()
 		{
 			if (UDynamicMaterialModelEditorOnlyData* ModelEditorOnlyData = Slot->GetMaterialModelEditorOnlyData())
 			{
-				ModelEditorOnlyData->GetOnValueUpdateDelegate().RemoveAll(this);
-				ModelEditorOnlyData->GetOnTextureUVUpdateDelegate().RemoveAll(this);
+				if (UDynamicMaterialModel* MaterialModel = ModelEditorOnlyData->GetMaterialModel())
+				{
+					MaterialModel->GetOnValueUpdateDelegate().RemoveAll(this);
+					MaterialModel->GetOnTextureUVUpdateDelegate().RemoveAll(this);
+				}
 			}
 		}
 	}
diff --git a/Engine/Plugins/Experimental/DynamicMaterial/Source/DynamicMaterialEditor/Private/Model/DynamicMaterialModelEditorOnlyData.cpp b/Engine/Plugins/Experimental/DynamicMaterial/Source/DynamicMaterialEditor/Private/Model/DynamicMaterialModelEditorOnlyData.cpp
--- a/Engine/Plugins/Experimental/DynamicMaterial/Source/DynamicMaterialEditor/Private/Model/DynamicMaterialModelEditorOnlyData.cpp
+++ b/Engine/Plugins/Experimental/DynamicMaterial/Source/DynamicMaterialEditor/Private/Model/DynamicMaterialModelEditorOnlyData.cpp
@@ -1622,18 +1622,6 @@ void UDynamicMaterialModelEditorOnlyData::OnValueUpdated(UDMMaterialValue* InVal
 {
 	check(InValue);
 
-	if (InValue->GetMaterialModel() != MaterialModel)
-	{
-		return;
-	}
-
-	if (InUpdateType == EDMUpdateType::Value && IsValid(MaterialModel) && IsValid(MaterialModel->DynamicMaterialInstance))
-	{
-		InValue->SetMIDParameter(MaterialModel->DynamicMaterialInstance);
-	}
-
-	OnValueUpdateDelegate.Broadcast(MaterialModel, InValue);
-
 	// Non-exported materials have their values update via settings parameters
 	// Exported materials need to be rebuilt to update the main material.
 	const bool bMaterialInDifferentPackage = MaterialModel->DynamicMaterial ? MaterialModel->DynamicMaterial->GetPackage() != GetPackage() : true;
@@ -1648,18 +1636,6 @@ void UDynamicMaterialModelEditorOnlyData::OnTextureUVUpdated(UDMTextureUV* InTex
 {
 	check(InTextureUV);
 
-	if (InTextureUV->GetMaterialModel() != MaterialModel)
-	{
-		return;
-	}
-
-	if (IsValid(MaterialModel) && IsValid(MaterialModel->DynamicMaterialInstance))
-	{
-		InTextureUV->SetMIDParameters(MaterialModel->DynamicMaterialInstance);
-	}
-
-	OnTextureUVUpdateDelegate.Broadcast(MaterialModel, InTextureUV);
-
 	// Non-exported materials have their values update via settings parameters
 	// Exported materials need to be rebuilt to update the main material.
 	if (MaterialModel->DynamicMaterial && MaterialModel->DynamicMaterial->GetPackage() != GetPackage())
diff --git a/Engine/Plugins/Experimental/DynamicMaterial/Source/DynamicMaterialEditor/Public/Model/DynamicMaterialModelEditorOnlyData.h b/Engine/Plugins/Experimental/DynamicMaterial/Source/DynamicMaterialEditor/Public/Model/DynamicMaterialModelEditorOnlyData.h
--- a/Engine/Plugins/Experimental/DynamicMaterial/Source/DynamicMaterialEditor/Public/Model/DynamicMaterialModelEditorOnlyData.h
+++ b/Engine/Plugins/Experimental/DynamicMaterial/Source/DynamicMaterialEditor/Public/Model/DynamicMaterialModelEditorOnlyData.h
@@ -32,9 +32,7 @@ struct FDMMaterialChannelListPreset;
 
 DECLARE_MULTICAST_DELEGATE_OneParam(FDMOnMaterialBuilt, UDynamicMaterialModel*);
 DECLARE_MULTICAST_DELEGATE_OneParam(FDMOnValueListUpdated, UDynamicMaterialModel*);
-DECLARE_MULTICAST_DELEGATE_TwoParams(FDMOnValueUpdated, UDynamicMaterialModel*, UDMMaterialValue*);
 DECLARE_MULTICAST_DELEGATE_OneParam(FDMOnSlotListUpdated, UDynamicMaterialModel*);
-DECLARE_MULTICAST_DELEGATE_TwoParams(FDMOnTextureUVUpdated, UDynamicMaterialModel*, UDMTextureUV*);
 
 UENUM(BlueprintType)
 enum class EDMState : uint8
@@ -175,9 +173,7 @@ public:
 	
 	FDMOnMaterialBuilt& GetOnMaterialBuiltDelegate() { return OnMaterialBuiltDelegate; }
 	FDMOnValueListUpdated& GetOnValueListUpdateDelegate() { return OnValueListUpdateDelegate; }
-	FDMOnValueUpdated& GetOnValueUpdateDelegate() { return OnValueUpdateDelegate; }
 	FDMOnSlotListUpdated& GetOnSlotListUpdateDelegate() { return OnSlotListUpdateDelegate; }
-	FDMOnTextureUVUpdated& GetOnTextureUVUpdateDelegate() { return OnTextureUVUpdateDelegate; }
 
 	TSharedRef<FDMMaterialBuildState> CreateBuildState(UMaterial* InMaterialToBuild, bool bInDirtyAssets = true) const;
 
@@ -221,8 +217,8 @@ public:
 	//~ Begin IDynamicMaterialModelEditorOnlyDataInterface
 	virtual void PostEditorDuplicate() override;
 	virtual void RequestMaterialBuild() override;
-	virtual void OnValueUpdated(UDMMaterialValue* InValue, EDMUpdateType InUpdateType) override;
 	virtual void OnValueListUpdate() override;
+	virtual void OnValueUpdated(UDMMaterialValue* InValue, EDMUpdateType InUpdateType) override;
 	virtual void OnTextureUVUpdated(UDMTextureUV* InTextureUV) override;
 	virtual void LoadDeprecatedModelData(UDynamicMaterialModel* InMaterialModel) override;
 	virtual TSharedRef<IDMMaterialBuildStateInterface> CreateBuildStateInterface(UMaterial* InMaterialToBuild) const override;
@@ -277,9 +273,7 @@ protected:
 
 	FDMOnMaterialBuilt OnMaterialBuiltDelegate;
 	FDMOnValueListUpdated OnValueListUpdateDelegate;
-	FDMOnValueUpdated OnValueUpdateDelegate;
 	FDMOnSlotListUpdated OnSlotListUpdateDelegate;
-	FDMOnTextureUVUpdated OnTextureUVUpdateDelegate;
 
 	void CreateMaterial();
 
