Material Designer
- The actor slot picker now displays correctly.
- Removed the word "Property" as it is not needed.
- When editing an asset, the asset's name is now displayed and the word "Actor" changes to "Asset".

diff --git a/Engine/Plugins/Experimental/DynamicMaterial/Source/DynamicMaterialEditor/Private/Slate/SDMEditor.cpp b/Engine/Plugins/Experimental/DynamicMaterial/Source/DynamicMaterialEditor/Private/Slate/SDMEditor.cpp
--- a/Engine/Plugins/Experimental/DynamicMaterial/Source/DynamicMaterialEditor/Private/Slate/SDMEditor.cpp
+++ b/Engine/Plugins/Experimental/DynamicMaterial/Source/DynamicMaterialEditor/Private/Slate/SDMEditor.cpp
@@ -538,11 +538,7 @@ void SDMEditor::SetMaterialActor(AActor* InActor)
 		);
 	}
 
-
-	if (Toolbar->GetMaterialActor() != InActor)
-	{
-		Toolbar->SetMaterialActor(InActor);
-	}
+	Toolbar->SetMaterialActor(InActor);
 }
 
 void SDMEditor::SetEmptyLayout()
@@ -657,7 +653,7 @@ TSharedRef<SWidget> SDMEditor::CreateActorMaterialSlotSelector(AActor* InActor)
 		[
 			SNew(STextBlock)
 			.TextStyle(FDynamicMaterialEditorStyle::Get(), "ActorNameBig")
-			.Text(Toolbar.Get(), &SDMToolBar::GetSlotActorDisplayName)
+			.Text(Toolbar.Get(), &SDMToolBar::GetMaterialContainerName)
 		];
 
 	const UObject* CurrentOuter = nullptr;
@@ -1123,23 +1119,36 @@ void SDMEditor::Tick(const FGeometry& AllottedGeometry, const double InCurrentTi
 
 	if (Toolbar.IsValid())
 	{
-		if (Toolbar->GetMaterialModel() != MaterialModel)
+		UDynamicMaterialModel* ToolbarMaterialModel = Toolbar->GetMaterialModel();
+		if (ToolbarMaterialModel != MaterialModel)
 		{
 			Toolbar->SetMaterialModel(MaterialModel);
 		}
 
 		if (ObjectProperty.IsValid())
 		{
-			AActor* Actor = ObjectProperty.GetTypedOuter<AActor>();
+			AActor* MaterialActor = ObjectProperty.GetTypedOuter<AActor>();
 
-			if (Toolbar->GetMaterialActor() != Actor)
+			if (Toolbar->GetMaterialActor() != MaterialActor)
 			{
-				Toolbar->SetMaterialActor(Actor);
+				Toolbar->SetMaterialActor(ObjectProperty.GetTypedOuter<AActor>());
 			}
 		}
 		else if (Toolbar->GetMaterialActor())
 		{
-			Toolbar->SetMaterialActor(nullptr);
+			if (!ToolbarMaterialModel)
+			{
+				Toolbar->SetMaterialActor(nullptr);
+			}
+			else
+			{
+				AActor* Actor = ToolbarMaterialModel->GetTypedOuter<AActor>();
+
+				if (Actor != Toolbar->GetMaterialActor())
+				{
+					Toolbar->SetMaterialActor(Actor);
+				}
+			}
 		}
 	}
 
@@ -1313,10 +1322,7 @@ FReply SDMEditor::OnCreateMaterialButtonClicked(FDMObjectMaterialProperty InMate
 
 		if (AActor* Actor = InMaterialProperty.GetTypedOuter<AActor>())
 		{
-			if (Toolbar->GetMaterialActor() != Actor)
-			{
-				Toolbar->SetMaterialActor(Actor);
-			}
+			Toolbar->SetMaterialActor(Actor);
 		}
 	}
 
diff --git a/Engine/Plugins/Experimental/DynamicMaterial/Source/DynamicMaterialEditor/Private/Slate/SDMMaterialWizard.cpp b/Engine/Plugins/Experimental/DynamicMaterial/Source/DynamicMaterialEditor/Private/Slate/SDMMaterialWizard.cpp
--- a/Engine/Plugins/Experimental/DynamicMaterial/Source/DynamicMaterialEditor/Private/Slate/SDMMaterialWizard.cpp
+++ b/Engine/Plugins/Experimental/DynamicMaterial/Source/DynamicMaterialEditor/Private/Slate/SDMMaterialWizard.cpp
@@ -253,6 +253,7 @@ FReply SDMMaterialWizard::Accept_OnClick()
 		{
 			if (UDynamicMaterialModelEditorOnlyData* EditorOnlyData = UDynamicMaterialModelEditorOnlyData::Get(MaterialModel))
 			{
+				AActor* MaterialActor = Editor->GetMaterialActor();
 				FDMObjectMaterialProperty MaterialProperty = Editor->GetMaterialObjectProperty();
 				Editor->ClearEditor();
 
@@ -268,6 +269,8 @@ FReply SDMMaterialWizard::Accept_OnClick()
 				{
 					Editor->SetMaterialModel(MaterialModel);
 				}
+
+				Editor->SetMaterialActor(MaterialActor);
 			}
 		}
 	}
diff --git a/Engine/Plugins/Experimental/DynamicMaterial/Source/DynamicMaterialEditor/Private/Slate/SDMToolBar.cpp b/Engine/Plugins/Experimental/DynamicMaterial/Source/DynamicMaterialEditor/Private/Slate/SDMToolBar.cpp
--- a/Engine/Plugins/Experimental/DynamicMaterial/Source/DynamicMaterialEditor/Private/Slate/SDMToolBar.cpp
+++ b/Engine/Plugins/Experimental/DynamicMaterial/Source/DynamicMaterialEditor/Private/Slate/SDMToolBar.cpp
@@ -72,7 +72,7 @@ TSharedRef<SWidget> SDMToolBar::CreateToolBarEntries()
 			.Orientation(Orient_Horizontal)
 			.UseAllottedSize(true)
 			.HAlign(HAlign_Left)
-			.InnerSlotPadding(FVector2D(20.0f))
+			.InnerSlotPadding(FVector2D(5.0f))
 			+ SWrapBox::Slot()
 			.HAlign(HAlign_Left)
 			.VAlign(VAlign_Center)
@@ -81,45 +81,28 @@ TSharedRef<SWidget> SDMToolBar::CreateToolBarEntries()
 				SNew(SHorizontalBox)
 				+ SHorizontalBox::Slot()
 				.AutoWidth()
-				.HAlign(HAlign_Left)
-				.VAlign(VAlign_Center)
-				.Padding(0.0f, 0.0f, 10.0f, 0.0f)
+				.Padding(0.f, 0.f, 5.f, 0.f)
 				[
 					SNew(STextBlock)
-					.TextStyle(FDynamicMaterialEditorStyle::Get(), "RegularFont")
-					.Text(LOCTEXT("MaterialDesignerInstanceActorLabel", "Actor"))
+					.TextStyle(FDynamicMaterialEditorStyle::Get(), "ActorName")
+					.Text(this, &SDMToolBar::GetMaterialContainerTypeName)
 				]
 				+ SHorizontalBox::Slot()
 				.AutoWidth()
-				.HAlign(HAlign_Left)
-				.VAlign(VAlign_Center)
+				.Padding(0.f, 0.f, 0.f, 0.f)
 				[
 					SNew(STextBlock)
 					.TextStyle(FDynamicMaterialEditorStyle::Get(), "ActorName")
-					.Text(this, &SDMToolBar::GetSlotActorDisplayName)
+					.Text(this, &SDMToolBar::GetMaterialContainerName)
 				]
 			]
 			+ SWrapBox::Slot()
-			.FillEmptySpace(true)
-			.HAlign(HAlign_Fill)
+			.FillEmptySpace(false)
+			.HAlign(HAlign_Left)
 			.VAlign(VAlign_Center)
+			.Padding(5.0f, 0.0f, 0.0f, 0.0f)
 			[
-				SNew(SHorizontalBox)
-				.Visibility(this, &SDMToolBar::GetSlotsComboBoxWidgetVisibiltiy)
-				+ SHorizontalBox::Slot()
-				.AutoWidth()
-				.HAlign(HAlign_Left)
-				.VAlign(VAlign_Center)
-				.Padding(0.0f, 0.0f, 10.0f, 0.0f)
-				[
-					SNew(STextBlock)
-					.TextStyle(FDynamicMaterialEditorStyle::Get(), "RegularFont")
-					.Text(LOCTEXT("MaterialDesignerInstanceActorSlotLabel", "Property"))
-				]
-				+ SHorizontalBox::Slot()
-				.AutoWidth()
-				.HAlign(HAlign_Left)
-				.VAlign(VAlign_Center)
+				SAssignNew(SlotSelectorContainer, SBox)
 				[
 					CreateSlotsComboBoxWidget()
 				]
@@ -241,12 +224,12 @@ TSharedRef<SWidget> SDMToolBar::CreateToolBarButton(TAttribute<const FSlateBrush
 
 TSharedRef<SWidget> SDMToolBar::CreateSlotsComboBoxWidget()
 {
-	if (!MaterialActorWeak.IsValid() || !MaterialModelWeak.IsValid())
+	if (!MaterialActorWeak.IsValid() || !MaterialModelWeak.IsValid() || ActorMaterialProperties.Num() <= 1)
 	{
 		return SNullWidget::NullWidget;
 	}
 
-	const TSharedPtr<FDMObjectMaterialProperty> InitiallySelectedItem = 
+	const TSharedPtr<FDMObjectMaterialProperty> InitiallySelectedItem =
 		ActorMaterialProperties.IsValidIndex(SelectedMaterialSlotIndex) ? ActorMaterialProperties[SelectedMaterialSlotIndex] : nullptr;
 
 	return 
@@ -322,10 +305,54 @@ EVisibility SDMToolBar::GetSlotsComboBoxWidgetVisibiltiy() const
 	return ActorMaterialProperties.Num() > 1 ? EVisibility::Visible : EVisibility::Collapsed;
 }
 
-FText SDMToolBar::GetSlotActorDisplayName() const
+FText SDMToolBar::GetMaterialContainerName() const
+{
+	if (const AActor* const SlotActor = GetMaterialActor())
+	{
+		return FText::FromString(SlotActor->GetActorLabel());
+	}
+
+	if (UDynamicMaterialModel* MaterialModel = MaterialModelWeak.Get())
+	{
+		if (UDynamicMaterialInstance* MaterialInstance = MaterialModel->GetDynamicMaterialInstance())
+		{
+			if (MaterialInstance->IsAsset())
+			{
+				return FText::FromString(MaterialInstance->GetName());
+			}
+		}
+		else if (MaterialModel->IsAsset())
+		{
+			return FText::FromString(MaterialModel->GetName());
+		}
+	}
+
+	return FText::GetEmpty();
+}
+
+FText SDMToolBar::GetMaterialContainerTypeName() const
 {
-	const AActor* const SlotActor = GetMaterialActor();
-	return IsValid(SlotActor) ? FText::FromString(SlotActor->GetActorLabel()) : FText();
+	if (const AActor* const SlotActor = GetMaterialActor())
+	{
+		return LOCTEXT("Actor", "Actor");
+	}
+
+	if (UDynamicMaterialModel* MaterialModel = MaterialModelWeak.Get())
+	{
+		if (UDynamicMaterialInstance* MaterialInstance = MaterialModel->GetDynamicMaterialInstance())
+		{
+			if (MaterialInstance->IsAsset())
+			{
+				return LOCTEXT("Asset", "Asset");
+			}
+		}
+		else if (MaterialModel->IsAsset())
+		{
+			return LOCTEXT("Asset", "Asset");
+		}
+	}
+
+	return FText::GetEmpty();
 }
 
 void SDMToolBar::SetMaterialProperties(const TArray<TSharedPtr<FDMObjectMaterialProperty>>& InActorMaterialProperties)
@@ -374,25 +401,25 @@ void SDMToolBar::SetMaterialActor(AActor* InActor, const int32 InActiveSlotIndex
 
 	UseButton->SetVisibility(MaterialActorWeak.IsValid() ? EVisibility::Visible : EVisibility::Collapsed);
 
-	if (!IsValid(InActor))
+	if (IsValid(InActor))
 	{
-		return;
-	}
+		TArray<FDMObjectMaterialProperty> ActorProperties = UDMBlueprintFunctionLibrary::GetActorMaterialProperties(InActor);
+		UDynamicMaterialModel* MaterialModel = MaterialModelWeak.Get();
 
-	TArray<FDMObjectMaterialProperty> ActorProperties = UDMBlueprintFunctionLibrary::GetActorMaterialProperties(InActor);
-	UDynamicMaterialModel* MaterialModel = MaterialModelWeak.Get();
-
-	for (int32 MaterialPropertyIdx = 0; MaterialPropertyIdx < ActorProperties.Num(); ++MaterialPropertyIdx)
-	{
-		const FDMObjectMaterialProperty& MaterialProperty = ActorProperties[MaterialPropertyIdx];
+		for (int32 MaterialPropertyIdx = 0; MaterialPropertyIdx < ActorProperties.Num(); ++MaterialPropertyIdx)
+		{
+			const FDMObjectMaterialProperty& MaterialProperty = ActorProperties[MaterialPropertyIdx];
 
-		ActorMaterialProperties.Add(MakeShared<FDMObjectMaterialProperty>(MaterialProperty));
+			ActorMaterialProperties.Add(MakeShared<FDMObjectMaterialProperty>(MaterialProperty));
 
-		if (MaterialProperty.GetMaterialModel() == MaterialModel)
-		{
-			SelectedMaterialSlotIndex = MaterialPropertyIdx;
+			if (MaterialProperty.GetMaterialModel() == MaterialModel)
+			{
+				SelectedMaterialSlotIndex = MaterialPropertyIdx;
+			}
 		}
 	}
+
+	SlotSelectorContainer->SetContent(CreateSlotsComboBoxWidget());
 }
 
 const FSlateBrush* SDMToolBar::GetFollowSelectionBrush() const
diff --git a/Engine/Plugins/Experimental/DynamicMaterial/Source/DynamicMaterialEditor/Private/Slate/SDMToolBar.h b/Engine/Plugins/Experimental/DynamicMaterial/Source/DynamicMaterialEditor/Private/Slate/SDMToolBar.h
--- a/Engine/Plugins/Experimental/DynamicMaterial/Source/DynamicMaterialEditor/Private/Slate/SDMToolBar.h
+++ b/Engine/Plugins/Experimental/DynamicMaterial/Source/DynamicMaterialEditor/Private/Slate/SDMToolBar.h
@@ -7,6 +7,7 @@
 #include "Widgets/SCompoundWidget.h"
 
 class AActor;
+class SBox;
 class SDMEditor;
 class SWidget;
 class UDMMaterialStageExpression;
@@ -44,7 +45,9 @@ public:
 	UDynamicMaterialModel* GetMaterialModel() const { return MaterialModelWeak.Get(); }
 	void SetMaterialModel(UDynamicMaterialModel* InModel);
 
-	FText GetSlotActorDisplayName() const;
+	FText GetMaterialContainerName() const;
+
+	FText GetMaterialContainerTypeName() const;
 
 protected:
 	TWeakPtr<SDMEditor> EditorWeak;
@@ -56,6 +59,7 @@ protected:
 	TArray<TSharedPtr<FDMObjectMaterialProperty>> ActorMaterialProperties;
 	int32 SelectedMaterialSlotIndex;
 
+	TSharedPtr<SBox> SlotSelectorContainer;
 	TSharedPtr<SWidget> BrowseButton;
 	TSharedPtr<SWidget> UseButton;
 
