Material Designer
- Added full-text versions of the _ORM, _ORDp, _RM and _RMA texture set filters.

diff --git a/Engine/Plugins/Experimental/DynamicMaterial/Source/DynamicMaterialTextureSetEditor/Private/DMTextureSetSettings.cpp b/Engine/Plugins/Experimental/DynamicMaterial/Source/DynamicMaterialTextureSetEditor/Private/DMTextureSetSettings.cpp
--- a/Engine/Plugins/Experimental/DynamicMaterial/Source/DynamicMaterialTextureSetEditor/Private/DMTextureSetSettings.cpp
+++ b/Engine/Plugins/Experimental/DynamicMaterial/Source/DynamicMaterialTextureSetEditor/Private/DMTextureSetSettings.cpp
@@ -89,7 +89,7 @@ UDMTextureSetSettings::UDMTextureSetSettings()
 	++FilterCount;
 
 	FDMTextureSetFilter ORM;
-	ORM.FilterStrings = {TEXT("_ORM")};
+	ORM.FilterStrings = {TEXT("_ORM"), TEXT("OcclusionRoughnessMetallic")};
 	ORM.MaterialProperties = {
 		{EDMTextureSetMaterialProperty::Opacity, EDMTextureChannelMask::Red},
 		{EDMTextureSetMaterialProperty::Roughness, EDMTextureChannelMask::Green},
@@ -98,7 +98,7 @@ UDMTextureSetSettings::UDMTextureSetSettings()
 	++FilterCount;
 
 	FDMTextureSetFilter ORDp;
-	ORDp.FilterStrings = {TEXT("_ORDp")};
+	ORDp.FilterStrings = {TEXT("_ORDp"), TEXT("AmbientOcclusionRoughnessDisplacement")};
 	ORDp.MaterialProperties = {
 		{EDMTextureSetMaterialProperty::AmbientOcclusion, EDMTextureChannelMask::Red},
 		{EDMTextureSetMaterialProperty::Roughness, EDMTextureChannelMask::Green},
@@ -107,7 +107,7 @@ UDMTextureSetSettings::UDMTextureSetSettings()
 	++FilterCount;
 
 	FDMTextureSetFilter RM;
-	RM.FilterStrings = {TEXT("_RM")};
+	RM.FilterStrings = {TEXT("_RM"), TEXT("RoughnessMetallic")};
 	RM.MaterialProperties = {
 		{EDMTextureSetMaterialProperty::Roughness, EDMTextureChannelMask::Red},
 		{EDMTextureSetMaterialProperty::Metallic, EDMTextureChannelMask::Green}
@@ -115,7 +115,7 @@ UDMTextureSetSettings::UDMTextureSetSettings()
 	++FilterCount;
 
 	FDMTextureSetFilter RMA;
-	RMA.FilterStrings = {TEXT("_RMA")};
+	RMA.FilterStrings = {TEXT("_RMA"), TEXT("RoughnessMetallicAmbientOcclusion")};
 	RMA.MaterialProperties = {
 		{EDMTextureSetMaterialProperty::Roughness, EDMTextureChannelMask::Red},
 		{EDMTextureSetMaterialProperty::Metallic, EDMTextureChannelMask::Green},
