Material Designer
- Channel overrides are now applied correctly when using a texture set.

diff --git a/Engine/Plugins/Experimental/DynamicMaterial/Source/DynamicMaterialEditor/Private/Model/DynamicMaterialModelEditorOnlyData.cpp b/Engine/Plugins/Experimental/DynamicMaterial/Source/DynamicMaterialEditor/Private/Model/DynamicMaterialModelEditorOnlyData.cpp
--- a/Engine/Plugins/Experimental/DynamicMaterial/Source/DynamicMaterialEditor/Private/Model/DynamicMaterialModelEditorOnlyData.cpp
+++ b/Engine/Plugins/Experimental/DynamicMaterial/Source/DynamicMaterialEditor/Private/Model/DynamicMaterialModelEditorOnlyData.cpp
@@ -846,13 +846,6 @@ bool UDynamicMaterialModelEditorOnlyData::AddTextureSet(UDMTextureSet* InTexture
 			Slot->Modify();
 		}
 
-		const int32 OutputMask = MaterialTexture.Value.TextureChannel == EDMTextureChannelMask::RGBA
-			? FDMMaterialStageConnectorChannel::WHOLE_CHANNEL
-			: (EnumHasAnyFlags(MaterialTexture.Value.TextureChannel, EDMTextureChannelMask::Red) ? FDMMaterialStageConnectorChannel::FIRST_CHANNEL : 0)
-				+ (EnumHasAnyFlags(MaterialTexture.Value.TextureChannel, EDMTextureChannelMask::Green) ? FDMMaterialStageConnectorChannel::FIRST_CHANNEL : 0)
-				+ (EnumHasAnyFlags(MaterialTexture.Value.TextureChannel, EDMTextureChannelMask::Blue) ? FDMMaterialStageConnectorChannel::FIRST_CHANNEL : 0)
-				+ (EnumHasAnyFlags(MaterialTexture.Value.TextureChannel, EDMTextureChannelMask::Alpha) ? FDMMaterialStageConnectorChannel::FIRST_CHANNEL : 0);
-
 		UDMMaterialLayerObject* Layer = nullptr;
 
 		{
@@ -880,7 +873,7 @@ bool UDynamicMaterialModelEditorOnlyData::AddTextureSet(UDMTextureSet* InTexture
 				UDMMaterialStageBlend::InputB,
 				FDMMaterialStageConnectorChannel::WHOLE_CHANNEL,
 				0,
-				OutputMask
+				FDMMaterialStageConnectorChannel::WHOLE_CHANNEL
 			);
 
 			if (!ensure(NewExpression))
@@ -914,6 +907,36 @@ bool UDynamicMaterialModelEditorOnlyData::AddTextureSet(UDMTextureSet* InTexture
 				}
 			}
 
+			if (UDMMaterialStageBlend* Blend = Cast<UDMMaterialStageBlend>(Stage->GetSource()))
+			{
+				EAvaColorChannel AvaColorChannel = EAvaColorChannel::None;
+
+				if (EnumHasAnyFlags(MaterialTexture.Value.TextureChannel, EDMTextureChannelMask::Red))
+				{
+					AvaColorChannel |= EAvaColorChannel::Red;
+				}
+
+				if (EnumHasAnyFlags(MaterialTexture.Value.TextureChannel, EDMTextureChannelMask::Green))
+				{
+					AvaColorChannel |= EAvaColorChannel::Green;
+				}
+
+				if (EnumHasAnyFlags(MaterialTexture.Value.TextureChannel, EDMTextureChannelMask::Blue))
+				{
+					AvaColorChannel |= EAvaColorChannel::Blue;
+				}
+
+				if (EnumHasAnyFlags(MaterialTexture.Value.TextureChannel, EDMTextureChannelMask::Alpha))
+				{
+					AvaColorChannel |= EAvaColorChannel::Alpha;
+				}
+
+				if (AvaColorChannel != EAvaColorChannel::RGBA)
+				{
+					Blend->SetBaseChannelOverride(AvaColorChannel);
+				}				
+			}
+
 			if (bInReplaceSlots)
 			{
 				for (int32 Index = Slot->GetLayers().Num() - 1; Index >= 0; --Index)
diff --git a/Engine/Plugins/Experimental/DynamicMaterial/Source/DynamicMaterialEditor/Public/DMEDefs.h b/Engine/Plugins/Experimental/DynamicMaterial/Source/DynamicMaterialEditor/Public/DMEDefs.h
--- a/Engine/Plugins/Experimental/DynamicMaterial/Source/DynamicMaterialEditor/Public/DMEDefs.h
+++ b/Engine/Plugins/Experimental/DynamicMaterial/Source/DynamicMaterialEditor/Public/DMEDefs.h
@@ -149,5 +149,8 @@ enum class EAvaColorChannel : uint8
 	Red = 1 << 0,
 	Green = 1 << 1,
 	Blue = 1 << 2,
-	Alpha = 1 << 3
+	Alpha = 1 << 3,
+	RGB = Red|Green|Blue UMETA(Hidden),
+	RGBA = Red|Green|Blue|Alpha UMETA(Hidden)
 };
+ENUM_CLASS_FLAGS(EAvaColorChannel)
diff --git a/Engine/Plugins/Experimental/DynamicMaterial/Source/DynamicMaterialTextureSet/Public/DMTextureChannelMask.h b/Engine/Plugins/Experimental/DynamicMaterial/Source/DynamicMaterialTextureSet/Public/DMTextureChannelMask.h
--- a/Engine/Plugins/Experimental/DynamicMaterial/Source/DynamicMaterialTextureSet/Public/DMTextureChannelMask.h
+++ b/Engine/Plugins/Experimental/DynamicMaterial/Source/DynamicMaterialTextureSet/Public/DMTextureChannelMask.h
@@ -8,10 +8,11 @@
 UENUM(BlueprintType)
 enum class EDMTextureChannelMask : uint8
 {
-	Red = 0,
-	Green = 1 << 0,
-	Blue = 1 << 1,
-	Alpha = 1 << 2,
+	None = 0 UMETA(Hidden),
+	Red = 1 << 0,
+	Green = 1 << 1,
+	Blue = 1 << 2,
+	Alpha = 1 << 3,
 	RGB = Red|Green|Blue,
 	RGBA = Red|Green|Blue|Alpha,
 };
