Material Designer
- Changed the toolbar so that both actor and asset info can be shown at once.

diff --git a/Engine/Plugins/Experimental/DynamicMaterial/Source/DynamicMaterialEditor/Private/Slate/SDMEditor.cpp b/Engine/Plugins/Experimental/DynamicMaterial/Source/DynamicMaterialEditor/Private/Slate/SDMEditor.cpp
--- a/Engine/Plugins/Experimental/DynamicMaterial/Source/DynamicMaterialEditor/Private/Slate/SDMEditor.cpp
+++ b/Engine/Plugins/Experimental/DynamicMaterial/Source/DynamicMaterialEditor/Private/Slate/SDMEditor.cpp
@@ -653,7 +653,7 @@ TSharedRef<SWidget> SDMEditor::CreateActorMaterialSlotSelector(AActor* InActor)
 		[
 			SNew(STextBlock)
 			.TextStyle(FDynamicMaterialEditorStyle::Get(), "ActorNameBig")
-			.Text(Toolbar.Get(), &SDMToolBar::GetMaterialContainerName)
+			.Text(Toolbar.Get(), &SDMToolBar::GetActorName)
 		];
 
 	const UObject* CurrentOuter = nullptr;
diff --git a/Engine/Plugins/Experimental/DynamicMaterial/Source/DynamicMaterialEditor/Private/Slate/SDMToolBar.cpp b/Engine/Plugins/Experimental/DynamicMaterial/Source/DynamicMaterialEditor/Private/Slate/SDMToolBar.cpp
--- a/Engine/Plugins/Experimental/DynamicMaterial/Source/DynamicMaterialEditor/Private/Slate/SDMToolBar.cpp
+++ b/Engine/Plugins/Experimental/DynamicMaterial/Source/DynamicMaterialEditor/Private/Slate/SDMToolBar.cpp
@@ -79,13 +79,14 @@ TSharedRef<SWidget> SDMToolBar::CreateToolBarEntries()
 			.Padding(5.0f, 0.0f, 0.0f, 0.0f)
 			[
 				SNew(SHorizontalBox)
+				.Visibility(this, &SDMToolBar::GetActorVisibility)
 				+ SHorizontalBox::Slot()
 				.AutoWidth()
 				.Padding(0.f, 0.f, 5.f, 0.f)
 				[
 					SNew(STextBlock)
 					.TextStyle(FDynamicMaterialEditorStyle::Get(), "ActorName")
-					.Text(this, &SDMToolBar::GetMaterialContainerTypeName)
+					.Text(LOCTEXT("Actor", "Actor"))
 				]
 				+ SHorizontalBox::Slot()
 				.AutoWidth()
@@ -93,7 +94,7 @@ TSharedRef<SWidget> SDMToolBar::CreateToolBarEntries()
 				[
 					SNew(STextBlock)
 					.TextStyle(FDynamicMaterialEditorStyle::Get(), "ActorName")
-					.Text(this, &SDMToolBar::GetMaterialContainerName)
+					.Text(this, &SDMToolBar::GetActorName)
 				]
 			]
 			+ SWrapBox::Slot()
@@ -103,10 +104,36 @@ TSharedRef<SWidget> SDMToolBar::CreateToolBarEntries()
 			.Padding(5.0f, 0.0f, 0.0f, 0.0f)
 			[
 				SAssignNew(SlotSelectorContainer, SBox)
+				.Visibility(this, &SDMToolBar::GetActorVisibility)
 				[
 					CreateSlotsComboBoxWidget()
 				]
 			]
+			+ SWrapBox::Slot()
+			.HAlign(HAlign_Left)
+			.VAlign(VAlign_Center)
+			.Padding(5.0f, 0.0f, 0.0f, 0.0f)
+			[
+				SNew(SHorizontalBox)
+				.Visibility(this, &SDMToolBar::GetAssetVisibility)
+				+ SHorizontalBox::Slot()
+				.AutoWidth()
+				.Padding(0.f, 0.f, 5.f, 0.f)
+				[
+					SNew(STextBlock)
+					.TextStyle(FDynamicMaterialEditorStyle::Get(), "ActorName")
+					.Text(LOCTEXT("Asset", "Asset"))
+				]
+				+ SHorizontalBox::Slot()
+				.AutoWidth()
+				.Padding(0.f, 0.f, 0.f, 0.f)
+				[
+					SNew(STextBlock)
+					.TextStyle(FDynamicMaterialEditorStyle::Get(), "ActorName")
+					.Text(this, &SDMToolBar::GetAssetName)
+					.ToolTipText(this, &SDMToolBar::GetAssetToolTip)
+				]
+			]
 		]
 		
 		+ SHorizontalBox::Slot()
@@ -305,56 +332,6 @@ EVisibility SDMToolBar::GetSlotsComboBoxWidgetVisibiltiy() const
 	return ActorMaterialProperties.Num() > 1 ? EVisibility::Visible : EVisibility::Collapsed;
 }
 
-FText SDMToolBar::GetMaterialContainerName() const
-{
-	if (const AActor* const SlotActor = GetMaterialActor())
-	{
-		return FText::FromString(SlotActor->GetActorLabel());
-	}
-
-	if (UDynamicMaterialModel* MaterialModel = MaterialModelWeak.Get())
-	{
-		if (UDynamicMaterialInstance* MaterialInstance = MaterialModel->GetDynamicMaterialInstance())
-		{
-			if (MaterialInstance->IsAsset())
-			{
-				return FText::FromString(MaterialInstance->GetName());
-			}
-		}
-		else if (MaterialModel->IsAsset())
-		{
-			return FText::FromString(MaterialModel->GetName());
-		}
-	}
-
-	return FText::GetEmpty();
-}
-
-FText SDMToolBar::GetMaterialContainerTypeName() const
-{
-	if (const AActor* const SlotActor = GetMaterialActor())
-	{
-		return LOCTEXT("Actor", "Actor");
-	}
-
-	if (UDynamicMaterialModel* MaterialModel = MaterialModelWeak.Get())
-	{
-		if (UDynamicMaterialInstance* MaterialInstance = MaterialModel->GetDynamicMaterialInstance())
-		{
-			if (MaterialInstance->IsAsset())
-			{
-				return LOCTEXT("Asset", "Asset");
-			}
-		}
-		else if (MaterialModel->IsAsset())
-		{
-			return LOCTEXT("Asset", "Asset");
-		}
-	}
-
-	return FText::GetEmpty();
-}
-
 void SDMToolBar::SetMaterialProperties(const TArray<TSharedPtr<FDMObjectMaterialProperty>>& InActorMaterialProperties)
 {
 	ActorMaterialProperties = InActorMaterialProperties;
@@ -646,4 +623,81 @@ FReply SDMToolBar::OnUseClicked()
 	return FReply::Handled();
 }
 
+FText SDMToolBar::GetActorName() const
+{
+	if (const AActor* const SlotActor = GetMaterialActor())
+	{
+		return FText::FromString(SlotActor->GetActorLabel());
+	}
+
+	return FText::GetEmpty();
+}
+
+EVisibility SDMToolBar::GetActorVisibility() const
+{
+	return GetMaterialActor()
+		? EVisibility::Visible
+		: EVisibility::Collapsed;
+}
+
+EVisibility SDMToolBar::GetAssetVisibility() const
+{
+	if (UDynamicMaterialModel* MaterialModel = MaterialModelWeak.Get())
+	{
+		if (UDynamicMaterialInstance* MaterialInstance = MaterialModel->GetDynamicMaterialInstance())
+		{
+			if (MaterialInstance->IsAsset())
+			{
+				return EVisibility::Visible;
+			}
+		}
+		else if (MaterialModel->IsAsset())
+		{
+			return EVisibility::Visible;
+		}
+	}
+
+	return EVisibility::Collapsed;
+}
+
+FText SDMToolBar::GetAssetName() const
+{
+	if (UDynamicMaterialModel* MaterialModel = MaterialModelWeak.Get())
+	{
+		if (UDynamicMaterialInstance* MaterialInstance = MaterialModel->GetDynamicMaterialInstance())
+		{
+			if (MaterialInstance->IsAsset())
+			{
+				return FText::FromString(MaterialInstance->GetName());
+			}
+		}
+		else if (MaterialModel->IsAsset())
+		{
+			return FText::FromString(MaterialModel->GetName());
+		}
+	}
+
+	return FText::GetEmpty();
+}
+
+FText SDMToolBar::GetAssetToolTip() const
+{
+	if (UDynamicMaterialModel* MaterialModel = MaterialModelWeak.Get())
+	{
+		if (UDynamicMaterialInstance* MaterialInstance = MaterialModel->GetDynamicMaterialInstance())
+		{
+			if (MaterialInstance->IsAsset())
+			{
+				return FText::FromString(MaterialInstance->GetPathName());
+			}
+		}
+		else if (MaterialModel->IsAsset())
+		{
+			return FText::FromString(MaterialModel->GetPathName());
+		}
+	}
+
+	return FText::GetEmpty();
+}
+
 #undef LOCTEXT_NAMESPACE
diff --git a/Engine/Plugins/Experimental/DynamicMaterial/Source/DynamicMaterialEditor/Private/Slate/SDMToolBar.h b/Engine/Plugins/Experimental/DynamicMaterial/Source/DynamicMaterialEditor/Private/Slate/SDMToolBar.h
--- a/Engine/Plugins/Experimental/DynamicMaterial/Source/DynamicMaterialEditor/Private/Slate/SDMToolBar.h
+++ b/Engine/Plugins/Experimental/DynamicMaterial/Source/DynamicMaterialEditor/Private/Slate/SDMToolBar.h
@@ -45,9 +45,7 @@ public:
 	UDynamicMaterialModel* GetMaterialModel() const { return MaterialModelWeak.Get(); }
 	void SetMaterialModel(UDynamicMaterialModel* InModel);
 
-	FText GetMaterialContainerName() const;
-
-	FText GetMaterialContainerTypeName() const;
+	FText GetActorName() const;
 
 protected:
 	TWeakPtr<SDMEditor> EditorWeak;
@@ -91,4 +89,12 @@ protected:
 	FReply OnBrowseClicked();
 
 	FReply OnUseClicked();
+
+	EVisibility GetActorVisibility() const;
+
+	EVisibility GetAssetVisibility() const;
+
+	FText GetAssetName() const;
+
+	FText GetAssetToolTip() const;
 };
