Material Designer
- Critical fix to UVs.

diff --git a/Engine/Plugins/Experimental/DynamicMaterial/Source/DynamicMaterial/Private/Components/DMTextureUV.cpp b/Engine/Plugins/Experimental/DynamicMaterial/Source/DynamicMaterial/Private/Components/DMTextureUV.cpp
--- a/Engine/Plugins/Experimental/DynamicMaterial/Source/DynamicMaterial/Private/Components/DMTextureUV.cpp
+++ b/Engine/Plugins/Experimental/DynamicMaterial/Source/DynamicMaterial/Private/Components/DMTextureUV.cpp
@@ -399,15 +399,10 @@ void UDMTextureUV::Update(EDMUpdateType InUpdateType)
 
 	Super::Update(InUpdateType);
 
-#if WITH_EDITOR
-	UDynamicMaterialModel* MaterialModel = GetMaterialModel();
-	check(MaterialModel);
-
-	if (IDynamicMaterialModelEditorOnlyDataInterface* ModelEditorOnlyData = MaterialModel->GetEditorOnlyData())
+	if (UDynamicMaterialModel* MaterialModel = GetMaterialModel())
 	{
-		ModelEditorOnlyData->OnTextureUVUpdated(this);
+		MaterialModel->OnTextureUVUpdated(this);
 	}
-#endif
 }
 
 UDynamicMaterialModel* UDMTextureUV::GetMaterialModel() const
