Material Designer
- The toolbar now has a save button.

diff --git a/Engine/Plugins/Experimental/DynamicMaterial/Source/DynamicMaterialEditor/Private/Slate/SDMToolBar.cpp b/Engine/Plugins/Experimental/DynamicMaterial/Source/DynamicMaterialEditor/Private/Slate/SDMToolBar.cpp
--- a/Engine/Plugins/Experimental/DynamicMaterial/Source/DynamicMaterialEditor/Private/Slate/SDMToolBar.cpp
+++ b/Engine/Plugins/Experimental/DynamicMaterial/Source/DynamicMaterialEditor/Private/Slate/SDMToolBar.cpp
@@ -10,13 +10,15 @@
 #include "EngineAnalytics.h"
 #include "GameFramework/Actor.h"
 #include "IContentBrowserSingleton.h"
-#include "SDMEditor.h"
-#include "Selection.h"
 #include "Material/DynamicMaterialInstance.h"
 #include "Model/DynamicMaterialModel.h"
+#include "PackageTools.h"
+#include "SDMEditor.h"
+#include "Selection.h"
 #include "SlateOptMacros.h"
 #include "Styling/AppStyle.h"
 #include "Styling/StyleColors.h"
+#include "UObject/Package.h"
 #include "Utils/DMBlueprintFunctionLibrary.h"
 #include "Widgets/Images/SImage.h"
 #include "Widgets/Input/SButton.h"
@@ -174,6 +176,25 @@ TSharedRef<SWidget> SDMToolBar::CreateToolBarEntries()
 			]
 		]
 		
+		+ SHorizontalBox::Slot()
+		.AutoWidth()
+		.HAlign(HAlign_Right)
+		.VAlign(VAlign_Top)
+		.Padding(5.0f, 0.0f, 0.0f, 0.0f)
+		[
+			SAssignNew(SaveButton, SButton)
+			.IsEnabled(this, &SDMToolBar::CanSave)
+			.ContentPadding(GetLargeIconToolBarButtonContentPadding())
+			.ButtonStyle(FDynamicMaterialEditorStyle::Get(), "HoverHintOnly")
+			.ToolTipText(LOCTEXT("MaterialDesignerBrowseTooltip", "Save the Material Designer asset\n\nCaution: If this asset lives inside an actor, the actor/level will be saved."))
+			.OnClicked(this, &SDMToolBar::OnSaveClicked)
+			[
+				SNew(SImage)
+				.Image(this, &SDMToolBar::GetSaveIcon)
+				.DesiredSizeOverride(GetLargeIconToolBarButtonSize())
+			]
+		]
+		
 		+ SHorizontalBox::Slot()
 		.AutoWidth()
 		.HAlign(HAlign_Right)
@@ -623,6 +644,23 @@ FReply SDMToolBar::OnUseClicked()
 	return FReply::Handled();
 }
 
+UPackage* SDMToolBar::GetSaveablePackage(UObject* InObject)
+{
+	if (!IsValid(InObject))
+	{
+		return nullptr;
+	}
+
+	UPackage* Package = InObject->GetPackage();
+
+	if (!Package || Package->HasAllFlags(RF_Transient))
+	{
+		return nullptr;
+	}
+
+	return Package;
+}
+
 FText SDMToolBar::GetActorName() const
 {
 	if (const AActor* const SlotActor = GetMaterialActor())
@@ -700,4 +738,37 @@ FText SDMToolBar::GetAssetToolTip() const
 	return FText::GetEmpty();
 }
 
+bool SDMToolBar::CanSave() const
+{
+	return !!GetSaveablePackage(GetMaterialModel());
+}
+
+const FSlateBrush* SDMToolBar::GetSaveIcon() const
+{
+	if (UPackage* Package = GetSaveablePackage(GetMaterialModel()))
+	{
+		if (Package->IsDirty())
+		{
+			return FAppStyle::Get().GetBrush("Icons.SaveModified");
+		}
+	}
+
+	return FAppStyle::Get().GetBrush("Icons.Save");
+}
+
+FReply SDMToolBar::OnSaveClicked()
+{
+	if (UDynamicMaterialModel* MaterialModel = GetMaterialModel())
+	{
+		if (UPackage* Package = GetSaveablePackage(MaterialModel))
+		{
+			TArray<UObject*> TexturesToSaveArray;
+			TexturesToSaveArray.Add(MaterialModel);
+			UPackageTools::SavePackagesForObjects(TexturesToSaveArray);
+		}
+	}
+
+	return FReply::Handled();
+}
+
 #undef LOCTEXT_NAMESPACE
diff --git a/Engine/Plugins/Experimental/DynamicMaterial/Source/DynamicMaterialEditor/Private/Slate/SDMToolBar.h b/Engine/Plugins/Experimental/DynamicMaterial/Source/DynamicMaterialEditor/Private/Slate/SDMToolBar.h
--- a/Engine/Plugins/Experimental/DynamicMaterial/Source/DynamicMaterialEditor/Private/Slate/SDMToolBar.h
+++ b/Engine/Plugins/Experimental/DynamicMaterial/Source/DynamicMaterialEditor/Private/Slate/SDMToolBar.h
@@ -12,6 +12,8 @@ class SDMEditor;
 class SWidget;
 class UDMMaterialStageExpression;
 class UDynamicMaterialModel;
+class UObject;
+class UPackage;
 enum class EDMExpressionMenu : uint8;
 struct FDMObjectMaterialProperty;
 
@@ -48,6 +50,8 @@ public:
 	FText GetActorName() const;
 
 protected:
+	static UPackage* GetSaveablePackage(UObject* InObject);
+
 	TWeakPtr<SDMEditor> EditorWeak;
 	TWeakObjectPtr<AActor> MaterialActorWeak;
 	TWeakObjectPtr<UDynamicMaterialModel> MaterialModelWeak;
@@ -58,6 +62,7 @@ protected:
 	int32 SelectedMaterialSlotIndex;
 
 	TSharedPtr<SBox> SlotSelectorContainer;
+	TSharedPtr<SWidget> SaveButton;
 	TSharedPtr<SWidget> BrowseButton;
 	TSharedPtr<SWidget> UseButton;
 
@@ -97,4 +102,10 @@ protected:
 	FText GetAssetName() const;
 
 	FText GetAssetToolTip() const;
+
+	bool CanSave() const;
+
+	const FSlateBrush* GetSaveIcon() const;
+
+	FReply OnSaveClicked();
 };
