Material Designer
- You can now drag layers to the bottom of the layer list.

diff --git a/Engine/Plugins/Experimental/DynamicMaterial/Source/DynamicMaterialEditor/Private/Components/DMMaterialLayer.cpp b/Engine/Plugins/Experimental/DynamicMaterial/Source/DynamicMaterialEditor/Private/Components/DMMaterialLayer.cpp
--- a/Engine/Plugins/Experimental/DynamicMaterial/Source/DynamicMaterialEditor/Private/Components/DMMaterialLayer.cpp
+++ b/Engine/Plugins/Experimental/DynamicMaterial/Source/DynamicMaterialEditor/Private/Components/DMMaterialLayer.cpp
@@ -650,24 +650,24 @@ bool UDMMaterialLayerObject::CanMoveLayerBelow(UDMMaterialLayerObject* InLayer)
 		return false;
 	}
 
-	const int32 ThisIndex = FindIndex();
+	const UDMMaterialSlot* const ThisSlot = GetSlot();
 
-	// Already top level - or invalid.
-	if (ThisIndex == 0 || ThisIndex == INDEX_NONE)
+	if (!IsValid(ThisSlot))
 	{
 		return false;
 	}
 
-	const int32 DraggedLayerIndex = InLayer->FindIndex();
+	const int32 ThisIndex = FindIndex();
 
-	if (DraggedLayerIndex == INDEX_NONE)
+	// Already bottom level - or invalid.
+	if (ThisIndex == (ThisSlot->GetLayers().Num() - 1) || ThisIndex == INDEX_NONE)
 	{
 		return false;
 	}
 
-	const UDMMaterialSlot* const ThisSlot = GetSlot();
+	const int32 DraggedLayerIndex = InLayer->FindIndex();
 
-	if (!IsValid(ThisSlot))
+	if (DraggedLayerIndex == INDEX_NONE)
 	{
 		return false;
 	}
