Material Designer
- Inputs now get their proper name on the details panel.
- Added the "Vector Noise" world space noise pattern generator

diff --git a/Engine/Plugins/Experimental/DynamicMaterial/Source/DynamicMaterialEditor/Private/Components/MaterialStageExpressions/DMMSETextureSampleBase.cpp b/Engine/Plugins/Experimental/DynamicMaterial/Source/DynamicMaterialEditor/Private/Components/MaterialStageExpressions/DMMSETextureSampleBase.cpp
--- a/Engine/Plugins/Experimental/DynamicMaterial/Source/DynamicMaterialEditor/Private/Components/MaterialStageExpressions/DMMSETextureSampleBase.cpp
+++ b/Engine/Plugins/Experimental/DynamicMaterial/Source/DynamicMaterialEditor/Private/Components/MaterialStageExpressions/DMMSETextureSampleBase.cpp
@@ -14,7 +14,6 @@
 #include "Components/MaterialStageInputs/DMMSIValue.h"
 #include "Components/MaterialValues/DMMaterialValueTexture.h"
 #include "CoreGlobals.h"
-#include "Engine/TextureCube.h"
 #include "Materials/MaterialExpressionTextureSample.h"
 #include "Model/DMMaterialBuildState.h"
 #include "Model/DynamicMaterialModelEditorOnlyData.h"
diff --git a/Engine/Plugins/Experimental/DynamicMaterial/Source/DynamicMaterialEditor/Private/Components/MaterialStageExpressions/DMMSEWorldPositionNoise.cpp b/Engine/Plugins/Experimental/DynamicMaterial/Source/DynamicMaterialEditor/Private/Components/MaterialStageExpressions/DMMSEWorldPositionNoise.cpp
new file mode 100644
--- /dev/null
+++ b/Engine/Plugins/Experimental/DynamicMaterial/Source/DynamicMaterialEditor/Private/Components/MaterialStageExpressions/DMMSEWorldPositionNoise.cpp
@@ -0,0 +1,198 @@
+// Copyright Epic Games, Inc. All Rights Reserved.
+
+#include "Components/MaterialStageExpressions/DMMSEWorldPositionNoise.h"
+
+#include "Components/DMMaterialLayer.h"
+#include "Components/DMMaterialProperty.h"
+#include "Components/DMMaterialStageBlend.h"
+#include "Components/DMMaterialSubStage.h"
+#include "Components/MaterialStageExpressions/DMMSETextureSample.h"
+#include "CoreGlobals.h"
+#include "Components/MaterialStageInputs/DMMSIValue.h"
+#include "Components/MaterialValues/DMMaterialValueFloat1.h"
+#include "Materials/MaterialExpressionDivide.h"
+#include "Materials/MaterialExpressionVectorNoise.h"
+#include "Materials/MaterialExpressionWorldPosition.h"
+#include "Model/DMMaterialBuildState.h"
+#include "Model/DMMaterialBuildUtils.h"
+
+#define LOCTEXT_NAMESPACE "DMMaterialStageExpressionTextureSample"
+
+UDMMaterialStageExpressionWorldPositionNoise::UDMMaterialStageExpressionWorldPositionNoise()
+	: UDMMaterialStageExpression(
+		LOCTEXT("UMaterialExpressionVectorNoise", "Noise"),
+		UMaterialExpressionVectorNoise::StaticClass()
+	)
+	, ShaderOffset(EWorldPositionIncludedOffsets::WPT_Default)
+	, NoiseFunction(EVectorNoiseFunction::VNF_VectorALU)
+	, Quality(1)
+	, bTiling(false)
+	, TileSize(300)
+{
+	bInputRequired = true;
+	bAllowNestedInputs = true;
+
+	InputConnectors.Add({1, LOCTEXT("Scale", "Scale"), EDMValueType::VT_Float1});
+
+	OutputConnectors.Add({0, LOCTEXT("ColorRGB", "Color (RGB)"), EDMValueType::VT_Float3_RGB});
+
+	EditableProperties.Add(GET_MEMBER_NAME_CHECKED(UDMMaterialStageExpressionWorldPositionNoise, ShaderOffset));
+	EditableProperties.Add(GET_MEMBER_NAME_CHECKED(UDMMaterialStageExpressionWorldPositionNoise, NoiseFunction));
+	EditableProperties.Add(GET_MEMBER_NAME_CHECKED(UDMMaterialStageExpressionWorldPositionNoise, Quality));
+	EditableProperties.Add(GET_MEMBER_NAME_CHECKED(UDMMaterialStageExpressionWorldPositionNoise, bTiling));
+	EditableProperties.Add(GET_MEMBER_NAME_CHECKED(UDMMaterialStageExpressionWorldPositionNoise, TileSize));
+}
+
+void UDMMaterialStageExpressionWorldPositionNoise::GenerateExpressions(const TSharedRef<FDMMaterialBuildState>& InBuildState) const
+{
+	if (!IsComponentValid() || !IsComponentAdded())
+	{
+		return;
+	}
+
+	check(MaterialExpressionClass.Get());
+
+	if (InBuildState->HasStageSource(this))
+	{
+		return;
+	}
+
+	UMaterialExpressionWorldPosition* WorldPosition = InBuildState->GetBuildUtils().CreateExpression<UMaterialExpressionWorldPosition>(UE_DM_NodeComment_Default);
+	WorldPosition->WorldPositionShaderOffset = ShaderOffset;
+
+	UMaterialExpressionDivide* Divide = InBuildState->GetBuildUtils().CreateExpression<UMaterialExpressionDivide>(UE_DM_NodeComment_Default);
+
+	WorldPosition->ConnectExpression(&Divide->A, 0);
+
+	UMaterialExpressionVectorNoise* VectorNoise = InBuildState->GetBuildUtils().CreateExpression<UMaterialExpressionVectorNoise>(UE_DM_NodeComment_Default);
+	VectorNoise->NoiseFunction = NoiseFunction;
+	VectorNoise->Quality = FMath::Clamp(Quality, 1, 4);
+	VectorNoise->bTiling = bTiling ? 1 : 0;
+	VectorNoise->TileSize = TileSize;
+
+	Divide->ConnectExpression(&VectorNoise->Position, 0);
+
+	InBuildState->AddStageSourceExpressions(this, {WorldPosition, Divide, VectorNoise});
+}
+
+void UDMMaterialStageExpressionWorldPositionNoise::SetShaderOffset(EWorldPositionIncludedOffsets InShaderOffset)
+{
+	if (ShaderOffset == InShaderOffset)
+	{
+		return;
+	}
+
+	ShaderOffset = InShaderOffset;
+
+	Update(EDMUpdateType::Structure);
+}
+
+void UDMMaterialStageExpressionWorldPositionNoise::SetNoiseFunction(EVectorNoiseFunction InNoiseFunction)
+{
+	if (NoiseFunction == InNoiseFunction)
+	{
+		return;
+	}
+
+	NoiseFunction = InNoiseFunction;
+
+	Update(EDMUpdateType::Structure);
+}
+
+void UDMMaterialStageExpressionWorldPositionNoise::SetQuality(int32 InQuality)
+{
+	if (Quality == InQuality)
+	{
+		return;
+	}
+
+	Quality = InQuality;
+
+	Update(EDMUpdateType::Structure);
+}
+
+void UDMMaterialStageExpressionWorldPositionNoise::SetTiling(bool bInTiling)
+{
+	if (bTiling == bInTiling)
+	{
+		return;
+	}
+
+	bTiling = bInTiling;
+
+	Update(EDMUpdateType::Structure);
+}
+
+void UDMMaterialStageExpressionWorldPositionNoise::SetTileSize(int32 InTileSize)
+{
+	if (TileSize == InTileSize)
+	{
+		return;
+	}
+
+	TileSize = InTileSize;
+
+	Update(EDMUpdateType::Structure);
+}
+
+void UDMMaterialStageExpressionWorldPositionNoise::AddDefaultInput(int32 InInputIndex) const
+{
+	Super::AddDefaultInput(InInputIndex);
+
+	UDMMaterialStage* Stage = GetStage();
+	check(Stage);
+
+	switch (InInputIndex)
+	{
+		case 0:
+		{
+			UDMMaterialStageInputValue* InputValue = Cast<UDMMaterialStageInputValue>(Stage->GetInputs().Last());
+			check(InputValue);
+
+			UDMMaterialValueFloat1* Float1Value = Cast<UDMMaterialValueFloat1>(InputValue->GetValue());
+			check(Float1Value);
+
+			Float1Value->SetDefaultValue(10.f);
+			Float1Value->ApplyDefaultValue();
+			break;
+		}
+
+		default:
+			// Do nothing
+			break;
+	}
+}
+
+UMaterialExpression* UDMMaterialStageExpressionWorldPositionNoise::GetExpressionForInput(const TArray<UMaterialExpression*>& StageSourceExpressions, int32 InputIdx)
+{
+	if (InputIdx == 1 && StageSourceExpressions.IsValidIndex(1))
+	{
+		return StageSourceExpressions[1];
+	}
+
+	return Super::GetExpressionForInput(StageSourceExpressions, InputIdx);
+}
+
+void UDMMaterialStageExpressionWorldPositionNoise::PostEditChangeProperty(FPropertyChangedEvent& InPropertyChangedEvent)
+{
+	Super::PostEditChangeProperty(InPropertyChangedEvent);
+
+	static const FName ShaderOffsetName = GET_MEMBER_NAME_CHECKED(UDMMaterialStageExpressionWorldPositionNoise, ShaderOffset);
+	static const FName NoiseFunctionName = GET_MEMBER_NAME_CHECKED(UDMMaterialStageExpressionWorldPositionNoise, NoiseFunction);
+	static const FName QualityName = GET_MEMBER_NAME_CHECKED(UDMMaterialStageExpressionWorldPositionNoise, Quality);
+	static const FName TilingName = GET_MEMBER_NAME_CHECKED(UDMMaterialStageExpressionWorldPositionNoise, bTiling);
+	static const FName TileSizeName = GET_MEMBER_NAME_CHECKED(UDMMaterialStageExpressionWorldPositionNoise, TileSize);
+
+	const FName PropertyName = InPropertyChangedEvent.GetMemberPropertyName();
+
+	if (PropertyName == ShaderOffsetName
+		|| PropertyName == NoiseFunctionName
+		|| PropertyName == QualityName
+		|| PropertyName == TilingName
+		|| PropertyName == TileSizeName)
+	{
+		Update(EDMUpdateType::Structure);
+	}
+}
+
+#undef LOCTEXT_NAMESPACE
diff --git a/Engine/Plugins/Experimental/DynamicMaterial/Source/DynamicMaterialEditor/Private/Components/MaterialStageExpressions/DMMSEWorldPositionNoise.h b/Engine/Plugins/Experimental/DynamicMaterial/Source/DynamicMaterialEditor/Private/Components/MaterialStageExpressions/DMMSEWorldPositionNoise.h
new file mode 100644
--- /dev/null
+++ b/Engine/Plugins/Experimental/DynamicMaterial/Source/DynamicMaterialEditor/Private/Components/MaterialStageExpressions/DMMSEWorldPositionNoise.h
@@ -0,0 +1,82 @@
+// Copyright Epic Games, Inc. All Rights Reserved.
+
+#pragma once
+
+#include "Components/DMMaterialStageExpression.h"
+#include "DMMSEWorldPositionNoise.generated.h"
+
+enum EVectorNoiseFunction : int;
+enum EWorldPositionIncludedOffsets : int;
+
+UCLASS(BlueprintType, Blueprintable, ClassGroup = "Material Designer", meta = (DisplayName = "Noise"))
+class UDMMaterialStageExpressionWorldPositionNoise : public UDMMaterialStageExpression
+{
+	GENERATED_BODY()
+
+public:
+	UDMMaterialStageExpressionWorldPositionNoise();
+
+	//~ Begin UDMMaterialStageExpression
+	virtual void GenerateExpressions(const TSharedRef<FDMMaterialBuildState>& InBuildState) const override;
+	//~ End UDMMaterialStageExpression
+
+	UFUNCTION(BlueprintPure, Category = "Material Designer")
+	virtual TEnumAsByte<EWorldPositionIncludedOffsets> GetShaderOffset() const { return ShaderOffset; }
+
+	UFUNCTION(BlueprintCallable, Category = "Material Designer")
+	virtual void SetShaderOffset(EWorldPositionIncludedOffsets InShaderOffset);
+
+	UFUNCTION(BlueprintPure, Category = "Material Designer")
+	virtual TEnumAsByte<EVectorNoiseFunction> GetNoiseFunction() const { return NoiseFunction; }
+
+	UFUNCTION(BlueprintCallable, Category = "Material Designer")
+	virtual void SetNoiseFunction(EVectorNoiseFunction InNoiseFunction);
+
+	UFUNCTION(BlueprintPure, Category = "Material Designer")
+	virtual int32 GetQuality() const { return Quality; }
+
+	UFUNCTION(BlueprintCallable, Category = "Material Designer")
+	virtual void SetQuality(int32 InQuality);
+
+	UFUNCTION(BlueprintPure, Category = "Material Designer")
+	virtual bool GetTiling() const { return bTiling; }
+
+	UFUNCTION(BlueprintCallable, Category = "Material Designer")
+	virtual void SetTiling(bool bInTiling);
+
+	UFUNCTION(BlueprintPure, Category = "Material Designer")
+	virtual int32 GetTileSize() const { return TileSize; }
+
+	UFUNCTION(BlueprintCallable, Category = "Material Designer")
+	virtual void SetTileSize(int32 InTileSize);
+
+	//~ Begin UDMMaterialStageThroughput
+	virtual void AddDefaultInput(int32 InInputIndex) const override;
+	virtual UMaterialExpression* GetExpressionForInput(const TArray<UMaterialExpression*>& StageSourceExpressions, int32 InputIdx) override;
+	//~ End UDMMaterialStageThroughput
+
+	//~ Begin UObject
+	virtual void PostEditChangeProperty(FPropertyChangedEvent& InPropertyChangedEvent) override;
+	//~ End UObject
+
+protected:
+	UPROPERTY(EditInstanceOnly, BlueprintReadWrite, Setter = SetShaderOffset, BlueprintSetter = SetShaderOffset, Category = "Material Designer",
+		meta = (NotKeyframeable, AllowPrivateAccess = "true"))
+	TEnumAsByte<EWorldPositionIncludedOffsets> ShaderOffset;
+
+	UPROPERTY(EditInstanceOnly, BlueprintReadWrite, Setter = SetNoiseFunction, BlueprintSetter = SetNoiseFunction, Category = "Material Designer",
+		meta = (NotKeyframeable, AllowPrivateAccess = "true"))
+	TEnumAsByte<EVectorNoiseFunction> NoiseFunction;
+
+	UPROPERTY(EditInstanceOnly, BlueprintReadWrite, Setter = SetQuality, BlueprintSetter = SetQuality, Category = "Material Designer",
+		meta = (NotKeyframeable, AllowPrivateAccess = "true", UIMin = "1", UIMax = "4", ClampMin = "1", ClampMax = "4"))
+	int32 Quality;
+
+	UPROPERTY(EditInstanceOnly, BlueprintReadWrite, Setter = SetTiling, BlueprintSetter = SetTiling, Category = "Material Designer",
+		meta = (NotKeyframeable, AllowPrivateAccess = "true"))
+	bool bTiling;
+
+	UPROPERTY(EditInstanceOnly, BlueprintReadWrite, Setter = SetTileSize, BlueprintSetter = SetTileSize, Category = "Material Designer",
+		meta = (NotKeyframeable, AllowPrivateAccess = "true", EditCondition="bTiling", EditConditionHides))
+	int32 TileSize;
+};
diff --git a/Engine/Plugins/Experimental/DynamicMaterial/Source/DynamicMaterialEditor/Private/Menus/DMMaterialSlotLayerMenus.cpp b/Engine/Plugins/Experimental/DynamicMaterial/Source/DynamicMaterialEditor/Private/Menus/DMMaterialSlotLayerMenus.cpp
--- a/Engine/Plugins/Experimental/DynamicMaterial/Source/DynamicMaterialEditor/Private/Menus/DMMaterialSlotLayerMenus.cpp
+++ b/Engine/Plugins/Experimental/DynamicMaterial/Source/DynamicMaterialEditor/Private/Menus/DMMaterialSlotLayerMenus.cpp
@@ -9,6 +9,7 @@
 #include "Components/DMMaterialValue.h"
 #include "Components/MaterialStageExpressions/DMMSETextureSample.h"
 #include "Components/MaterialStageExpressions/DMMSETextureSampleEdgeColor.h"
+#include "Components/MaterialStageExpressions/DMMSEWorldPositionNoise.h"
 #include "Components/RenderTargetRenderers/DMRenderTargetTextRenderer.h"
 #include "Components/RenderTargetRenderers/DMRenderTargetUMGWidgetRenderer.h"
 #include "DMDefs.h"
@@ -577,6 +578,19 @@ void FDMMaterialSlotLayerMenus::AddAddLayerSection(UToolMenu* InMenu)
 		))
 	);
 
+	NewSection.AddMenuEntry(
+		NAME_None,
+		LOCTEXT("AddNoise", "Noise"),
+		LOCTEXT("AddNoiseTooltip", "Add a new Material Layer with a noise pattern."),
+		FSlateIcon(),
+		FUIAction(FExecuteAction::CreateSP(
+			EditorWidget->GetActiveSlotWidget().Get(),
+			&SDMSlot::AddNewLayer_Expression,
+			TSubclassOf<UDMMaterialStageExpression>(UDMMaterialStageExpressionWorldPositionNoise::StaticClass()),
+			EDMMaterialLayerStage::All
+		))
+	);
+
 	if constexpr (UE::DynamicMaterialEditor::bAdvancedSlotsEnabled)
 	{
 		if (bHasValidSlot)
diff --git a/Engine/Plugins/Experimental/DynamicMaterial/Source/DynamicMaterialEditor/Private/Menus/DMMaterialStageInputMenus.cpp b/Engine/Plugins/Experimental/DynamicMaterial/Source/DynamicMaterialEditor/Private/Menus/DMMaterialStageInputMenus.cpp
--- a/Engine/Plugins/Experimental/DynamicMaterial/Source/DynamicMaterialEditor/Private/Menus/DMMaterialStageInputMenus.cpp
+++ b/Engine/Plugins/Experimental/DynamicMaterial/Source/DynamicMaterialEditor/Private/Menus/DMMaterialStageInputMenus.cpp
@@ -11,6 +11,7 @@
 #include "Components/DMMaterialValue.h"
 #include "Components/MaterialStageExpressions/DMMSETextureSample.h"
 #include "Components/MaterialStageExpressions/DMMSETextureSampleEdgeColor.h"
+#include "Components/MaterialStageExpressions/DMMSEWorldPositionNoise.h"
 #include "Components/MaterialStageInputs/DMMSIExpression.h"
 #include "Components/MaterialStageInputs/DMMSIGradient.h"
 #include "Components/MaterialStageInputs/DMMSISlot.h"
@@ -215,6 +216,11 @@ EndStageSearch:
 			GenerateChangeInputMenu_Expression(MenuBuilder, TextureSampleEdgeColor, InThroughput, InInputIndex, InInputChannel);
 		}
 
+		if (UDMMaterialStageExpression* WorldPositionNoise = Cast<UDMMaterialStageExpression>(UDMMaterialStageExpressionWorldPositionNoise::StaticClass()->GetDefaultObject(true)))
+		{
+			GenerateChangeInputMenu_Expression(MenuBuilder, WorldPositionNoise, InThroughput, InInputIndex, InInputChannel);
+		}
+
 		if constexpr (UE::DynamicMaterialEditor::bAdvancedSlotsEnabled)
 		{
 			if (bHasValidSlot)
diff --git a/Engine/Plugins/Experimental/DynamicMaterial/Source/DynamicMaterialEditor/Private/Menus/DMMaterialStageSourceMenus.cpp b/Engine/Plugins/Experimental/DynamicMaterial/Source/DynamicMaterialEditor/Private/Menus/DMMaterialStageSourceMenus.cpp
--- a/Engine/Plugins/Experimental/DynamicMaterial/Source/DynamicMaterialEditor/Private/Menus/DMMaterialStageSourceMenus.cpp
+++ b/Engine/Plugins/Experimental/DynamicMaterial/Source/DynamicMaterialEditor/Private/Menus/DMMaterialStageSourceMenus.cpp
@@ -17,6 +17,7 @@
 #include "Components/MaterialStageExpressions/DMMSESceneTexture.h"
 #include "Components/MaterialStageExpressions/DMMSETextureSample.h"
 #include "Components/MaterialStageExpressions/DMMSETextureSampleEdgeColor.h"
+#include "Components/MaterialStageExpressions/DMMSEWorldPositionNoise.h"
 #include "Components/MaterialStageInputs/DMMSIExpression.h"
 #include "Components/MaterialStageInputs/DMMSIFunction.h"
 #include "Components/MaterialStageInputs/DMMSIGradient.h"
@@ -776,6 +777,66 @@ namespace UE::DynamicMaterialEditor::Private
 		}
 	}
 
+	void ChangeSourceToNoiseFromContext(UDMMenuContext* InMenuContext)
+	{
+		if (!IsValid(InMenuContext))
+		{
+			return;
+		}
+
+		UDMMaterialStage* const Stage = InMenuContext->GetStage();
+
+		if (!Stage)
+		{
+			return;
+		}
+
+		UDMMaterialStageSource* const StageSource = Stage->GetSource();
+
+		if (!StageSource)
+		{
+			return;
+		}
+
+		if (StageSource->IsA<UDMMaterialStageBlend>())
+		{
+			FScopedTransaction Transaction(LOCTEXT("SetStageInputBase", "Set Material Designer Base Source"));
+			Stage->Modify();
+
+			UDMMaterialStageInputExpression::ChangeStageInput_Expression(
+				Stage,
+				UDMMaterialStageExpressionWorldPositionNoise::StaticClass(),
+				UDMMaterialStageBlend::InputB,
+				FDMMaterialStageConnectorChannel::WHOLE_CHANNEL,
+				0,
+				FDMMaterialStageConnectorChannel::THREE_CHANNELS
+			);
+		}
+		else if (StageSource->IsA<UDMMaterialStageThroughputLayerBlend>())
+		{
+			FScopedTransaction Transaction(LOCTEXT("SetStageInputMask", "Set Material Designer Mask Source"));
+			Stage->Modify();
+
+			UDMMaterialStageInputExpression::ChangeStageInput_Expression(
+				Stage,
+				UDMMaterialStageExpressionWorldPositionNoise::StaticClass(),
+				UDMMaterialStageThroughputLayerBlend::InputMaskSource,
+				FDMMaterialStageConnectorChannel::WHOLE_CHANNEL,
+				0,
+				FDMMaterialStageConnectorChannel::WHOLE_CHANNEL
+			);
+		}
+		else
+		{
+			ensureMsgf(false, TEXT("Invalid stage type (%s)"), *StageSource->GetClass()->GetName());
+		}
+
+		if (TSharedPtr<SDMEditor> EditorWidget = InMenuContext->GetEditorWidget())
+		{
+			EditorWidget->InvalidateComponentEditWidget();
+		}
+	}
+
 	void ChangeSourceToSolidColorRGBFromContext(UDMMenuContext* InMenuContext)
 	{
 		if (!IsValid(InMenuContext))
@@ -1377,6 +1438,18 @@ namespace UE::DynamicMaterialEditor::Private
 			)
 		);
 
+		InSection.AddMenuEntry("Noise",
+			LOCTEXT("ChangeSourceNoise", "Noise"),
+			LOCTEXT("ChangeSourceNoiseTooltip", "Change the source of this stage to a Noise Renderer."),
+			FSlateIcon(),
+			FUIAction(
+				FExecuteAction::CreateStatic(
+					&ChangeSourceToNoiseFromContext,
+					MenuContext
+				)
+			)
+		);
+
 		if constexpr (UE::DynamicMaterialEditor::bGlobalValuesEnabled)
 		{
 			const TArray<UDMMaterialValue*>& Values = MaterialModel->GetValues();
diff --git a/Engine/Plugins/Experimental/DynamicMaterial/Source/DynamicMaterialEditor/Private/Slate/Properties/Generators/DMThroughputPropertyRowGenerator.cpp b/Engine/Plugins/Experimental/DynamicMaterial/Source/DynamicMaterialEditor/Private/Slate/Properties/Generators/DMThroughputPropertyRowGenerator.cpp
--- a/Engine/Plugins/Experimental/DynamicMaterial/Source/DynamicMaterialEditor/Private/Slate/Properties/Generators/DMThroughputPropertyRowGenerator.cpp
+++ b/Engine/Plugins/Experimental/DynamicMaterial/Source/DynamicMaterialEditor/Private/Slate/Properties/Generators/DMThroughputPropertyRowGenerator.cpp
@@ -1,10 +1,11 @@
 // Copyright Epic Games, Inc. All Rights Reserved.
 
 #include "Slate/Properties/Generators/DMThroughputPropertyRowGenerator.h"
-#include "DynamicMaterialEditorModule.h"
 #include "Components/DMMaterialStage.h"
 #include "Components/DMMaterialStageInput.h"
 #include "Components/DMMaterialStageThroughput.h"
+#include "Components/DMMaterialValue.h"
+#include "DynamicMaterialEditorModule.h"
 
 const TSharedRef<FDMThroughputPropertyRowGenerator>& FDMThroughputPropertyRowGenerator::Get()
 {
@@ -68,6 +69,8 @@ void FDMThroughputPropertyRowGenerator::AddComponentProperties(const TSharedRef<
 				continue;
 			}
 
+			const int32 StartRow = InOutPropertyRows.Num();
+
 			for (const FDMMaterialStageConnectorChannel& Channel : InputMap[InputIdx].Channels)
 			{
 				const int32 StageInputIdx = Channel.SourceIndex - FDMMaterialStageConnectorChannel::FIRST_STAGE_INPUT;
@@ -77,6 +80,32 @@ void FDMThroughputPropertyRowGenerator::AddComponentProperties(const TSharedRef<
 					FDynamicMaterialEditorModule::GeneratorComponentPropertyRows(InComponentEditWidget, Inputs[StageInputIdx], InOutPropertyRows, InOutProcessedObjects);
 				}
 			}
+
+			for (int32 PropertyRowIdx = StartRow; PropertyRowIdx < InOutPropertyRows.Num(); ++PropertyRowIdx)
+			{
+				FDMPropertyHandle& PropertyRow = InOutPropertyRows[PropertyRowIdx];
+
+				if (PropertyRow.NameOverride.IsSet())
+				{
+					continue;
+				}
+
+				if (!PropertyRow.PropertyHandle.IsValid()
+					|| !PropertyRow.PropertyHandle->GetProperty())
+				{
+					continue;
+				}
+
+				TArray<UObject*> Outers;
+				PropertyRow.PropertyHandle->GetOuterObjects(Outers);
+
+				if (Outers.IsEmpty() || !Outers[0]->IsA<UDMMaterialValue>())
+				{
+					continue;
+				}
+
+				PropertyRow.NameOverride = InputConnectors[InputIdx].Name;
+			}
 		}
 	}
 
