Material Designer
- Renamed Input Idx to Material Input Idx to connect input throughput methods.
- Added Input Idx (for the stage input index) for the above methods.
- Added world/actor space selector to the noise source.
- Added offset parameter to the noise source.

diff --git a/Engine/Plugins/Experimental/DynamicMaterial/Source/DynamicMaterial/Public/DMDefs.h b/Engine/Plugins/Experimental/DynamicMaterial/Source/DynamicMaterial/Public/DMDefs.h
--- a/Engine/Plugins/Experimental/DynamicMaterial/Source/DynamicMaterial/Public/DMDefs.h
+++ b/Engine/Plugins/Experimental/DynamicMaterial/Source/DynamicMaterial/Public/DMDefs.h
@@ -77,6 +77,13 @@ enum class EDMUVSource : uint8
 	WorldPosition
 };
 
+UENUM(BlueprintType)
+enum class EDMLocationType : uint8
+{
+	World,
+	Actor
+};
+
 /**
  * An individual component of a connector (e.g. G from RGB.)
  */
diff --git a/Engine/Plugins/Experimental/DynamicMaterial/Source/DynamicMaterialEditor/Private/Components/DMMaterialStage.cpp b/Engine/Plugins/Experimental/DynamicMaterial/Source/DynamicMaterialEditor/Private/Components/DMMaterialStage.cpp
--- a/Engine/Plugins/Experimental/DynamicMaterial/Source/DynamicMaterialEditor/Private/Components/DMMaterialStage.cpp
+++ b/Engine/Plugins/Experimental/DynamicMaterial/Source/DynamicMaterialEditor/Private/Components/DMMaterialStage.cpp
@@ -575,7 +575,7 @@ void UDMMaterialStage::GenerateExpressions(const TSharedRef<FDMMaterialBuildStat
 		{
 			const TArray<FDMMaterialStageConnector>& ThroughputInputs = Throughput->GetInputConnectors();
 
-			for (int32 InputIdx = 0; InputIdx < ThroughputInputs.Num() && InputIdx < InputConnectionMap.Num(); ++InputIdx)
+			for (int32 ThroughputInputIdx = 0; ThroughputInputIdx < ThroughputInputs.Num() && ThroughputInputIdx < InputConnectionMap.Num(); ++ThroughputInputIdx)
 			{
 				FDMMaterialStageConnectorChannel Channel;
 				TArray<UMaterialExpression*> Expressions;
@@ -589,7 +589,7 @@ void UDMMaterialStage::GenerateExpressions(const TSharedRef<FDMMaterialBuildStat
 
 				const int32 NodeOutputIndex = Throughput->ResolveInput(
 					InBuildState, 
-					InputIdx, 
+					ThroughputInputIdx, 
 					Channel, 
 					Expressions
 				);
@@ -600,7 +600,8 @@ void UDMMaterialStage::GenerateExpressions(const TSharedRef<FDMMaterialBuildStat
 
 					Throughput->ConnectOutputToInput(
 						InBuildState, 
-						ThroughputInputs[InputIdx].Index,
+						ThroughputInputIdx,
+						ThroughputInputs[ThroughputInputIdx].Index, 
 						Expressions.Last(), 
 						NodeOutputIndex, 
 						Channel.OutputChannel
diff --git a/Engine/Plugins/Experimental/DynamicMaterial/Source/DynamicMaterialEditor/Private/Components/DMMaterialStageBlendFunction.cpp b/Engine/Plugins/Experimental/DynamicMaterial/Source/DynamicMaterialEditor/Private/Components/DMMaterialStageBlendFunction.cpp
--- a/Engine/Plugins/Experimental/DynamicMaterial/Source/DynamicMaterialEditor/Private/Components/DMMaterialStageBlendFunction.cpp
+++ b/Engine/Plugins/Experimental/DynamicMaterial/Source/DynamicMaterialEditor/Private/Components/DMMaterialStageBlendFunction.cpp
@@ -53,12 +53,12 @@ void UDMMaterialStageBlendFunction::GenerateExpressions(const TSharedRef<FDMMate
 	InBuildState->AddStageSourceExpressions(this, {FunctionCall});
 }
 
-void UDMMaterialStageBlendFunction::ConnectOutputToInput(const TSharedRef<FDMMaterialBuildState>& InBuildState, int32 InInputIndex, UMaterialExpression* InSourceExpression,
-	int32 InSourceOutputIndex, int32 InSourceOutputChannel)
+void UDMMaterialStageBlendFunction::ConnectOutputToInput(const TSharedRef<FDMMaterialBuildState>& InBuildState, int32 InInputIdx, int32 InExpressionInputIndex,
+	UMaterialExpression* InSourceExpression, int32 InSourceOutputIndex, int32 InSourceOutputChannel)
 {
 	check(InSourceExpression);
 	check(InSourceExpression->GetOutputs().IsValidIndex(InSourceOutputIndex));
-	check(InInputIndex >= 0 && InInputIndex <= 2);
+	check(InExpressionInputIndex >= 0 && InExpressionInputIndex <= 2);
 
 	const TArray<UMaterialExpression*>& StageSourceExpressions = InBuildState->GetStageSourceExpressions(this);
 	check(!StageSourceExpressions.IsEmpty());
@@ -99,7 +99,7 @@ void UDMMaterialStageBlendFunction::ConnectOutputToInput(const TSharedRef<FDMMat
 		}
 	}
 
-	switch (InInputIndex)
+	switch (InExpressionInputIndex)
 	{
 		case InputA:
 			ConnectOutputToInput_Internal(InBuildState, StageSourceExpressions[0] /* FunctionCall */, InputAIndex, InSourceExpression, InSourceOutputIndex, InSourceOutputChannel);
diff --git a/Engine/Plugins/Experimental/DynamicMaterial/Source/DynamicMaterialEditor/Private/Components/DMMaterialStageThroughput.cpp b/Engine/Plugins/Experimental/DynamicMaterial/Source/DynamicMaterialEditor/Private/Components/DMMaterialStageThroughput.cpp
--- a/Engine/Plugins/Experimental/DynamicMaterial/Source/DynamicMaterialEditor/Private/Components/DMMaterialStageThroughput.cpp
+++ b/Engine/Plugins/Experimental/DynamicMaterial/Source/DynamicMaterialEditor/Private/Components/DMMaterialStageThroughput.cpp
@@ -40,33 +40,33 @@ const TArray<TStrongObjectPtr<UClass>>& UDMMaterialStageThroughput::GetAvailable
 	return Throughputs;
 }
 
-bool UDMMaterialStageThroughput::CanInputAcceptType(int32 InputIndex, EDMValueType ValueType) const
+bool UDMMaterialStageThroughput::CanInputAcceptType(int32 InThroughputInputIndex, EDMValueType InValueType) const
 {
-	check(InputConnectors.IsValidIndex(InputIndex));
+	check(InputConnectors.IsValidIndex(InThroughputInputIndex));
 
-	return InputConnectors[InputIndex].IsCompatibleWith(ValueType);
+	return InputConnectors[InThroughputInputIndex].IsCompatibleWith(InValueType);
 }
 
-bool UDMMaterialStageThroughput::CanInputConnectTo(int32 InputIndex, const FDMMaterialStageConnector& OutputConnector, int32 OutputChannel, bool bCheckSingleFloat)
+bool UDMMaterialStageThroughput::CanInputConnectTo(int32 InThroughputInputIndex, const FDMMaterialStageConnector& InOutputConnector, int32 InOutputChannel, bool bInCheckSingleFloat)
 {
-	if (OutputChannel == FDMMaterialStageConnectorChannel::WHOLE_CHANNEL)
+	if (InOutputChannel == FDMMaterialStageConnectorChannel::WHOLE_CHANNEL)
 	{
-		if (CanInputAcceptType(InputIndex, OutputConnector.Type))
+		if (CanInputAcceptType(InThroughputInputIndex, InOutputConnector.Type))
 		{
 			return true;
 		}
 
-		if (UDMValueDefinitionLibrary::GetValueDefinition(OutputConnector.Type).IsFloatType() && bCheckSingleFloat)
+		if (UDMValueDefinitionLibrary::GetValueDefinition(InOutputConnector.Type).IsFloatType() && bInCheckSingleFloat)
 		{
-			return CanInputAcceptType(InputIndex, EDMValueType::VT_Float1);
+			return CanInputAcceptType(InThroughputInputIndex, EDMValueType::VT_Float1);
 		}
 
 		return false;
 	}
 
-	if (UDMValueDefinitionLibrary::GetValueDefinition(OutputConnector.Type).IsFloatType())
+	if (UDMValueDefinitionLibrary::GetValueDefinition(InOutputConnector.Type).IsFloatType())
 	{
-		return CanInputAcceptType(InputIndex, EDMValueType::VT_Float1);
+		return CanInputAcceptType(InThroughputInputIndex, EDMValueType::VT_Float1);
 	}
 
 	// Only float types can have non-whole channels.
@@ -74,19 +74,19 @@ bool UDMMaterialStageThroughput::CanInputConnectTo(int32 InputIndex, const FDMMa
 	return false;
 }
 
-bool UDMMaterialStageThroughput::CanChangeInput(int32 InputIndex) const
+bool UDMMaterialStageThroughput::CanChangeInput(int32 InThroughputInputIndex) const
 {
 	return true;
 }
 
-bool UDMMaterialStageThroughput::CanChangeInputType(int32 InputIndex) const
+bool UDMMaterialStageThroughput::CanChangeInputType(int32 InThroughputInputIndex) const
 {
-	check(InputConnectors.IsValidIndex(InputIndex));
+	check(InputConnectors.IsValidIndex(InThroughputInputIndex));
 
-	return InputConnectors[InputIndex].Type != EDMValueType::VT_Texture;
+	return InputConnectors[InThroughputInputIndex].Type != EDMValueType::VT_Texture;
 }
 
-bool UDMMaterialStageThroughput::IsInputVisible(int32 InputIndex) const
+bool UDMMaterialStageThroughput::IsInputVisible(int32 InThroughputInputIndex) const
 {
 	UDMMaterialStage* Stage = GetStage();
 	check(Stage);
@@ -115,7 +115,7 @@ bool UDMMaterialStageThroughput::IsInputVisible(int32 InputIndex) const
 		}
 	}
 
-	if (bSupportsUVLink && InputIndex == GetLayerMaskTextureUVLinkInputIndex())
+	if (bSupportsUVLink && InThroughputInputIndex == GetLayerMaskTextureUVLinkInputIndex())
 	{
 		return !Layer->IsTextureUVLinkEnabled() || Layer->GetStageType(Stage) != EDMMaterialLayerStage::Mask;
 	}
@@ -123,12 +123,12 @@ bool UDMMaterialStageThroughput::IsInputVisible(int32 InputIndex) const
 	return true;
 }
 
-bool UDMMaterialStageThroughput::ShouldKeepInput(int32 InInputIdx)
+bool UDMMaterialStageThroughput::ShouldKeepInput(int32 InThroughputInputIndex)
 {
 	UDMMaterialStage* Stage = GetStage();
 	check(Stage);
 
-	if (!CanChangeInput(InInputIdx) || !Stage->IsInputMapped(InInputIdx))
+	if (!CanChangeInput(InThroughputInputIndex) || !Stage->IsInputMapped(InThroughputInputIndex))
 	{
 		return false;
 	}
@@ -139,21 +139,21 @@ bool UDMMaterialStageThroughput::ShouldKeepInput(int32 InInputIdx)
 	const EDMMaterialPropertyType StageProperty = Layer->GetMaterialProperty();
 	const TArray<FDMMaterialStageConnection>& InputConnectionMap = Stage->GetInputConnectionMap();
 
-	if (!InputConnectionMap.IsValidIndex(InInputIdx))
+	if (!InputConnectionMap.IsValidIndex(InThroughputInputIndex))
 	{
 		return false;
 	}
 
-	if (InputConnectors[InInputIdx].Type == EDMValueType::VT_Float_Any || InputConnectors[InInputIdx].Type == EDMValueType::VT_Texture)
+	if (InputConnectors[InThroughputInputIndex].Type == EDMValueType::VT_Float_Any || InputConnectors[InThroughputInputIndex].Type == EDMValueType::VT_Texture)
 	{
 		return false;
 	}
 
-	const int RequiredInputCount = UDMValueDefinitionLibrary::GetValueDefinition(InputConnectors[InInputIdx].Type).GetFloatCount();
+	const int RequiredInputCount = UDMValueDefinitionLibrary::GetValueDefinition(InputConnectors[InThroughputInputIndex].Type).GetFloatCount();
 
 	int ActualInputCount = 0;
 
-	for (const FDMMaterialStageConnectorChannel& Channel : InputConnectionMap[InInputIdx].Channels)
+	for (const FDMMaterialStageConnectorChannel& Channel : InputConnectionMap[InThroughputInputIndex].Channels)
 	{
 		int32 ThisInputCount = 0;
 
@@ -175,9 +175,9 @@ bool UDMMaterialStageThroughput::ShouldKeepInput(int32 InInputIdx)
 
 						if (UDMValueDefinitionLibrary::GetValueDefinition(PreviousStageOutputConnectors[Channel.OutputIndex].Type).IsFloatType() == false)
 						{
-							check(InputConnectionMap[InInputIdx].Channels.Num() == 1);
+							check(InputConnectionMap[InThroughputInputIndex].Channels.Num() == 1);
 
-							return InputConnectors[InInputIdx].IsCompatibleWith(PreviousStageOutputConnectors[Channel.OutputIndex]);
+							return InputConnectors[InThroughputInputIndex].IsCompatibleWith(PreviousStageOutputConnectors[Channel.OutputIndex]);
 						}
 						else
 						{
@@ -219,9 +219,9 @@ bool UDMMaterialStageThroughput::ShouldKeepInput(int32 InInputIdx)
 
 				if (UDMValueDefinitionLibrary::GetValueDefinition(InputOutputConnectors[Channel.OutputIndex].Type).IsFloatType() == false)
 				{
-					check(InputConnectionMap[InInputIdx].Channels.Num() == 1);
+					check(InputConnectionMap[InThroughputInputIndex].Channels.Num() == 1);
 
-					return InputConnectors[InInputIdx].IsCompatibleWith(InputOutputConnectors[Channel.OutputIndex]);
+					return InputConnectors[InThroughputInputIndex].IsCompatibleWith(InputOutputConnectors[Channel.OutputIndex]);
 				}
 				else
 				{
@@ -256,8 +256,8 @@ bool UDMMaterialStageThroughput::ShouldKeepInput(int32 InInputIdx)
 	return (ActualInputCount == RequiredInputCount);
 }
 
-void UDMMaterialStageThroughput::ConnectOutputToInput(const TSharedRef<FDMMaterialBuildState>& InBuildState, int32 InInputIndex,
-	UMaterialExpression* InSourceExpression, int32 InSourceOutputIndex, int32 InSourceOutputChannel)
+void UDMMaterialStageThroughput::ConnectOutputToInput(const TSharedRef<FDMMaterialBuildState>& InBuildState, int32 InThroughputInputIndex,
+	int32 InExpressionInputIndex, UMaterialExpression* InSourceExpression, int32 InSourceOutputIndex, int32 InSourceOutputChannel)
 {
 	check(InSourceExpression);
 	check(InSourceExpression->GetOutputs().IsValidIndex(InSourceOutputIndex));
@@ -265,52 +265,52 @@ void UDMMaterialStageThroughput::ConnectOutputToInput(const TSharedRef<FDMMateri
 	const TArray<UMaterialExpression*>& TargetExpressions = InBuildState->GetStageSourceExpressions(this);
 	check(!TargetExpressions.IsEmpty());
 
-	UMaterialExpression* TargetExpression = GetExpressionForInput(TargetExpressions, InInputIndex);
-	check(TargetExpression->GetInput(InInputIndex));
+	UMaterialExpression* TargetExpression = GetExpressionForInput(TargetExpressions, InThroughputInputIndex, InExpressionInputIndex);
+	check(TargetExpression->GetInput(InExpressionInputIndex));
 
 	ConnectOutputToInput_Internal(
 		InBuildState, 
 		TargetExpression, 
-		InInputIndex, 
+		InExpressionInputIndex, 
 		InSourceExpression, 
 		InSourceOutputIndex, 
 		InSourceOutputChannel
 	);
 }
 
-void UDMMaterialStageThroughput::ConnectOutputToInput_Internal(const TSharedRef<FDMMaterialBuildState>& InBuildState, UMaterialExpression* TargetExpression,
-	int32 InputIndex, UMaterialExpression* SourceExpression, int32 SourceOutputIndex, int32 SourceOutputChannel) const
+void UDMMaterialStageThroughput::ConnectOutputToInput_Internal(const TSharedRef<FDMMaterialBuildState>& InBuildState, UMaterialExpression* InTargetExpression,
+	int32 InExpressionInputIndex, UMaterialExpression* InSourceExpression, int32 InSourceOutputIndex, int32 InSourceOutputChannel) const
 {
-	check(TargetExpression != SourceExpression);
+	check(InTargetExpression != InSourceExpression);
 
-	FExpressionInput* ExpressionInput = TargetExpression->GetInput(InputIndex);
+	FExpressionInput* ExpressionInput = InTargetExpression->GetInput(InExpressionInputIndex);
 	
-	if (SourceOutputChannel == FDMMaterialStageConnectorChannel::WHOLE_CHANNEL)
+	if (InSourceOutputChannel == FDMMaterialStageConnectorChannel::WHOLE_CHANNEL)
 	{
-		SourceExpression->ConnectExpression(ExpressionInput, SourceOutputIndex);
+		InSourceExpression->ConnectExpression(ExpressionInput, InSourceOutputIndex);
 	}
-	else if (SourceExpression->IsA<UMaterialExpressionTextureBase>()
-		&& (SourceOutputChannel == FDMMaterialStageConnectorChannel::FIRST_CHANNEL
-			|| SourceOutputChannel == FDMMaterialStageConnectorChannel::SECOND_CHANNEL
-			|| SourceOutputChannel == FDMMaterialStageConnectorChannel::THIRD_CHANNEL
-			|| SourceOutputChannel == FDMMaterialStageConnectorChannel::FOURTH_CHANNEL))
+	else if (InSourceExpression->IsA<UMaterialExpressionTextureBase>()
+		&& (InSourceOutputChannel == FDMMaterialStageConnectorChannel::FIRST_CHANNEL
+			|| InSourceOutputChannel == FDMMaterialStageConnectorChannel::SECOND_CHANNEL
+			|| InSourceOutputChannel == FDMMaterialStageConnectorChannel::THIRD_CHANNEL
+			|| InSourceOutputChannel == FDMMaterialStageConnectorChannel::FOURTH_CHANNEL))
 	{
-		switch (SourceOutputChannel)
+		switch (InSourceOutputChannel)
 		{
 			case FDMMaterialStageConnectorChannel::FIRST_CHANNEL:
-				SourceExpression->ConnectExpression(ExpressionInput, 1);
+				InSourceExpression->ConnectExpression(ExpressionInput, 1);
 				break;
 
 			case FDMMaterialStageConnectorChannel::SECOND_CHANNEL:
-				SourceExpression->ConnectExpression(ExpressionInput, 2);
+				InSourceExpression->ConnectExpression(ExpressionInput, 2);
 				break;
 
 			case FDMMaterialStageConnectorChannel::THIRD_CHANNEL:
-				SourceExpression->ConnectExpression(ExpressionInput, 3);
+				InSourceExpression->ConnectExpression(ExpressionInput, 3);
 				break;
 
 			case FDMMaterialStageConnectorChannel::FOURTH_CHANNEL:
-				SourceExpression->ConnectExpression(ExpressionInput, 4);
+				InSourceExpression->ConnectExpression(ExpressionInput, 4);
 				break;
 
 			default:
@@ -320,9 +320,9 @@ void UDMMaterialStageThroughput::ConnectOutputToInput_Internal(const TSharedRef<
 	else
 	{
 		UMaterialExpression* MaskExpression = InBuildState->GetBuildUtils().CreateExpressionBitMask(
-			SourceExpression, 
-			SourceOutputIndex, 
-			SourceOutputChannel
+			InSourceExpression, 
+			InSourceOutputIndex, 
+			InSourceOutputChannel
 		);
 
 		MaskExpression->ConnectExpression(ExpressionInput, 0);
@@ -356,11 +356,12 @@ FDMExpressionInput UDMMaterialStageThroughput::GetLayerMaskLinkTextureUVInputExp
 	return ExpressionInput;
 }
 
-UMaterialExpression* UDMMaterialStageThroughput::GetExpressionForInput(const TArray<UMaterialExpression*>& StageSourceExpressions, int32 InputIdx)
+UMaterialExpression* UDMMaterialStageThroughput::GetExpressionForInput(const TArray<UMaterialExpression*>& InStageSourceExpressions, 
+	int32 InThroughputInputIndex, int32 InExpressionInputIndex)
 {
-	check(!StageSourceExpressions.IsEmpty());
+	check(!InStageSourceExpressions.IsEmpty());
 
-	return StageSourceExpressions[0];
+	return InStageSourceExpressions[0];
 }
 
 void UDMMaterialStageThroughput::OnComponentAdded()
@@ -498,7 +499,7 @@ void UDMMaterialStageThroughput::AddDefaultInput(int32 InInputIndex) const
 	}
 }
 
-int32 UDMMaterialStageThroughput::ResolveInput(const TSharedRef<FDMMaterialBuildState>& InBuildState, int32 InputIndex, FDMMaterialStageConnectorChannel& OutChannel,
+int32 UDMMaterialStageThroughput::ResolveInput(const TSharedRef<FDMMaterialBuildState>& InBuildState, int32 InThroughputInputIndex, FDMMaterialStageConnectorChannel& OutChannel,
 	TArray<UMaterialExpression*>& OutExpressions) const
 {
 	UDMMaterialStage* Stage = GetStage();
@@ -518,11 +519,11 @@ int32 UDMMaterialStageThroughput::ResolveInput(const TSharedRef<FDMMaterialBuild
 		&& Layer->GetStage(EDMMaterialLayerStage::Base, /* Enabled Only */ true)
 		&& Layer->GetStageType(ParentMostStage) == EDMMaterialLayerStage::Mask
 		&& SupportsLayerMaskTextureUVLink() 
-		&& (GetLayerMaskTextureUVLinkInputIndex() == INDEX_NONE || GetLayerMaskTextureUVLinkInputIndex() == InputIndex))
+		&& (GetLayerMaskTextureUVLinkInputIndex() == INDEX_NONE || GetLayerMaskTextureUVLinkInputIndex() == InThroughputInputIndex))
 	{
 		const int32 OutputIndex = ResolveLayerMaskTextureUVLinkInput(
 			InBuildState,
-			InputIndex, 
+			InThroughputInputIndex, 
 			OutChannel, 
 			OutExpressions
 		);
@@ -535,24 +536,24 @@ int32 UDMMaterialStageThroughput::ResolveInput(const TSharedRef<FDMMaterialBuild
 
 	const TArray<FDMMaterialStageConnection>& InputConnectionMap = Stage->GetInputConnectionMap();
 
-	if (!InputConnectionMap.IsValidIndex(InputIndex))
+	if (!InputConnectionMap.IsValidIndex(InThroughputInputIndex))
 	{
 		return INDEX_NONE;
 	}
 
-	if (InputConnectionMap[InputIndex].Channels.IsEmpty())
+	if (InputConnectionMap[InThroughputInputIndex].Channels.IsEmpty())
 	{
 		return INDEX_NONE;
 	}
 
-	if (InputConnectionMap[InputIndex].Channels.Num() == 1)
+	if (InputConnectionMap[InThroughputInputIndex].Channels.Num() == 1)
 	{
 		// Full copy in case it is changed by the channel resolve.
-		OutChannel = InputConnectionMap[InputIndex].Channels[0];
+		OutChannel = InputConnectionMap[InThroughputInputIndex].Channels[0];
 
 		return ResolveInputChannel(
 			InBuildState, 
-			InputIndex, 
+			InThroughputInputIndex, 
 			0, 
 			OutChannel, 
 			OutExpressions
@@ -560,12 +561,12 @@ int32 UDMMaterialStageThroughput::ResolveInput(const TSharedRef<FDMMaterialBuild
 	}
 
 	// Only floats can have sub-channel mapping
-	check(UDMValueDefinitionLibrary::GetValueDefinition(InputConnectors[InputIndex].Type).IsFloatType());
+	check(UDMValueDefinitionLibrary::GetValueDefinition(InputConnectors[InThroughputInputIndex].Type).IsFloatType());
 
 	// Valid scalar/vector is from 1 to 4 floats.
-	check(InputConnectionMap[InputIndex].Channels.Num() <= 4);
+	check(InputConnectionMap[InThroughputInputIndex].Channels.Num() <= 4);
 
-	for (const FDMMaterialStageConnectorChannel& Channel : InputConnectionMap[InputIndex].Channels)
+	for (const FDMMaterialStageConnectorChannel& Channel : InputConnectionMap[InThroughputInputIndex].Channels)
 	{
 		check(Channel.OutputChannel != FDMMaterialStageConnectorChannel::WHOLE_CHANNEL);
 	}
@@ -575,7 +576,7 @@ int32 UDMMaterialStageThroughput::ResolveInput(const TSharedRef<FDMMaterialBuild
 
 	UMaterialExpression* MakeFloat = nullptr;
 
-	switch (InputConnectionMap[InputIndex].Channels.Num())
+	switch (InputConnectionMap[InThroughputInputIndex].Channels.Num())
 	{
 		case 2:
 			MakeFloat = FDMMaterialFunctionLibrary::Get().GetMakeFloat2(InBuildState->GetDynamicMaterial(), UE_DM_NodeComment_Default);
@@ -594,14 +595,14 @@ int32 UDMMaterialStageThroughput::ResolveInput(const TSharedRef<FDMMaterialBuild
 			return INDEX_NONE;
 	}
 
-	for (int32 ChannelIdx = 0; ChannelIdx < InputConnectionMap[InputIndex].Channels.Num(); ++ChannelIdx)
+	for (int32 ChannelIdx = 0; ChannelIdx < InputConnectionMap[InThroughputInputIndex].Channels.Num(); ++ChannelIdx)
 	{
-		FDMMaterialStageConnectorChannel ChannelTemp = InputConnectionMap[InputIndex].Channels[ChannelIdx];
+		FDMMaterialStageConnectorChannel ChannelTemp = InputConnectionMap[InThroughputInputIndex].Channels[ChannelIdx];
 		TArray<UMaterialExpression*> ChannelExpressions;
 
 		ResolveInputChannel(
 			InBuildState, 
-			InputIndex, 
+			InThroughputInputIndex, 
 			ChannelIdx, 
 			ChannelTemp, 
 			ChannelExpressions
@@ -619,7 +620,7 @@ int32 UDMMaterialStageThroughput::ResolveInput(const TSharedRef<FDMMaterialBuild
 	return 0;
 }
 
-int32 UDMMaterialStageThroughput::ResolveLayerMaskTextureUVLinkInput(const TSharedRef<FDMMaterialBuildState>& InBuildState, int32 InputIndex, 
+int32 UDMMaterialStageThroughput::ResolveLayerMaskTextureUVLinkInput(const TSharedRef<FDMMaterialBuildState>& InBuildState, int32 InThroughputInputIndex, 
 	FDMMaterialStageConnectorChannel& OutChannel, TArray<UMaterialExpression*>& OutExpressions) const
 {
 	UDMMaterialStage* Stage = GetStage();
@@ -638,11 +639,11 @@ int32 UDMMaterialStageThroughput::ResolveLayerMaskTextureUVLinkInput(const TShar
 }
 
 int32 UDMMaterialStageThroughput::ResolveLayerMaskTextureUVLinkInputImpl(const TSharedRef<FDMMaterialBuildState>& InBuildState, 
-	const UDMMaterialStageSource* StageSource, FDMMaterialStageConnectorChannel& OutChannel, TArray<UMaterialExpression*>& OutExpressions)
+	const UDMMaterialStageSource* InStageSource, FDMMaterialStageConnectorChannel& OutChannel, TArray<UMaterialExpression*>& OutExpressions)
 {
 	FDMExpressionInput ConnectedInputExpressions;
 
-	if (const UDMMaterialStageInputTextureUV* InputTextureUV = Cast<const UDMMaterialStageInputTextureUV>(StageSource))
+	if (const UDMMaterialStageInputTextureUV* InputTextureUV = Cast<const UDMMaterialStageInputTextureUV>(InStageSource))
 	{
 		InputTextureUV->GenerateExpressions(InBuildState);
 		OutExpressions = InBuildState->GetStageSourceExpressions(InputTextureUV);
@@ -652,7 +653,7 @@ int32 UDMMaterialStageThroughput::ResolveLayerMaskTextureUVLinkInputImpl(const T
 		return 0;
 	}
 
-	else if (const UDMMaterialStageThroughput* Throughput = Cast<const UDMMaterialStageThroughput>(StageSource))
+	else if (const UDMMaterialStageThroughput* Throughput = Cast<const UDMMaterialStageThroughput>(InStageSource))
 	{
 		if (Throughput->SupportsLayerMaskTextureUVLink())
 		{
@@ -660,7 +661,7 @@ int32 UDMMaterialStageThroughput::ResolveLayerMaskTextureUVLinkInputImpl(const T
 		}
 	}
 
-	else if (const UDMMaterialStageInputThroughput* InputThroughput = Cast<const UDMMaterialStageInputThroughput>(StageSource))
+	else if (const UDMMaterialStageInputThroughput* InputThroughput = Cast<const UDMMaterialStageInputThroughput>(InStageSource))
 	{
 		if (InputThroughput->GetMaterialStageThroughput())
 		{
@@ -686,7 +687,7 @@ int32 UDMMaterialStageThroughput::ResolveLayerMaskTextureUVLinkInputImpl(const T
 	return INDEX_NONE;
 }
 
-int32 UDMMaterialStageThroughput::ResolveInputChannel(const TSharedRef<FDMMaterialBuildState>& InBuildState, int32 InputIndex, int32 ChannelIndex,
+int32 UDMMaterialStageThroughput::ResolveInputChannel(const TSharedRef<FDMMaterialBuildState>& InBuildState, int32 InThroughputInputIndex, int32 InChannelIndex,
 	FDMMaterialStageConnectorChannel& OutChannel, TArray<UMaterialExpression*>& OutExpressions) const
 {
 	if (OutChannel.SourceIndex == FDMMaterialStageConnectorChannel::NO_SOURCE)
@@ -694,7 +695,7 @@ int32 UDMMaterialStageThroughput::ResolveInputChannel(const TSharedRef<FDMMateri
 		return INDEX_NONE;
 	}
 
-	check(InputConnectors.IsValidIndex(InputIndex));
+	check(InputConnectors.IsValidIndex(InThroughputInputIndex));
 
 	UDMMaterialStage* Stage = GetStage();
 	check(Stage);
@@ -764,7 +765,7 @@ int32 UDMMaterialStageThroughput::ResolveInputChannel(const TSharedRef<FDMMateri
 	// If our "Previous Stage" is blank, these need default values.
 	const int32 OutputFloatCount = InputSourceOutputConnectors ? UDMValueDefinitionLibrary::GetValueDefinition((*InputSourceOutputConnectors)[OutChannel.OutputIndex].Type).GetFloatCount() : 3;
 	const bool bOutputIsFloatType = InputSourceOutputConnectors ? UDMValueDefinitionLibrary::GetValueDefinition((*InputSourceOutputConnectors)[OutChannel.OutputIndex].Type).IsFloatType() : true;
-	const int32 InputFloatCount = UDMValueDefinitionLibrary::GetValueDefinition(InputConnectors[InputIndex].Type).GetFloatCount();
+	const int32 InputFloatCount = UDMValueDefinitionLibrary::GetValueDefinition(InputConnectors[InThroughputInputIndex].Type).GetFloatCount();
 
 	if (OutChannel.OutputChannel == FDMMaterialStageConnectorChannel::WHOLE_CHANNEL
 		&& bOutputIsFloatType && OutputFloatCount > 0 && InputFloatCount > 0)
diff --git a/Engine/Plugins/Experimental/DynamicMaterial/Source/DynamicMaterialEditor/Private/Components/DMMaterialStageThroughputLayerBlend.cpp b/Engine/Plugins/Experimental/DynamicMaterial/Source/DynamicMaterialEditor/Private/Components/DMMaterialStageThroughputLayerBlend.cpp
--- a/Engine/Plugins/Experimental/DynamicMaterial/Source/DynamicMaterialEditor/Private/Components/DMMaterialStageThroughputLayerBlend.cpp
+++ b/Engine/Plugins/Experimental/DynamicMaterial/Source/DynamicMaterialEditor/Private/Components/DMMaterialStageThroughputLayerBlend.cpp
@@ -675,12 +675,12 @@ void UDMMaterialStageThroughputLayerBlend::NotifyPostChange(const FPropertyChang
 	Super::NotifyPostChange(InPropertyChangedEvent, InPropertyThatChanged);
 }
 
-void UDMMaterialStageThroughputLayerBlend::ConnectOutputToInput(const TSharedRef<FDMMaterialBuildState>& InBuildState, int32 InInputIndex, 
-	UMaterialExpression* InSourceExpression, int32 InSourceOutputIndex, int32 InSourceOutputChannel)
+void UDMMaterialStageThroughputLayerBlend::ConnectOutputToInput(const TSharedRef<FDMMaterialBuildState>& InBuildState, int32 InInputIdx,
+	int32 InExpressionInputIndex, UMaterialExpression* InSourceExpression, int32 InSourceOutputIndex, int32 InSourceOutputChannel)
 {
 	check(InSourceExpression);
 	check(InSourceExpression->GetOutputs().IsValidIndex(InSourceOutputIndex));
-	check(InInputIndex >= InputPreviousLayer && InInputIndex <= InputMaskSource);
+	check(InExpressionInputIndex >= InputPreviousLayer && InExpressionInputIndex <= InputMaskSource);
  
 	const TArray<UMaterialExpression*>& StageSourceExpressions = InBuildState->GetStageSourceExpressions(this);
 	check(!StageSourceExpressions.IsEmpty());
@@ -720,7 +720,7 @@ void UDMMaterialStageThroughputLayerBlend::ConnectOutputToInput(const TSharedRef
 		}
 	}
  
-	switch (InInputIndex)
+	switch (InExpressionInputIndex)
 	{
 		case InputPreviousLayer:
 			ConnectOutputToInput_Internal(
diff --git a/Engine/Plugins/Experimental/DynamicMaterial/Source/DynamicMaterialEditor/Private/Components/MaterialStageExpressions/DMMSEWorldPositionNoise.cpp b/Engine/Plugins/Experimental/DynamicMaterial/Source/DynamicMaterialEditor/Private/Components/MaterialStageExpressions/DMMSEWorldPositionNoise.cpp
--- a/Engine/Plugins/Experimental/DynamicMaterial/Source/DynamicMaterialEditor/Private/Components/MaterialStageExpressions/DMMSEWorldPositionNoise.cpp
+++ b/Engine/Plugins/Experimental/DynamicMaterial/Source/DynamicMaterialEditor/Private/Components/MaterialStageExpressions/DMMSEWorldPositionNoise.cpp
@@ -6,11 +6,12 @@
 #include "Components/DMMaterialProperty.h"
 #include "Components/DMMaterialStageBlend.h"
 #include "Components/DMMaterialSubStage.h"
-#include "Components/MaterialStageExpressions/DMMSETextureSample.h"
-#include "CoreGlobals.h"
 #include "Components/MaterialStageInputs/DMMSIValue.h"
-#include "Components/MaterialValues/DMMaterialValueFloat1.h"
+#include "Components/MaterialValues/DMMaterialValueFloat3XYZ.h"
+#include "CoreGlobals.h"
+#include "Materials/MaterialExpressionAdd.h"
 #include "Materials/MaterialExpressionDivide.h"
+#include "Materials/MaterialExpressionSubtract.h"
 #include "Materials/MaterialExpressionVectorNoise.h"
 #include "Materials/MaterialExpressionWorldPosition.h"
 #include "Model/DMMaterialBuildState.h"
@@ -23,6 +24,7 @@ UDMMaterialStageExpressionWorldPositionNoise::UDMMaterialStageExpressionWorldPos
 		LOCTEXT("UMaterialExpressionVectorNoise", "Noise"),
 		UMaterialExpressionVectorNoise::StaticClass()
 	)
+	, LocationType(EDMLocationType::World)
 	, ShaderOffset(EWorldPositionIncludedOffsets::WPT_Default)
 	, NoiseFunction(EVectorNoiseFunction::VNF_VectorALU)
 	, Quality(1)
@@ -32,10 +34,12 @@ UDMMaterialStageExpressionWorldPositionNoise::UDMMaterialStageExpressionWorldPos
 	bInputRequired = true;
 	bAllowNestedInputs = true;
 
-	InputConnectors.Add({1, LOCTEXT("Scale", "Scale"), EDMValueType::VT_Float1});
+	InputConnectors.Add({1, LOCTEXT("Scale", "Scale"), EDMValueType::VT_Float3_XYZ});
+	InputConnectors.Add({1, LOCTEXT("Offset", "Offset"), EDMValueType::VT_Float3_XYZ});
 
 	OutputConnectors.Add({0, LOCTEXT("ColorRGB", "Color (RGB)"), EDMValueType::VT_Float3_RGB});
 
+	EditableProperties.Add(GET_MEMBER_NAME_CHECKED(UDMMaterialStageExpressionWorldPositionNoise, LocationType));
 	EditableProperties.Add(GET_MEMBER_NAME_CHECKED(UDMMaterialStageExpressionWorldPositionNoise, ShaderOffset));
 	EditableProperties.Add(GET_MEMBER_NAME_CHECKED(UDMMaterialStageExpressionWorldPositionNoise, NoiseFunction));
 	EditableProperties.Add(GET_MEMBER_NAME_CHECKED(UDMMaterialStageExpressionWorldPositionNoise, Quality));
@@ -62,7 +66,29 @@ void UDMMaterialStageExpressionWorldPositionNoise::GenerateExpressions(const TSh
 
 	UMaterialExpressionDivide* Divide = InBuildState->GetBuildUtils().CreateExpression<UMaterialExpressionDivide>(UE_DM_NodeComment_Default);
 
-	WorldPosition->ConnectExpression(&Divide->A, 0);
+	switch (LocationType)
+	{
+		case EDMLocationType::World:
+			WorldPosition->ConnectExpression(&Divide->A, 0);
+			break;
+
+		case EDMLocationType::Actor:
+		{
+			// Not exported
+			TSubclassOf<UMaterialExpression> ActorPositionWSClass = FindClass("MaterialExpressionActorPositionWS");
+			UMaterialExpression* ActorPositionWS = InBuildState->GetBuildUtils().CreateExpression(ActorPositionWSClass, UE_DM_NodeComment_Default);
+
+			UMaterialExpressionSubtract* Subtract = InBuildState->GetBuildUtils().CreateExpression<UMaterialExpressionSubtract>(UE_DM_NodeComment_Default);
+			WorldPosition->ConnectExpression(&Subtract->A, 0);
+			ActorPositionWS->ConnectExpression(&Subtract->B, 0);
+			Subtract->ConnectExpression(&Divide->A, 0);
+
+			break;
+		}
+	}
+
+	UMaterialExpressionAdd* Add = InBuildState->GetBuildUtils().CreateExpression<UMaterialExpressionAdd>(UE_DM_NodeComment_Default);
+	Divide->ConnectExpression(&Add->A, 0);
 
 	UMaterialExpressionVectorNoise* VectorNoise = InBuildState->GetBuildUtils().CreateExpression<UMaterialExpressionVectorNoise>(UE_DM_NodeComment_Default);
 	VectorNoise->NoiseFunction = NoiseFunction;
@@ -70,9 +96,21 @@ void UDMMaterialStageExpressionWorldPositionNoise::GenerateExpressions(const TSh
 	VectorNoise->bTiling = bTiling ? 1 : 0;
 	VectorNoise->TileSize = TileSize;
 
-	Divide->ConnectExpression(&VectorNoise->Position, 0);
+	Add->ConnectExpression(&VectorNoise->Position, 0);
+
+	InBuildState->AddStageSourceExpressions(this, {WorldPosition, Divide, Add, VectorNoise});
+}
+
+void UDMMaterialStageExpressionWorldPositionNoise::SetLocationType(EDMLocationType InLocationType)
+{
+	if (LocationType == InLocationType)
+	{
+		return;
+	}
+
+	LocationType = InLocationType;
 
-	InBuildState->AddStageSourceExpressions(this, {WorldPosition, Divide, VectorNoise});
+	Update(EDMUpdateType::Structure);
 }
 
 void UDMMaterialStageExpressionWorldPositionNoise::SetShaderOffset(EWorldPositionIncludedOffsets InShaderOffset)
@@ -144,16 +182,31 @@ void UDMMaterialStageExpressionWorldPositionNoise::AddDefaultInput(int32 InInput
 
 	switch (InInputIndex)
 	{
+		// Scale
 		case 0:
 		{
 			UDMMaterialStageInputValue* InputValue = Cast<UDMMaterialStageInputValue>(Stage->GetInputs().Last());
 			check(InputValue);
 
-			UDMMaterialValueFloat1* Float1Value = Cast<UDMMaterialValueFloat1>(InputValue->GetValue());
-			check(Float1Value);
+			UDMMaterialValueFloat3XYZ* Float3Value = Cast<UDMMaterialValueFloat3XYZ>(InputValue->GetValue());
+			check(Float3Value);
 
-			Float1Value->SetDefaultValue(10.f);
-			Float1Value->ApplyDefaultValue();
+			Float3Value->SetDefaultValue({1.f, 10.f, 10.f});
+			Float3Value->ApplyDefaultValue();
+			break;
+		}
+
+		// Offset
+		case 1:
+		{
+			UDMMaterialStageInputValue* InputValue = Cast<UDMMaterialStageInputValue>(Stage->GetInputs().Last());
+			check(InputValue);
+
+			UDMMaterialValueFloat3XYZ* Float3Value = Cast<UDMMaterialValueFloat3XYZ>(InputValue->GetValue());
+			check(Float3Value);
+
+			Float3Value->SetDefaultValue({0.f, 0.f, 0.f});
+			Float3Value->ApplyDefaultValue();
 			break;
 		}
 
@@ -163,20 +216,36 @@ void UDMMaterialStageExpressionWorldPositionNoise::AddDefaultInput(int32 InInput
 	}
 }
 
-UMaterialExpression* UDMMaterialStageExpressionWorldPositionNoise::GetExpressionForInput(const TArray<UMaterialExpression*>& StageSourceExpressions, int32 InputIdx)
+UMaterialExpression* UDMMaterialStageExpressionWorldPositionNoise::GetExpressionForInput(const TArray<UMaterialExpression*>& InStageSourceExpressions, 
+	int32 InInputIndex, int32 InExpressionInputIndex)
 {
-	if (InputIdx == 1 && StageSourceExpressions.IsValidIndex(1))
+	switch (InInputIndex)
 	{
-		return StageSourceExpressions[1];
+		// Scale
+		case 0:
+			if (InStageSourceExpressions.IsValidIndex(1))
+			{
+				return InStageSourceExpressions[1];
+			}
+			break;
+
+		// Offset
+		case 1:
+			if (InStageSourceExpressions.IsValidIndex(2))
+			{
+				return InStageSourceExpressions[2];
+			}
+			break;
 	}
 
-	return Super::GetExpressionForInput(StageSourceExpressions, InputIdx);
+	return Super::GetExpressionForInput(InStageSourceExpressions, InInputIndex, InExpressionInputIndex);
 }
 
 void UDMMaterialStageExpressionWorldPositionNoise::PostEditChangeProperty(FPropertyChangedEvent& InPropertyChangedEvent)
 {
 	Super::PostEditChangeProperty(InPropertyChangedEvent);
 
+	static const FName LocationTypeName = GET_MEMBER_NAME_CHECKED(UDMMaterialStageExpressionWorldPositionNoise, LocationType);
 	static const FName ShaderOffsetName = GET_MEMBER_NAME_CHECKED(UDMMaterialStageExpressionWorldPositionNoise, ShaderOffset);
 	static const FName NoiseFunctionName = GET_MEMBER_NAME_CHECKED(UDMMaterialStageExpressionWorldPositionNoise, NoiseFunction);
 	static const FName QualityName = GET_MEMBER_NAME_CHECKED(UDMMaterialStageExpressionWorldPositionNoise, Quality);
@@ -185,7 +254,8 @@ void UDMMaterialStageExpressionWorldPositionNoise::PostEditChangeProperty(FPrope
 
 	const FName PropertyName = InPropertyChangedEvent.GetMemberPropertyName();
 
-	if (PropertyName == ShaderOffsetName
+	if (PropertyName == LocationTypeName
+		|| PropertyName == ShaderOffsetName
 		|| PropertyName == NoiseFunctionName
 		|| PropertyName == QualityName
 		|| PropertyName == TilingName
diff --git a/Engine/Plugins/Experimental/DynamicMaterial/Source/DynamicMaterialEditor/Private/Components/MaterialStageExpressions/DMMSEWorldPositionNoise.h b/Engine/Plugins/Experimental/DynamicMaterial/Source/DynamicMaterialEditor/Private/Components/MaterialStageExpressions/DMMSEWorldPositionNoise.h
--- a/Engine/Plugins/Experimental/DynamicMaterial/Source/DynamicMaterialEditor/Private/Components/MaterialStageExpressions/DMMSEWorldPositionNoise.h
+++ b/Engine/Plugins/Experimental/DynamicMaterial/Source/DynamicMaterialEditor/Private/Components/MaterialStageExpressions/DMMSEWorldPositionNoise.h
@@ -5,6 +5,7 @@
 #include "Components/DMMaterialStageExpression.h"
 #include "DMMSEWorldPositionNoise.generated.h"
 
+enum class EDMLocationType : uint8;
 enum EVectorNoiseFunction : int;
 enum EWorldPositionIncludedOffsets : int;
 
@@ -20,6 +21,12 @@ public:
 	virtual void GenerateExpressions(const TSharedRef<FDMMaterialBuildState>& InBuildState) const override;
 	//~ End UDMMaterialStageExpression
 
+	UFUNCTION(BlueprintPure, Category = "Material Designer")
+	virtual EDMLocationType GetLocationType() const { return LocationType; }
+
+	UFUNCTION(BlueprintCallable, Category = "Material Designer")
+	virtual void SetLocationType(EDMLocationType InLocationType);
+
 	UFUNCTION(BlueprintPure, Category = "Material Designer")
 	virtual TEnumAsByte<EWorldPositionIncludedOffsets> GetShaderOffset() const { return ShaderOffset; }
 
@@ -52,7 +59,7 @@ public:
 
 	//~ Begin UDMMaterialStageThroughput
 	virtual void AddDefaultInput(int32 InInputIndex) const override;
-	virtual UMaterialExpression* GetExpressionForInput(const TArray<UMaterialExpression*>& StageSourceExpressions, int32 InputIdx) override;
+	virtual UMaterialExpression* GetExpressionForInput(const TArray<UMaterialExpression*>& InStageSourceExpressions, int32 InInputIndex, int32 InExpressionInputIndex) override;
 	//~ End UDMMaterialStageThroughput
 
 	//~ Begin UObject
@@ -60,6 +67,10 @@ public:
 	//~ End UObject
 
 protected:
+	UPROPERTY(EditInstanceOnly, BlueprintReadWrite, Setter = SetLocationType, BlueprintSetter = SetLocationType, Category = "Material Designer",
+		meta = (NotKeyframeable, AllowPrivateAccess = "true"))
+	EDMLocationType LocationType;
+
 	UPROPERTY(EditInstanceOnly, BlueprintReadWrite, Setter = SetShaderOffset, BlueprintSetter = SetShaderOffset, Category = "Material Designer",
 		meta = (NotKeyframeable, AllowPrivateAccess = "true"))
 	TEnumAsByte<EWorldPositionIncludedOffsets> ShaderOffset;
diff --git a/Engine/Plugins/Experimental/DynamicMaterial/Source/DynamicMaterialEditor/Public/Components/DMMaterialStageBlendFunction.h b/Engine/Plugins/Experimental/DynamicMaterial/Source/DynamicMaterialEditor/Public/Components/DMMaterialStageBlendFunction.h
--- a/Engine/Plugins/Experimental/DynamicMaterial/Source/DynamicMaterialEditor/Public/Components/DMMaterialStageBlendFunction.h
+++ b/Engine/Plugins/Experimental/DynamicMaterial/Source/DynamicMaterialEditor/Public/Components/DMMaterialStageBlendFunction.h
@@ -16,8 +16,8 @@ class DYNAMICMATERIALEDITOR_API UDMMaterialStageBlendFunction : public UDMMateri
 
 public:
 	virtual void GenerateExpressions(const TSharedRef<FDMMaterialBuildState>& InBuildState) const override;
-	virtual void ConnectOutputToInput(const TSharedRef<FDMMaterialBuildState>& InBuildState, int32 InInputIndex, UMaterialExpression* InSourceExpression,
-		int32 InSourceOutputIndex, int32 InSourceOutputChannel) override;
+	virtual void ConnectOutputToInput(const TSharedRef<FDMMaterialBuildState>& InBuildState, int32 InInputIdx, int32 InExpressionInputIndex,
+		UMaterialExpression* InSourceExpression, int32 InSourceOutputIndex, int32 InSourceOutputChannel) override;
 
 protected:
 	UPROPERTY()
diff --git a/Engine/Plugins/Experimental/DynamicMaterial/Source/DynamicMaterialEditor/Public/Components/DMMaterialStageThroughput.h b/Engine/Plugins/Experimental/DynamicMaterial/Source/DynamicMaterialEditor/Public/Components/DMMaterialStageThroughput.h
--- a/Engine/Plugins/Experimental/DynamicMaterial/Source/DynamicMaterialEditor/Public/Components/DMMaterialStageThroughput.h
+++ b/Engine/Plugins/Experimental/DynamicMaterial/Source/DynamicMaterialEditor/Public/Components/DMMaterialStageThroughput.h
@@ -50,19 +50,19 @@ public:
 	const TArray<FDMMaterialStageConnector>& GetInputConnectors() const { return InputConnectors; }
 
 	UFUNCTION(BlueprintCallable, Category = "Material Designer")
-	virtual bool CanInputAcceptType(int32 InputIndex, EDMValueType ValueType) const;
+	virtual bool CanInputAcceptType(int32 InThroughputInputIndex, EDMValueType InValueType) const;
 
 	UFUNCTION(BlueprintCallable, Category = "Material Designer")
-	virtual bool CanInputConnectTo(int32 InputIndex, const FDMMaterialStageConnector& OutputConnector, int32 OutputChannel,
-		bool bCheckSingleFloat = false);
+	virtual bool CanInputConnectTo(int32 InThroughputInputIndex, const FDMMaterialStageConnector& InOutputConnector, int32 InOutputChannel,
+		bool bInCheckSingleFloat = false);
 
-	virtual bool CanChangeInput(int32 InputIndex) const;
-	virtual bool CanChangeInputType(int32 InputIndex) const;
+	virtual bool CanChangeInput(int32 InThroughputInputIndex) const;
+	virtual bool CanChangeInputType(int32 InThroughputInputIndex) const;
 
-	virtual bool IsInputVisible(int32 InputIndex) const;
+	virtual bool IsInputVisible(int32 InThroughputInputIndex) const;
 
-	virtual void ConnectOutputToInput(const TSharedRef<FDMMaterialBuildState>& InBuildState, int32 InInputIndex, UMaterialExpression* InSourceExpression,
-		int32 InSourceOutputIndex, int32 InSourceOutputChannel);
+	virtual void ConnectOutputToInput(const TSharedRef<FDMMaterialBuildState>& InBuildState, int32 InThroughputInputIndex, int32 InExpressionInputIndex,
+		UMaterialExpression* InSourceExpression, int32 InSourceOutputIndex, int32 InSourceOutputChannel);
 
 	/** Returns true if the layer and mask can have their Texture UV linked. */
 	virtual bool SupportsLayerMaskTextureUVLink() const { return false; }
@@ -81,25 +81,26 @@ public:
 	 * Returns the first node in the array by default
 	 * --> In [ ]-[ ]-[ ] Out -->
 	 */
-	virtual UMaterialExpression* GetExpressionForInput(const TArray<UMaterialExpression*>& StageSourceExpressions, int32 InputIdx);
+	virtual UMaterialExpression* GetExpressionForInput(const TArray<UMaterialExpression*>& InStageSourceExpressions, int32 InThroughputInputIndex, 
+		int32 InExpressionInputIndex);
 
 	virtual void AddDefaultInput(int32 InInputIndex) const;
 
 	/** Returns the actual output index of the material expression */
-	virtual int32 ResolveInput(const TSharedRef<FDMMaterialBuildState>& InBuildState, int32 InputIndex, FDMMaterialStageConnectorChannel& OutChannel,
+	virtual int32 ResolveInput(const TSharedRef<FDMMaterialBuildState>& InBuildState, int32 InThroughputInputIndex, FDMMaterialStageConnectorChannel& OutChannel,
 		TArray<UMaterialExpression*>& OutExpressions) const;
 
-	virtual int32 ResolveLayerMaskTextureUVLinkInput(const TSharedRef<FDMMaterialBuildState>& InBuildState, int32 InputIndex, 
+	virtual int32 ResolveLayerMaskTextureUVLinkInput(const TSharedRef<FDMMaterialBuildState>& InBuildState, int32 InThroughputInputIndex, 
 		FDMMaterialStageConnectorChannel& OutChannel, TArray<UMaterialExpression*>& OutExpressions) const;
 
-	virtual void InputUpdated(int32 InInputIndex, EDMUpdateType InUpdateType) { }
+	virtual void InputUpdated(int32 InThroughputInputIndex, EDMUpdateType InUpdateType) { }
 
 protected:
 	static TArray<TStrongObjectPtr<UClass>> Throughputs;
 
 	static void GenerateThroughputList();
 
-	static int32 ResolveLayerMaskTextureUVLinkInputImpl(const TSharedRef<FDMMaterialBuildState>& InBuildState, const UDMMaterialStageSource* StageSource,
+	static int32 ResolveLayerMaskTextureUVLinkInputImpl(const TSharedRef<FDMMaterialBuildState>& InBuildState, const UDMMaterialStageSource* InStageSource,
 		FDMMaterialStageConnectorChannel& OutChannel, TArray<UMaterialExpression*>& OutExpressions);
 
 	UDMMaterialStageThroughput();
@@ -117,12 +118,12 @@ protected:
 	UPROPERTY(VisibleInstanceOnly, BlueprintReadOnly, Category = "Material Designer")
 	TArray<FDMMaterialStageConnector> InputConnectors;
 
-	virtual bool ShouldKeepInput(int32 InInputIdx);
+	virtual bool ShouldKeepInput(int32 InThroughputInputIndex);
 
-	void ConnectOutputToInput_Internal(const TSharedRef<FDMMaterialBuildState>& InBuildState, UMaterialExpression* TargetExpression,
-		int32 InputIndex, UMaterialExpression* SourceExpression, int32 SourceOutputIndex, int32 SourceOutputChannel) const;
+	void ConnectOutputToInput_Internal(const TSharedRef<FDMMaterialBuildState>& InBuildState, UMaterialExpression* InTargetExpression,
+		int32 InExpressionInputIndex, UMaterialExpression* InSourceExpression, int32 InSourceOutputIndex, int32 InSourceOutputChannel) const;
 
-	virtual int32 ResolveInputChannel(const TSharedRef<FDMMaterialBuildState>& InBuildState, int32 InputIndex, int32 ChannelIndex, 
+	virtual int32 ResolveInputChannel(const TSharedRef<FDMMaterialBuildState>& InBuildState, int32 InThroughputInputIndex, int32 InChannelIndex, 
 		FDMMaterialStageConnectorChannel& OutChannel, TArray<UMaterialExpression*>& OutExpressions) const;
 
 	virtual void UpdatePreviewMaterial(UMaterial* InPreviewMaterial = nullptr);
diff --git a/Engine/Plugins/Experimental/DynamicMaterial/Source/DynamicMaterialEditor/Public/Components/DMMaterialStageThroughputLayerBlend.h b/Engine/Plugins/Experimental/DynamicMaterial/Source/DynamicMaterialEditor/Public/Components/DMMaterialStageThroughputLayerBlend.h
--- a/Engine/Plugins/Experimental/DynamicMaterial/Source/DynamicMaterialEditor/Public/Components/DMMaterialStageThroughputLayerBlend.h
+++ b/Engine/Plugins/Experimental/DynamicMaterial/Source/DynamicMaterialEditor/Public/Components/DMMaterialStageThroughputLayerBlend.h
@@ -56,8 +56,8 @@ public:
 	virtual bool IsInputVisible(int32 InputIndex) const override;
 	virtual int32 ResolveInput(const TSharedRef<FDMMaterialBuildState>& InBuildState, int32 InputIndex, FDMMaterialStageConnectorChannel& OutChannel,
 		TArray<UMaterialExpression*>& OutExpressions) const override;
-	virtual void ConnectOutputToInput(const TSharedRef<FDMMaterialBuildState>& InBuildState, int32 InInputIndex, UMaterialExpression* InSourceExpression,
-		int32 InSourceOutputIndex, int32 InSourceOutputChannel) override;
+	virtual void ConnectOutputToInput(const TSharedRef<FDMMaterialBuildState>& InBuildState, int32 InInputIdx, int32 InExpressionInputIndex,
+		UMaterialExpression* InSourceExpression, int32 InSourceOutputIndex, int32 InSourceOutputChannel) override;
 	//~ End UDMMaterialStageThroughput
 
 	void GetMaskOutput(const TSharedRef<FDMMaterialBuildState>& InBuildState, UMaterialExpression*& OutExpression, int32& OutOutputIndex, int32& OutOutputChannel) const;
