Material Designer
- Horde issue 493518

diff --git a/Engine/Plugins/Experimental/DynamicMaterial/Source/DynamicMaterialEditor/Private/Slate/SDMToolBar.cpp b/Engine/Plugins/Experimental/DynamicMaterial/Source/DynamicMaterialEditor/Private/Slate/SDMToolBar.cpp
--- a/Engine/Plugins/Experimental/DynamicMaterial/Source/DynamicMaterialEditor/Private/Slate/SDMToolBar.cpp
+++ b/Engine/Plugins/Experimental/DynamicMaterial/Source/DynamicMaterialEditor/Private/Slate/SDMToolBar.cpp
@@ -186,7 +186,7 @@ TSharedRef<SWidget> SDMToolBar::CreateToolBarEntries()
 			.IsEnabled(this, &SDMToolBar::CanSave)
 			.ContentPadding(GetLargeIconToolBarButtonContentPadding())
 			.ButtonStyle(FDynamicMaterialEditorStyle::Get(), "HoverHintOnly")
-			.ToolTipText(LOCTEXT("MaterialDesignerBrowseTooltip", "Save the Material Designer asset\n\nCaution: If this asset lives inside an actor, the actor/level will be saved."))
+			.ToolTipText(LOCTEXT("MaterialDesignerSaveTooltip", "Save the Material Designer asset\n\nCaution: If this asset lives inside an actor, the actor/level will be saved."))
 			.OnClicked(this, &SDMToolBar::OnSaveClicked)
 			[
 				SNew(SImage)
