Material Designer
- Added some more NoCreate to asset properties.
- Removed NotKeyframeable from things that should be keyframeable.

diff --git a/Engine/Plugins/Experimental/DynamicMaterial/Source/DynamicMaterial/Public/Components/MaterialValues/DMMaterialValueColorAtlas.h b/Engine/Plugins/Experimental/DynamicMaterial/Source/DynamicMaterial/Public/Components/MaterialValues/DMMaterialValueColorAtlas.h
--- a/Engine/Plugins/Experimental/DynamicMaterial/Source/DynamicMaterial/Public/Components/MaterialValues/DMMaterialValueColorAtlas.h
+++ b/Engine/Plugins/Experimental/DynamicMaterial/Source/DynamicMaterial/Public/Components/MaterialValues/DMMaterialValueColorAtlas.h
@@ -72,11 +72,11 @@ protected:
 	float DefaultValue;
 
 	UPROPERTY(EditInstanceOnly, BlueprintReadWrite, BlueprintSetter = SetAtlas, Category = "Material Designer",
-		meta = (AllowPrivateAccess = "true"))
+		meta = (AllowPrivateAccess = "true", NotKeyframeable, NoCreate))
 	TObjectPtr<UCurveLinearColorAtlas> Atlas;
 
 	UPROPERTY(EditInstanceOnly, BlueprintReadWrite, BlueprintSetter = SetCurve, Category = "Material Designer",
-		meta = (AllowPrivateAccess = "true"))
+		meta = (AllowPrivateAccess = "true", NotKeyframeable, NoCreate))
 	TObjectPtr<UCurveLinearColor> Curve;
 #endif
 
diff --git a/Engine/Plugins/Experimental/DynamicMaterial/Source/DynamicMaterial/Public/Components/MaterialValues/DMMaterialValueRenderTarget.h b/Engine/Plugins/Experimental/DynamicMaterial/Source/DynamicMaterial/Public/Components/MaterialValues/DMMaterialValueRenderTarget.h
--- a/Engine/Plugins/Experimental/DynamicMaterial/Source/DynamicMaterial/Public/Components/MaterialValues/DMMaterialValueRenderTarget.h
+++ b/Engine/Plugins/Experimental/DynamicMaterial/Source/DynamicMaterial/Public/Components/MaterialValues/DMMaterialValueRenderTarget.h
@@ -89,19 +89,19 @@ protected:
 	FDelegateHandle EndOfFrameDelegateHandle;
 
 	UPROPERTY(EditAnywhere, BlueprintReadWrite, Setter = SetTextureSize, BlueprintSetter = SetTextureSize, Category = "Material Designer|Render Target",
-		meta = (AllowPrivateAccess = "true"))
+		meta = (AllowPrivateAccess = "true", NotKeyframeable))
 	FIntPoint TextureSize;
 
 	UPROPERTY(EditAnywhere, BlueprintReadWrite, Setter = SetTextureFormat, BlueprintSetter = SetTextureFormat, Category = "Material Designer|Render Target",
-		meta = (AllowPrivateAccess = "true"))
+		meta = (AllowPrivateAccess = "true", NotKeyframeable))
 	TEnumAsByte<ETextureRenderTargetFormat> TextureFormat;
 
 	UPROPERTY(EditAnywhere, BlueprintReadWrite, Setter = SetClearColor, BlueprintSetter = SetClearColor, Category = "Material Designer|Render Target",
-		meta = (AllowPrivateAccess = "true"))
+		meta = (AllowPrivateAccess = "true", NotKeyframeable))
 	FLinearColor ClearColor;
 
 	UPROPERTY(EditAnywhere, BlueprintReadWrite, Setter = SetRenderer, BlueprintSetter = SetRenderer, Category = "Material Designer",
-		meta = (AllowPrivateAccess = "true"))
+		meta = (AllowPrivateAccess = "true", NotKeyframeable, NoCreate))
 	TObjectPtr<UDMRenderTargetRenderer> Renderer;
 
 	void AsyncCreateRenderTarget();
diff --git a/Engine/Plugins/Experimental/DynamicMaterial/Source/DynamicMaterial/Public/Components/RenderTargetRenderers/DMRenderTargetTextRenderer.h b/Engine/Plugins/Experimental/DynamicMaterial/Source/DynamicMaterial/Public/Components/RenderTargetRenderers/DMRenderTargetTextRenderer.h
--- a/Engine/Plugins/Experimental/DynamicMaterial/Source/DynamicMaterial/Public/Components/RenderTargetRenderers/DMRenderTargetTextRenderer.h
+++ b/Engine/Plugins/Experimental/DynamicMaterial/Source/DynamicMaterial/Public/Components/RenderTargetRenderers/DMRenderTargetTextRenderer.h
@@ -259,27 +259,27 @@ protected:
 
 	/** Multiplier on the base font height. */
 	UPROPERTY(EditInstanceOnly, BlueprintReadWrite, Setter = SetLineHeight, BlueprintSetter = SetLineHeight, Category = "Material Designer|Text",
-		meta = (NotKeyframeable, AllowPrivateAccess = "true", UIMin = 0.1, UIMax = 2.0))
+		meta = (AllowPrivateAccess = "true", UIMin = 0.1, UIMax = 2.0))
 	float LineHeight = 1.f;
 
 	/** Extra space adding beyond the edge of the glyphs. Useful for shadows, glows, etc. */
 	UPROPERTY(EditInstanceOnly, BlueprintReadWrite, Setter = SetPaddingLeft, BlueprintSetter = SetPaddingLeft, Category = "Material Designer|Text",
-		meta = (NotKeyframeable, AllowPrivateAccess = "true"))
+		meta = (AllowPrivateAccess = "true"))
 	float PaddingLeft = 0.f;
 
 	/** Extra space adding beyond the edge of the glyphs. Useful for shadows, glows, etc. */
 	UPROPERTY(EditInstanceOnly, BlueprintReadWrite, Setter = SetPaddingRight, BlueprintSetter = SetPaddingRight, Category = "Material Designer|Text",
-		meta = (NotKeyframeable, AllowPrivateAccess = "true"))
+		meta = (AllowPrivateAccess = "true"))
 	float PaddingRight = 0.f;
 
 	/** Extra space adding beyond the edge of the glyphs. Useful for shadows, glows, etc. */
 	UPROPERTY(EditInstanceOnly, BlueprintReadWrite, Setter = SetPaddingTop, BlueprintSetter = SetPaddingTop, Category = "Material Designer|Text",
-		meta = (NotKeyframeable, AllowPrivateAccess = "true"))
+		meta = (AllowPrivateAccess = "true"))
 	float PaddingTop = 0.f;
 
 	/** Extra space adding beyond the edge of the glyphs. Useful for shadows, glows, etc. */
 	UPROPERTY(EditInstanceOnly, BlueprintReadWrite, Setter = SetPaddingBottom, BlueprintSetter = SetPaddingBottom, Category = "Material Designer|Text",
-		meta = (NotKeyframeable, AllowPrivateAccess = "true"))
+		meta = (AllowPrivateAccess = "true"))
 	float PaddingBottom = 0.f;
 
 	UPROPERTY(EditInstanceOnly, BlueprintReadWrite, Setter = SetOverrideRenderTargetSize, BlueprintSetter = SetOverrideRenderTargetSize, Category = "Material Designer|Text",
diff --git a/Engine/Plugins/Experimental/DynamicMaterial/Source/DynamicMaterial/Public/Components/RenderTargetRenderers/DMRenderTargetUMGWidgetRenderer.h b/Engine/Plugins/Experimental/DynamicMaterial/Source/DynamicMaterial/Public/Components/RenderTargetRenderers/DMRenderTargetUMGWidgetRenderer.h
--- a/Engine/Plugins/Experimental/DynamicMaterial/Source/DynamicMaterial/Public/Components/RenderTargetRenderers/DMRenderTargetUMGWidgetRenderer.h
+++ b/Engine/Plugins/Experimental/DynamicMaterial/Source/DynamicMaterial/Public/Components/RenderTargetRenderers/DMRenderTargetUMGWidgetRenderer.h
@@ -35,7 +35,7 @@ public:
 
 protected:
 	UPROPERTY(EditInstanceOnly, BlueprintReadWrite, Setter = SetWidgetClass, BlueprintSetter = SetWidgetClass, Category = "Material Designer",
-		meta = (NotKeyframeable, AllowPrivateAccess = "true"))
+		meta = (NotKeyframeable, AllowPrivateAccess = "true", NoCreate))
 	TSubclassOf<UWidget> WidgetClass;
 
 	UPROPERTY()
diff --git a/Engine/Plugins/Experimental/DynamicMaterial/Source/DynamicMaterialEditor/Public/Components/DMMaterialStageFunction.h b/Engine/Plugins/Experimental/DynamicMaterial/Source/DynamicMaterialEditor/Public/Components/DMMaterialStageFunction.h
--- a/Engine/Plugins/Experimental/DynamicMaterial/Source/DynamicMaterialEditor/Public/Components/DMMaterialStageFunction.h
+++ b/Engine/Plugins/Experimental/DynamicMaterial/Source/DynamicMaterialEditor/Public/Components/DMMaterialStageFunction.h
@@ -63,7 +63,7 @@ protected:
 	static TSoftObjectPtr<UMaterialFunctionInterface> NoOp;
 
 	UPROPERTY(EditAnywhere, BlueprintReadWrite, Getter = GetMaterialFunction, Setter = SetMaterialFunction, BlueprintSetter = SetMaterialFunction, Category = "Material Designer",
-		meta = (DisplayThumbnail = true, AllowPrivateAccess = "true", HighPriority, NotKeyframeable))
+		meta = (DisplayThumbnail = true, AllowPrivateAccess = "true", HighPriority, NotKeyframeable, NoCreate))
 	TObjectPtr<UMaterialFunctionInterface> MaterialFunction;
 
 	UPROPERTY(Transient, DuplicateTransient, TextExportTransient)
