Material Designer
- Fixed texture set settings category.
- Fixed opacity mask texture set filter.

diff --git a/Engine/Plugins/Experimental/DynamicMaterial/Source/DynamicMaterialTextureSetEditor/Private/DMTextureSetSettings.cpp b/Engine/Plugins/Experimental/DynamicMaterial/Source/DynamicMaterialTextureSetEditor/Private/DMTextureSetSettings.cpp
--- a/Engine/Plugins/Experimental/DynamicMaterial/Source/DynamicMaterialTextureSetEditor/Private/DMTextureSetSettings.cpp
+++ b/Engine/Plugins/Experimental/DynamicMaterial/Source/DynamicMaterialTextureSetEditor/Private/DMTextureSetSettings.cpp
@@ -6,6 +6,9 @@
 
 UDMTextureSetSettings::UDMTextureSetSettings()
 {
+	CategoryName = TEXT("Plugins");
+	SectionName = TEXT("Material Designer Texture Set");
+
 	int32 FilterCount = 0;
 
 	FDMTextureSetFilter BaseColor;
@@ -49,8 +52,8 @@ UDMTextureSetSettings::UDMTextureSetSettings()
 	++FilterCount;
 
 	FDMTextureSetFilter OpacityMask;
-	Opacity.FilterStrings = {TEXT("OpacityMask"), TEXT("Opacity_Mask"), TEXT("_OM"), TEXT("AlphaMask"), TEXT("Alpha_Mask"), TEXT("_AM")};
-	Opacity.MaterialProperties = {{EDMTextureSetMaterialProperty::OpacityMask, EDMTextureChannelMask::RGBA}};
+	OpacityMask.FilterStrings = {TEXT("OpacityMask"), TEXT("Opacity_Mask"), TEXT("_OM"), TEXT("AlphaMask"), TEXT("Alpha_Mask"), TEXT("_AM")};
+	OpacityMask.MaterialProperties = {{EDMTextureSetMaterialProperty::OpacityMask, EDMTextureChannelMask::RGBA}};
 	++FilterCount;
 
 	FDMTextureSetFilter Displacement;
diff --git a/Engine/Plugins/Experimental/DynamicMaterial/Source/DynamicMaterialTextureSetEditor/Private/DMTextureSetSettings.h b/Engine/Plugins/Experimental/DynamicMaterial/Source/DynamicMaterialTextureSetEditor/Private/DMTextureSetSettings.h
--- a/Engine/Plugins/Experimental/DynamicMaterial/Source/DynamicMaterialTextureSetEditor/Private/DMTextureSetSettings.h
+++ b/Engine/Plugins/Experimental/DynamicMaterial/Source/DynamicMaterialTextureSetEditor/Private/DMTextureSetSettings.h
@@ -11,7 +11,7 @@
 /**
  * DM Texture Set Settings
  */
-UCLASS(Config = EditorPerProjectUserSettings, meta = (DisplayName = "Material Designer"))
+UCLASS(Config = EditorPerProjectUserSettings, meta = (DisplayName = "Material Designer Texture Set"))
 class UDMTextureSetSettings : public UDeveloperSettings
 {
 	GENERATED_BODY()
