Material Designer
- Separated the base color and emissive channels.

diff --git a/Engine/Plugins/Experimental/DynamicMaterial/Source/DynamicMaterial/Private/Components/MaterialValues/DMMaterialValueBool.cpp b/Engine/Plugins/Experimental/DynamicMaterial/Source/DynamicMaterial/Private/Components/MaterialValues/DMMaterialValueBool.cpp
--- a/Engine/Plugins/Experimental/DynamicMaterial/Source/DynamicMaterial/Private/Components/MaterialValues/DMMaterialValueBool.cpp
+++ b/Engine/Plugins/Experimental/DynamicMaterial/Source/DynamicMaterial/Private/Components/MaterialValues/DMMaterialValueBool.cpp
@@ -9,6 +9,7 @@
 #include "Materials/MaterialExpressionStaticBoolParameter.h"
 #include "Model/IDMMaterialBuildStateInterface.h"
 #include "Model/IDMMaterialBuildUtilsInterface.h"
+#include "Utils/DMUtils.h"
 #endif
  
 #define LOCTEXT_NAMESPACE "DMMaterialValueBool"
diff --git a/Engine/Plugins/Experimental/DynamicMaterial/Source/DynamicMaterial/Private/Components/MaterialValues/DMMaterialValueColorAtlas.cpp b/Engine/Plugins/Experimental/DynamicMaterial/Source/DynamicMaterial/Private/Components/MaterialValues/DMMaterialValueColorAtlas.cpp
--- a/Engine/Plugins/Experimental/DynamicMaterial/Source/DynamicMaterial/Private/Components/MaterialValues/DMMaterialValueColorAtlas.cpp
+++ b/Engine/Plugins/Experimental/DynamicMaterial/Source/DynamicMaterial/Private/Components/MaterialValues/DMMaterialValueColorAtlas.cpp
@@ -13,6 +13,7 @@
 #include "Materials/MaterialExpressionCurveAtlasRowParameter.h"
 #include "Model/IDMMaterialBuildStateInterface.h"
 #include "Model/IDMMaterialBuildUtilsInterface.h"
+#include "Utils/DMUtils.h"
 #endif
 
 UDMMaterialValueColorAtlas::UDMMaterialValueColorAtlas()
diff --git a/Engine/Plugins/Experimental/DynamicMaterial/Source/DynamicMaterial/Private/Components/MaterialValues/DMMaterialValueFloat1.cpp b/Engine/Plugins/Experimental/DynamicMaterial/Source/DynamicMaterial/Private/Components/MaterialValues/DMMaterialValueFloat1.cpp
--- a/Engine/Plugins/Experimental/DynamicMaterial/Source/DynamicMaterial/Private/Components/MaterialValues/DMMaterialValueFloat1.cpp
+++ b/Engine/Plugins/Experimental/DynamicMaterial/Source/DynamicMaterial/Private/Components/MaterialValues/DMMaterialValueFloat1.cpp
@@ -11,6 +11,7 @@
 #include "Materials/MaterialExpressionScalarParameter.h"
 #include "Model/IDMMaterialBuildStateInterface.h"
 #include "Model/IDMMaterialBuildUtilsInterface.h"
+#include "Utils/DMUtils.h"
 #endif
  
 #define LOCTEXT_NAMESPACE "DMMaterialValueFloat1"
diff --git a/Engine/Plugins/Experimental/DynamicMaterial/Source/DynamicMaterial/Private/Components/MaterialValues/DMMaterialValueFloat2.cpp b/Engine/Plugins/Experimental/DynamicMaterial/Source/DynamicMaterial/Private/Components/MaterialValues/DMMaterialValueFloat2.cpp
--- a/Engine/Plugins/Experimental/DynamicMaterial/Source/DynamicMaterial/Private/Components/MaterialValues/DMMaterialValueFloat2.cpp
+++ b/Engine/Plugins/Experimental/DynamicMaterial/Source/DynamicMaterial/Private/Components/MaterialValues/DMMaterialValueFloat2.cpp
@@ -11,6 +11,7 @@
 #include "Materials/MaterialExpressionVectorParameter.h"
 #include "Model/IDMMaterialBuildStateInterface.h"
 #include "Model/IDMMaterialBuildUtilsInterface.h"
+#include "Utils/DMUtils.h"
 #endif
 
 #define LOCTEXT_NAMESPACE "DMMaterialValueFloat2"
diff --git a/Engine/Plugins/Experimental/DynamicMaterial/Source/DynamicMaterial/Private/Components/MaterialValues/DMMaterialValueFloat3RGB.cpp b/Engine/Plugins/Experimental/DynamicMaterial/Source/DynamicMaterial/Private/Components/MaterialValues/DMMaterialValueFloat3RGB.cpp
--- a/Engine/Plugins/Experimental/DynamicMaterial/Source/DynamicMaterial/Private/Components/MaterialValues/DMMaterialValueFloat3RGB.cpp
+++ b/Engine/Plugins/Experimental/DynamicMaterial/Source/DynamicMaterial/Private/Components/MaterialValues/DMMaterialValueFloat3RGB.cpp
@@ -10,6 +10,7 @@
 #include "Materials/MaterialExpressionVectorParameter.h"
 #include "Model/IDMMaterialBuildStateInterface.h"
 #include "Model/IDMMaterialBuildUtilsInterface.h"
+#include "Utils/DMUtils.h"
 #endif
  
 #define LOCTEXT_NAMESPACE "DMMaterialValueFloat3RGB"
diff --git a/Engine/Plugins/Experimental/DynamicMaterial/Source/DynamicMaterial/Private/Components/MaterialValues/DMMaterialValueFloat3RPY.cpp b/Engine/Plugins/Experimental/DynamicMaterial/Source/DynamicMaterial/Private/Components/MaterialValues/DMMaterialValueFloat3RPY.cpp
--- a/Engine/Plugins/Experimental/DynamicMaterial/Source/DynamicMaterial/Private/Components/MaterialValues/DMMaterialValueFloat3RPY.cpp
+++ b/Engine/Plugins/Experimental/DynamicMaterial/Source/DynamicMaterial/Private/Components/MaterialValues/DMMaterialValueFloat3RPY.cpp
@@ -10,6 +10,7 @@
 #include "Materials/MaterialExpressionVectorParameter.h"
 #include "Model/IDMMaterialBuildStateInterface.h"
 #include "Model/IDMMaterialBuildUtilsInterface.h"
+#include "Utils/DMUtils.h"
 #endif
  
 #define LOCTEXT_NAMESPACE "DMMaterialValueFloat3RPY"
diff --git a/Engine/Plugins/Experimental/DynamicMaterial/Source/DynamicMaterial/Private/Components/MaterialValues/DMMaterialValueFloat3XYZ.cpp b/Engine/Plugins/Experimental/DynamicMaterial/Source/DynamicMaterial/Private/Components/MaterialValues/DMMaterialValueFloat3XYZ.cpp
--- a/Engine/Plugins/Experimental/DynamicMaterial/Source/DynamicMaterial/Private/Components/MaterialValues/DMMaterialValueFloat3XYZ.cpp
+++ b/Engine/Plugins/Experimental/DynamicMaterial/Source/DynamicMaterial/Private/Components/MaterialValues/DMMaterialValueFloat3XYZ.cpp
@@ -10,6 +10,7 @@
 #include "Materials/MaterialExpressionVectorParameter.h"
 #include "Model/IDMMaterialBuildStateInterface.h"
 #include "Model/IDMMaterialBuildUtilsInterface.h"
+#include "Utils/DMUtils.h"
 #endif
  
 #define LOCTEXT_NAMESPACE "DMMaterialValueFloat3XYZ"
diff --git a/Engine/Plugins/Experimental/DynamicMaterial/Source/DynamicMaterial/Private/Components/MaterialValues/DMMaterialValueFloat4.cpp b/Engine/Plugins/Experimental/DynamicMaterial/Source/DynamicMaterial/Private/Components/MaterialValues/DMMaterialValueFloat4.cpp
--- a/Engine/Plugins/Experimental/DynamicMaterial/Source/DynamicMaterial/Private/Components/MaterialValues/DMMaterialValueFloat4.cpp
+++ b/Engine/Plugins/Experimental/DynamicMaterial/Source/DynamicMaterial/Private/Components/MaterialValues/DMMaterialValueFloat4.cpp
@@ -10,6 +10,7 @@
 #include "Materials/MaterialExpressionVectorParameter.h"
 #include "Model/IDMMaterialBuildStateInterface.h"
 #include "Model/IDMMaterialBuildUtilsInterface.h"
+#include "Utils/DMUtils.h"
 #endif
  
 #define LOCTEXT_NAMESPACE "DMMaterialValueFloat4"
diff --git a/Engine/Plugins/Experimental/DynamicMaterial/Source/DynamicMaterial/Private/Components/MaterialValues/DMMaterialValueTexture.cpp b/Engine/Plugins/Experimental/DynamicMaterial/Source/DynamicMaterial/Private/Components/MaterialValues/DMMaterialValueTexture.cpp
--- a/Engine/Plugins/Experimental/DynamicMaterial/Source/DynamicMaterial/Private/Components/MaterialValues/DMMaterialValueTexture.cpp
+++ b/Engine/Plugins/Experimental/DynamicMaterial/Source/DynamicMaterial/Private/Components/MaterialValues/DMMaterialValueTexture.cpp
@@ -14,6 +14,7 @@
 #include "Model/IDMMaterialBuildStateInterface.h"
 #include "Model/IDMMaterialBuildUtilsInterface.h"
 #include "RenderUtils.h"
+#include "Utils/DMUtils.h"
 #endif
 
 #if WITH_EDITOR
diff --git a/Engine/Plugins/Experimental/DynamicMaterial/Source/DynamicMaterial/Private/DMDefs.cpp b/Engine/Plugins/Experimental/DynamicMaterial/Source/DynamicMaterial/Private/DMDefs.cpp
--- a/Engine/Plugins/Experimental/DynamicMaterial/Source/DynamicMaterial/Private/DMDefs.cpp
+++ b/Engine/Plugins/Experimental/DynamicMaterial/Source/DynamicMaterial/Private/DMDefs.cpp
@@ -2,128 +2,9 @@
 
 #include "DMDefs.h"
 
-#include "DMTextureSetMaterialProperty.h"
-#include "Misc/Paths.h"
-
 int32 FDMUpdateGuard::GuardCount = 0;
 uint32 FDMInitializationGuard::GuardCount = 0;
 
-FString UE::DynamicMaterial::CreateNodeComment(const ANSICHAR* InFile, int InLine, const ANSICHAR* InFunction, const FString* InComment)
-{
-	FString File = FPaths::GetCleanFilename(ANSI_TO_TCHAR(InFile)).RightChop(2);
-
-	if (InComment)
-	{
-		return FString::Printf(TEXT("%s[%d]: %hs: %s"), *File, InLine, InFunction, **InComment);
-	}
-	else
-	{
-		return FString::Printf(TEXT("%s[%d]: %hs"), *File, InLine, InFunction);
-	}
-}
-
-EDMTextureSetMaterialProperty UE::DynamicMaterial::MaterialPropertyTypeToMaterialProperty(EDMMaterialPropertyType InPropertyType)
-{
-	switch (InPropertyType)
-	{
-		case EDMMaterialPropertyType::BaseColor:
-			return EDMTextureSetMaterialProperty::BaseColor;
-
-		case EDMMaterialPropertyType::EmissiveColor:
-			return EDMTextureSetMaterialProperty::EmissiveColor;
-
-		case EDMMaterialPropertyType::Opacity:
-			return EDMTextureSetMaterialProperty::Opacity;
-
-		case EDMMaterialPropertyType::OpacityMask:
-			return EDMTextureSetMaterialProperty::OpacityMask;
-
-		case EDMMaterialPropertyType::Roughness:
-			return EDMTextureSetMaterialProperty::Roughness;
-
-		case EDMMaterialPropertyType::Specular:
-			return EDMTextureSetMaterialProperty::Specular;
-
-		case EDMMaterialPropertyType::Metallic:
-			return EDMTextureSetMaterialProperty::Metallic;
-
-		case EDMMaterialPropertyType::Normal:
-			return EDMTextureSetMaterialProperty::Normal;
-
-		case EDMMaterialPropertyType::PixelDepthOffset:
-			return EDMTextureSetMaterialProperty::PixelDepthOffset;
-
-		case EDMMaterialPropertyType::WorldPositionOffset:
-			return EDMTextureSetMaterialProperty::WorldPositionOffset;
-
-		case EDMMaterialPropertyType::AmbientOcclusion:
-			return EDMTextureSetMaterialProperty::AmbientOcclusion;
-
-		case EDMMaterialPropertyType::Anisotropy:
-			return EDMTextureSetMaterialProperty::Anisotropy;
-
-		case EDMMaterialPropertyType::Refraction:
-			return EDMTextureSetMaterialProperty::Refraction;
-
-		case EDMMaterialPropertyType::Tangent:
-			return EDMTextureSetMaterialProperty::Tangent;
-
-		default:
-			return EDMTextureSetMaterialProperty::None;
-	}
-}
-
-EDMMaterialPropertyType UE::DynamicMaterial::MaterialPropertyToMaterialPropertyType(EDMTextureSetMaterialProperty InPropertyType)
-{
-	switch (InPropertyType)
-	{
-		case EDMTextureSetMaterialProperty::BaseColor:
-			return EDMMaterialPropertyType::BaseColor;
-
-		case EDMTextureSetMaterialProperty::EmissiveColor:
-			return EDMMaterialPropertyType::EmissiveColor;
-
-		case EDMTextureSetMaterialProperty::Opacity:
-			return EDMMaterialPropertyType::Opacity;
-
-		case EDMTextureSetMaterialProperty::OpacityMask:
-			return EDMMaterialPropertyType::OpacityMask;
-
-		case EDMTextureSetMaterialProperty::Roughness:
-			return EDMMaterialPropertyType::Roughness;
-
-		case EDMTextureSetMaterialProperty::Specular:
-			return EDMMaterialPropertyType::Specular;
-
-		case EDMTextureSetMaterialProperty::Metallic:
-			return EDMMaterialPropertyType::Metallic;
-
-		case EDMTextureSetMaterialProperty::Normal:
-			return EDMMaterialPropertyType::Normal;
-
-		case EDMTextureSetMaterialProperty::PixelDepthOffset:
-			return EDMMaterialPropertyType::PixelDepthOffset;
-
-		case EDMTextureSetMaterialProperty::WorldPositionOffset:
-			return EDMMaterialPropertyType::WorldPositionOffset;
-
-		case EDMTextureSetMaterialProperty::AmbientOcclusion:
-			return EDMMaterialPropertyType::AmbientOcclusion;
-
-		case EDMTextureSetMaterialProperty::Anisotropy:
-			return EDMMaterialPropertyType::Anisotropy;
-
-		case EDMTextureSetMaterialProperty::Refraction:
-			return EDMMaterialPropertyType::Refraction;
-
-		case EDMTextureSetMaterialProperty::Tangent:
-			return EDMMaterialPropertyType::Tangent;
-
-		default:
-			return EDMMaterialPropertyType::None;
-	}
-}
-
 bool FDMInitializationGuard::IsInitializing()
 {
 	return GuardCount > 0;
diff --git a/Engine/Plugins/Experimental/DynamicMaterial/Source/DynamicMaterial/Private/Utils/DMUtils.cpp b/Engine/Plugins/Experimental/DynamicMaterial/Source/DynamicMaterial/Private/Utils/DMUtils.cpp
new file mode 100644
--- /dev/null
+++ b/Engine/Plugins/Experimental/DynamicMaterial/Source/DynamicMaterial/Private/Utils/DMUtils.cpp
@@ -0,0 +1,231 @@
+// Copyright Epic Games, Inc. All Rights Reserved.
+
+#include "Utils/DMUtils.h"
+
+#include "DMDefs.h"
+#include "DMTextureSetMaterialProperty.h"
+#include "SceneTypes.h"
+
+#if WITH_EDITOR
+#include "Misc/Paths.h"
+#endif
+
+#if WITH_EDITOR
+FString FDMUtils::CreateNodeComment(const ANSICHAR* InFile, int InLine, const ANSICHAR* InFunction, const FString* InComment)
+{
+	FString File = FPaths::GetCleanFilename(ANSI_TO_TCHAR(InFile)).RightChop(2);
+
+	if (InComment)
+	{
+		return FString::Printf(TEXT("%s[%d]: %hs: %s"), *File, InLine, InFunction, **InComment);
+	}
+	else
+	{
+		return FString::Printf(TEXT("%s[%d]: %hs"), *File, InLine, InFunction);
+	}
+}
+#endif
+
+EDMTextureSetMaterialProperty FDMUtils::MaterialPropertyTypeToTextureSetMaterialProperty(EDMMaterialPropertyType InPropertyType)
+{
+	switch (InPropertyType)
+	{
+		case EDMMaterialPropertyType::BaseColor:
+			return EDMTextureSetMaterialProperty::BaseColor;
+
+		case EDMMaterialPropertyType::EmissiveColor:
+			return EDMTextureSetMaterialProperty::EmissiveColor;
+
+		case EDMMaterialPropertyType::Opacity:
+			return EDMTextureSetMaterialProperty::Opacity;
+
+		case EDMMaterialPropertyType::OpacityMask:
+			return EDMTextureSetMaterialProperty::OpacityMask;
+
+		case EDMMaterialPropertyType::Roughness:
+			return EDMTextureSetMaterialProperty::Roughness;
+
+		case EDMMaterialPropertyType::Specular:
+			return EDMTextureSetMaterialProperty::Specular;
+
+		case EDMMaterialPropertyType::Metallic:
+			return EDMTextureSetMaterialProperty::Metallic;
+
+		case EDMMaterialPropertyType::Normal:
+			return EDMTextureSetMaterialProperty::Normal;
+
+		case EDMMaterialPropertyType::PixelDepthOffset:
+			return EDMTextureSetMaterialProperty::PixelDepthOffset;
+
+		case EDMMaterialPropertyType::WorldPositionOffset:
+			return EDMTextureSetMaterialProperty::WorldPositionOffset;
+
+		case EDMMaterialPropertyType::AmbientOcclusion:
+			return EDMTextureSetMaterialProperty::AmbientOcclusion;
+
+		case EDMMaterialPropertyType::Anisotropy:
+			return EDMTextureSetMaterialProperty::Anisotropy;
+
+		case EDMMaterialPropertyType::Refraction:
+			return EDMTextureSetMaterialProperty::Refraction;
+
+		case EDMMaterialPropertyType::Tangent:
+			return EDMTextureSetMaterialProperty::Tangent;
+
+		default:
+			return EDMTextureSetMaterialProperty::None;
+	}
+}
+
+EDMMaterialPropertyType FDMUtils::TextureSetMaterialPropertyToMaterialPropertyType(EDMTextureSetMaterialProperty InPropertyType)
+{
+	switch (InPropertyType)
+	{
+		case EDMTextureSetMaterialProperty::BaseColor:
+			return EDMMaterialPropertyType::BaseColor;
+
+		case EDMTextureSetMaterialProperty::EmissiveColor:
+			return EDMMaterialPropertyType::EmissiveColor;
+
+		case EDMTextureSetMaterialProperty::Opacity:
+			return EDMMaterialPropertyType::Opacity;
+
+		case EDMTextureSetMaterialProperty::OpacityMask:
+			return EDMMaterialPropertyType::OpacityMask;
+
+		case EDMTextureSetMaterialProperty::Roughness:
+			return EDMMaterialPropertyType::Roughness;
+
+		case EDMTextureSetMaterialProperty::Specular:
+			return EDMMaterialPropertyType::Specular;
+
+		case EDMTextureSetMaterialProperty::Metallic:
+			return EDMMaterialPropertyType::Metallic;
+
+		case EDMTextureSetMaterialProperty::Normal:
+			return EDMMaterialPropertyType::Normal;
+
+		case EDMTextureSetMaterialProperty::PixelDepthOffset:
+			return EDMMaterialPropertyType::PixelDepthOffset;
+
+		case EDMTextureSetMaterialProperty::WorldPositionOffset:
+			return EDMMaterialPropertyType::WorldPositionOffset;
+
+		case EDMTextureSetMaterialProperty::AmbientOcclusion:
+			return EDMMaterialPropertyType::AmbientOcclusion;
+
+		case EDMTextureSetMaterialProperty::Anisotropy:
+			return EDMMaterialPropertyType::Anisotropy;
+
+		case EDMTextureSetMaterialProperty::Refraction:
+			return EDMMaterialPropertyType::Refraction;
+
+		case EDMTextureSetMaterialProperty::Tangent:
+			return EDMMaterialPropertyType::Tangent;
+
+		default:
+			return EDMMaterialPropertyType::None;
+	}
+}
+
+EMaterialProperty FDMUtils::MaterialPropertyTypeToMaterialProperty(EDMMaterialPropertyType InPropertyType)
+{
+	switch (InPropertyType)
+	{
+		case EDMMaterialPropertyType::BaseColor:
+			return EMaterialProperty::MP_BaseColor;
+
+		case EDMMaterialPropertyType::EmissiveColor:
+			return EMaterialProperty::MP_EmissiveColor;
+
+		case EDMMaterialPropertyType::Opacity:
+			return EMaterialProperty::MP_Opacity;
+
+		case EDMMaterialPropertyType::OpacityMask:
+			return EMaterialProperty::MP_OpacityMask;
+
+		case EDMMaterialPropertyType::Roughness:
+			return EMaterialProperty::MP_Roughness;
+
+		case EDMMaterialPropertyType::Specular:
+			return EMaterialProperty::MP_Specular;
+
+		case EDMMaterialPropertyType::Metallic:
+			return EMaterialProperty::MP_Metallic;
+
+		case EDMMaterialPropertyType::Normal:
+			return EMaterialProperty::MP_Normal;
+
+		case EDMMaterialPropertyType::PixelDepthOffset:
+			return EMaterialProperty::MP_PixelDepthOffset;
+
+		case EDMMaterialPropertyType::WorldPositionOffset:
+			return EMaterialProperty::MP_WorldPositionOffset;
+
+		case EDMMaterialPropertyType::AmbientOcclusion:
+			return EMaterialProperty::MP_AmbientOcclusion;
+
+		case EDMMaterialPropertyType::Anisotropy:
+			return EMaterialProperty::MP_Anisotropy;
+
+		case EDMMaterialPropertyType::Refraction:
+			return EMaterialProperty::MP_Refraction;
+
+		case EDMMaterialPropertyType::Tangent:
+			return EMaterialProperty::MP_Tangent;
+
+		default:
+			return EMaterialProperty::MP_MAX;
+	}
+}
+
+EDMMaterialPropertyType FDMUtils::MaterialPropertyToMaterialPropertyType(EMaterialProperty InPropertyType)
+{
+	switch (InPropertyType)
+	{
+		case EMaterialProperty::MP_BaseColor:
+			return EDMMaterialPropertyType::BaseColor;
+
+		case EMaterialProperty::MP_EmissiveColor:
+			return EDMMaterialPropertyType::EmissiveColor;
+
+		case EMaterialProperty::MP_Opacity:
+			return EDMMaterialPropertyType::Opacity;
+
+		case EMaterialProperty::MP_OpacityMask:
+			return EDMMaterialPropertyType::OpacityMask;
+
+		case EMaterialProperty::MP_Roughness:
+			return EDMMaterialPropertyType::Roughness;
+
+		case EMaterialProperty::MP_Specular:
+			return EDMMaterialPropertyType::Specular;
+
+		case EMaterialProperty::MP_Metallic:
+			return EDMMaterialPropertyType::Metallic;
+
+		case EMaterialProperty::MP_Normal:
+			return EDMMaterialPropertyType::Normal;
+
+		case EMaterialProperty::MP_PixelDepthOffset:
+			return EDMMaterialPropertyType::PixelDepthOffset;
+
+		case EMaterialProperty::MP_WorldPositionOffset:
+			return EDMMaterialPropertyType::WorldPositionOffset;
+
+		case EMaterialProperty::MP_AmbientOcclusion:
+			return EDMMaterialPropertyType::AmbientOcclusion;
+
+		case EMaterialProperty::MP_Anisotropy:
+			return EDMMaterialPropertyType::Anisotropy;
+
+		case EMaterialProperty::MP_Refraction:
+			return EDMMaterialPropertyType::Refraction;
+
+		case EMaterialProperty::MP_Tangent:
+			return EDMMaterialPropertyType::Tangent;
+
+		default:
+			return EDMMaterialPropertyType::None;
+	}
+}
diff --git a/Engine/Plugins/Experimental/DynamicMaterial/Source/DynamicMaterial/Public/DMDefs.h b/Engine/Plugins/Experimental/DynamicMaterial/Source/DynamicMaterial/Public/DMDefs.h
--- a/Engine/Plugins/Experimental/DynamicMaterial/Source/DynamicMaterial/Public/DMDefs.h
+++ b/Engine/Plugins/Experimental/DynamicMaterial/Source/DynamicMaterial/Public/DMDefs.h
@@ -7,9 +7,14 @@
 #include "Engine/EngineTypes.h"
 #include "DMDefs.generated.h"
 
-class UTexture;
 class UDMMaterialStage;
-enum class EDMTextureSetMaterialProperty : uint8;
+class UTexture;
+
+namespace UE::DynamicMaterial
+{
+	constexpr int32 RenameFlags = REN_DontCreateRedirectors | REN_DoNotDirty | REN_ForceNoResetLoaders | REN_NonTransactional;
+	constexpr EClassFlags InvalidClassFlags = CLASS_Abstract | CLASS_Hidden | CLASS_Deprecated | CLASS_NewerVersionExists;
+}
 
 /** Some short names provided so they fit nicely on buttons! */
 UENUM(BlueprintType)
@@ -168,19 +173,3 @@ public:
 private:
 	static uint32 GuardCount;
 };
-
-namespace UE::DynamicMaterial
-{
-	constexpr int32 RenameFlags = REN_DontCreateRedirectors | REN_DoNotDirty | REN_ForceNoResetLoaders | REN_NonTransactional;
-	constexpr EClassFlags InvalidClassFlags = CLASS_Abstract | CLASS_Hidden | CLASS_Deprecated | CLASS_NewerVersionExists;
-
-	/** Designed to be used with preprocessor macros. See below. */
-	DYNAMICMATERIAL_API FString CreateNodeComment(const ANSICHAR* InFile, int InLine, const ANSICHAR* InFunction, const FString* InComment = nullptr);
-
-	DYNAMICMATERIAL_API EDMTextureSetMaterialProperty MaterialPropertyTypeToMaterialProperty(EDMMaterialPropertyType InPropertyType);
-
-	DYNAMICMATERIAL_API EDMMaterialPropertyType MaterialPropertyToMaterialPropertyType(EDMTextureSetMaterialProperty InPropertyType);
-}
-
-#define UE_DM_NodeComment_Default UE::DynamicMaterial::CreateNodeComment(__FILE__, __LINE__, __FUNCTION__)
-#define UE_DM_NodeComment(Comment) UE::DynamicMaterial::CreateNodeComment(__FILE__, __LINE__, __FUNCTION__, &Comment)
diff --git a/Engine/Plugins/Experimental/DynamicMaterial/Source/DynamicMaterial/Public/Utils/DMUtils.h b/Engine/Plugins/Experimental/DynamicMaterial/Source/DynamicMaterial/Public/Utils/DMUtils.h
new file mode 100644
--- /dev/null
+++ b/Engine/Plugins/Experimental/DynamicMaterial/Source/DynamicMaterial/Public/Utils/DMUtils.h
@@ -0,0 +1,31 @@
+// Copyright Epic Games, Inc. All Rights Reserved.
+
+#pragma once
+
+#include "HAL/Platform.h"
+#include "Containers/UnrealString.h"
+
+enum class EDMMaterialPropertyType : uint8;
+enum class EDMTextureSetMaterialProperty : uint8;
+enum EMaterialProperty : int;
+
+struct DYNAMICMATERIAL_API FDMUtils
+{
+#if WITH_EDITOR
+	/** Designed to be used with preprocessor macros. */
+	static FString CreateNodeComment(const ANSICHAR* InFile, int InLine, const ANSICHAR* InFunction, const FString* InComment = nullptr);
+#endif
+
+	static EDMTextureSetMaterialProperty MaterialPropertyTypeToTextureSetMaterialProperty(EDMMaterialPropertyType InPropertyType);
+
+	static EDMMaterialPropertyType TextureSetMaterialPropertyToMaterialPropertyType(EDMTextureSetMaterialProperty InPropertyType);
+
+	static EMaterialProperty MaterialPropertyTypeToMaterialProperty(EDMMaterialPropertyType InPropertyType);
+
+	static EDMMaterialPropertyType MaterialPropertyToMaterialPropertyType(EMaterialProperty InPropertyType);
+};
+
+#if WITH_EDITOR
+#define UE_DM_NodeComment_Default FDMUtils::CreateNodeComment(__FILE__, __LINE__, __FUNCTION__)
+#define UE_DM_NodeComment(Comment) FDMUtils::CreateNodeComment(__FILE__, __LINE__, __FUNCTION__, &Comment)
+#endif
diff --git a/Engine/Plugins/Experimental/DynamicMaterial/Source/DynamicMaterialEditor/Private/Components/DMMaterialEffectFunction.cpp b/Engine/Plugins/Experimental/DynamicMaterial/Source/DynamicMaterialEditor/Private/Components/DMMaterialEffectFunction.cpp
--- a/Engine/Plugins/Experimental/DynamicMaterial/Source/DynamicMaterialEditor/Private/Components/DMMaterialEffectFunction.cpp
+++ b/Engine/Plugins/Experimental/DynamicMaterial/Source/DynamicMaterialEditor/Private/Components/DMMaterialEffectFunction.cpp
@@ -23,6 +23,7 @@
 #include "Model/DynamicMaterialModelEditorOnlyData.h"
 #include "Utils/DMMaterialFunctionLibrary.h"
 #include "Utils/DMPrivate.h"
+#include "Utils/DMUtils.h"
 
 #define LOCTEXT_NAMESPACE "DMMaterialEffectFunction"
 
diff --git a/Engine/Plugins/Experimental/DynamicMaterial/Source/DynamicMaterialEditor/Private/Components/DMMaterialProperty.cpp b/Engine/Plugins/Experimental/DynamicMaterial/Source/DynamicMaterialEditor/Private/Components/DMMaterialProperty.cpp
--- a/Engine/Plugins/Experimental/DynamicMaterial/Source/DynamicMaterialEditor/Private/Components/DMMaterialProperty.cpp
+++ b/Engine/Plugins/Experimental/DynamicMaterial/Source/DynamicMaterialEditor/Private/Components/DMMaterialProperty.cpp
@@ -34,7 +34,9 @@
 #include "Model/DynamicMaterialModel.h"
 #include "Model/DynamicMaterialModelEditorOnlyData.h"
 #include "Utils/DMMaterialFunctionLibrary.h"
+#include "Utils/DMMaterialUtils.h"
 #include "Utils/DMPrivate.h"
+#include "Utils/DMUtils.h"
 
 #define LOCTEXT_NAMESPACE "DMMaterialProperty"
 
@@ -328,12 +330,7 @@ void UDMMaterialProperty::GenerateExpressions(const TSharedRef<FDMMaterialBuildS
 {
 	UDynamicMaterialModelEditorOnlyData* EditorOnlyData = GetMaterialModelEditorOnlyData();
 
-	if (!EditorOnlyData)
-	{
-		return;
-	}
-
-	if (EditorOnlyData->GetDomain() == EMaterialDomain::MD_PostProcess && MaterialProperty != EDMMaterialPropertyType::EmissiveColor)
+	if (!EditorOnlyData || !IsValidForModel(*EditorOnlyData))
 	{
 		return;
 	}
@@ -678,6 +675,24 @@ void UDMMaterialProperty::SetOutputProcessor(UMaterialFunctionInterface* InFunct
 	OnOutputProcessorUpdated();
 }
 
+bool UDMMaterialProperty::IsValidForModel(UDynamicMaterialModelEditorOnlyData& InMaterialModel) const
+{
+	return FDMMaterialUtils::IsMaterialPropertyActive({
+		FDMUtils::MaterialPropertyTypeToMaterialProperty(MaterialProperty),
+		InMaterialModel.GetDomain(),
+		InMaterialModel.GetBlendMode(),
+		InMaterialModel.GetShadingModel() == EDMMaterialShadingModel::DefaultLit ? EMaterialShadingModel::MSM_DefaultLit : EMaterialShadingModel::MSM_Unlit,
+		TLM_Surface,
+		/* Tesselation enabled */ false,
+		/* BlendableOutputAlpha (Post Process Alpha) */ false,
+		/* Uses Distortion */ false,
+		/* Shading model from master material */ false,
+		/* Outputting translucency velocity */ InMaterialModel.GetBlendMode() != BLEND_Opaque,
+		/* Thin surface */ false,
+		/* Is supported (substate check) */ true
+	});
+}
+
 void UDMMaterialProperty::OnOutputProcessorUpdated()
 {
 	if (!OutputProcessor)
diff --git a/Engine/Plugins/Experimental/DynamicMaterial/Source/DynamicMaterialEditor/Private/Components/DMMaterialSlot.cpp b/Engine/Plugins/Experimental/DynamicMaterial/Source/DynamicMaterialEditor/Private/Components/DMMaterialSlot.cpp
--- a/Engine/Plugins/Experimental/DynamicMaterial/Source/DynamicMaterialEditor/Private/Components/DMMaterialSlot.cpp
+++ b/Engine/Plugins/Experimental/DynamicMaterial/Source/DynamicMaterialEditor/Private/Components/DMMaterialSlot.cpp
@@ -24,6 +24,7 @@
 #include "UObject/Package.h"
 #include "UObject/UObjectGlobals.h"
 #include "Utils/DMPrivate.h"
+#include "Utils/DMUtils.h"
 
 #define LOCTEXT_NAMESPACE "DMMaterialStage"
 
diff --git a/Engine/Plugins/Experimental/DynamicMaterial/Source/DynamicMaterialEditor/Private/Components/DMMaterialStageBlend.cpp b/Engine/Plugins/Experimental/DynamicMaterial/Source/DynamicMaterialEditor/Private/Components/DMMaterialStageBlend.cpp
--- a/Engine/Plugins/Experimental/DynamicMaterial/Source/DynamicMaterialEditor/Private/Components/DMMaterialStageBlend.cpp
+++ b/Engine/Plugins/Experimental/DynamicMaterial/Source/DynamicMaterialEditor/Private/Components/DMMaterialStageBlend.cpp
@@ -467,16 +467,36 @@ void UDMMaterialStageBlend::AddDefaultInput(int32 InInputIndex) const
 					0, 
 					FDMMaterialStageConnectorChannel::WHOLE_CHANNEL
 				);
-				break;
 			}
-			
-			Stage->ChangeInput_PreviousStage(
-				InInputIndex, 
-				FDMMaterialStageConnectorChannel::WHOLE_CHANNEL, 
-				EDMMaterialPropertyType::EmissiveColor,
-				0, 
-				FDMMaterialStageConnectorChannel::WHOLE_CHANNEL
-			);
+			else
+			{
+				EDMMaterialPropertyType DefaultProperty = EDMMaterialPropertyType::None;
+
+				if (UDMMaterialSlot* Slot = Layer->GetSlot())
+				{
+					if (UDynamicMaterialModelEditorOnlyData* ModelEditorOnlyData = Slot->GetMaterialModelEditorOnlyData())
+					{
+						if (ModelEditorOnlyData->GetSlotForMaterialProperty(EDMMaterialPropertyType::BaseColor))
+						{
+							DefaultProperty = EDMMaterialPropertyType::BaseColor;
+						}
+						else if (ModelEditorOnlyData->GetSlotForMaterialProperty(EDMMaterialPropertyType::EmissiveColor))
+						{
+							DefaultProperty = EDMMaterialPropertyType::EmissiveColor;
+						}
+					}
+				}
+
+				Stage->ChangeInput_PreviousStage(
+					InInputIndex,
+					FDMMaterialStageConnectorChannel::WHOLE_CHANNEL,
+					DefaultProperty,
+					0,
+					FDMMaterialStageConnectorChannel::WHOLE_CHANNEL
+				);
+
+			}
+
 			break;
 		}
  
diff --git a/Engine/Plugins/Experimental/DynamicMaterial/Source/DynamicMaterialEditor/Private/Components/DMMaterialStageBlendFunction.cpp b/Engine/Plugins/Experimental/DynamicMaterial/Source/DynamicMaterialEditor/Private/Components/DMMaterialStageBlendFunction.cpp
--- a/Engine/Plugins/Experimental/DynamicMaterial/Source/DynamicMaterialEditor/Private/Components/DMMaterialStageBlendFunction.cpp
+++ b/Engine/Plugins/Experimental/DynamicMaterial/Source/DynamicMaterialEditor/Private/Components/DMMaterialStageBlendFunction.cpp
@@ -8,6 +8,7 @@
 #include "Model/DMMaterialBuildState.h"
 #include "Model/DMMaterialBuildUtils.h"
 #include "Utils/DMMaterialFunctionLibrary.h"
+#include "Utils/DMUtils.h"
 
 #define LOCTEXT_NAMESPACE "DMMaterialStageBlendFunction"
 
diff --git a/Engine/Plugins/Experimental/DynamicMaterial/Source/DynamicMaterialEditor/Private/Components/DMMaterialStageExpression.cpp b/Engine/Plugins/Experimental/DynamicMaterial/Source/DynamicMaterialEditor/Private/Components/DMMaterialStageExpression.cpp
--- a/Engine/Plugins/Experimental/DynamicMaterial/Source/DynamicMaterialEditor/Private/Components/DMMaterialStageExpression.cpp
+++ b/Engine/Plugins/Experimental/DynamicMaterial/Source/DynamicMaterialEditor/Private/Components/DMMaterialStageExpression.cpp
@@ -9,6 +9,7 @@
 #include "Materials/MaterialExpressionComponentMask.h"
 #include "Model/DMMaterialBuildState.h"
 #include "Model/DMMaterialBuildUtils.h"
+#include "Utils/DMUtils.h"
 
 // Expressions to exclude
 #include "Materials/MaterialExpressionComment.h"
diff --git a/Engine/Plugins/Experimental/DynamicMaterial/Source/DynamicMaterialEditor/Private/Components/DMMaterialStageFunction.cpp b/Engine/Plugins/Experimental/DynamicMaterial/Source/DynamicMaterialEditor/Private/Components/DMMaterialStageFunction.cpp
--- a/Engine/Plugins/Experimental/DynamicMaterial/Source/DynamicMaterialEditor/Private/Components/DMMaterialStageFunction.cpp
+++ b/Engine/Plugins/Experimental/DynamicMaterial/Source/DynamicMaterialEditor/Private/Components/DMMaterialStageFunction.cpp
@@ -23,6 +23,7 @@
 #include "Model/DMMaterialBuildUtils.h"
 #include "Model/DynamicMaterialModelEditorOnlyData.h"
 #include "Utils/DMPrivate.h"
+#include "Utils/DMUtils.h"
 
 #define LOCTEXT_NAMESPACE "DMMaterialStageFunction"
 
@@ -152,11 +153,28 @@ void UDMMaterialStageFunction::AddDefaultInput(int32 InInputIndex) const
 	}
 	else
 	{
+		EDMMaterialPropertyType DefaultProperty = EDMMaterialPropertyType::None;
+
+		if (UDMMaterialSlot* Slot = Layer->GetSlot())
+		{
+			if (UDynamicMaterialModelEditorOnlyData* ModelEditorOnlyData = Slot->GetMaterialModelEditorOnlyData())
+			{
+				if (ModelEditorOnlyData->GetSlotForMaterialProperty(EDMMaterialPropertyType::BaseColor))
+				{
+					DefaultProperty = EDMMaterialPropertyType::BaseColor;
+				}
+				else if (ModelEditorOnlyData->GetSlotForMaterialProperty(EDMMaterialPropertyType::EmissiveColor))
+				{
+					DefaultProperty = EDMMaterialPropertyType::EmissiveColor;
+				}
+			}
+		}		
+
 		Stage->ChangeInput_PreviousStage(
-			InInputIndex, 
-			FDMMaterialStageConnectorChannel::WHOLE_CHANNEL, 
-			EDMMaterialPropertyType::EmissiveColor,
-			0, 
+			InInputIndex,
+			FDMMaterialStageConnectorChannel::WHOLE_CHANNEL,
+			DefaultProperty,
+			0,
 			FDMMaterialStageConnectorChannel::WHOLE_CHANNEL
 		);
 	}
diff --git a/Engine/Plugins/Experimental/DynamicMaterial/Source/DynamicMaterialEditor/Private/Components/DMMaterialStageSource.cpp b/Engine/Plugins/Experimental/DynamicMaterial/Source/DynamicMaterialEditor/Private/Components/DMMaterialStageSource.cpp
--- a/Engine/Plugins/Experimental/DynamicMaterial/Source/DynamicMaterialEditor/Private/Components/DMMaterialStageSource.cpp
+++ b/Engine/Plugins/Experimental/DynamicMaterial/Source/DynamicMaterialEditor/Private/Components/DMMaterialStageSource.cpp
@@ -248,9 +248,6 @@ bool UDMMaterialStageSource::UpdateStagePreviewMaterial(UDMMaterialStage* InStag
 	UDynamicMaterialModelEditorOnlyData* ModelEditorOnlyData = Slot->GetMaterialModelEditorOnlyData();
 	check(ModelEditorOnlyData);
 
-	UDMMaterialProperty* PropertyObj = ModelEditorOnlyData->GetMaterialProperty(EDMMaterialPropertyType::EmissiveColor);
-	check(PropertyObj);
-
 	TSharedRef<FDMMaterialBuildState> BuildState = ModelEditorOnlyData->CreateBuildState(InPreviewMaterial);
 	BuildState->SetPreviewMaterial();
 
@@ -280,9 +277,7 @@ bool UDMMaterialStageSource::UpdateStagePreviewMaterial(UDMMaterialStage* InStag
 
 	int32 BestMatch = INDEX_NONE;
 	int32 OutputCount = 0;
-	const int32 FloatsForPropertyType = bIsMaskStage
-		? 1
-		: UDMValueDefinitionLibrary::GetValueDefinition(PropertyObj->GetInputConnectorType()).GetFloatCount();
+	const int32 FloatsForPropertyType = bIsMaskStage ? 1 : 3;
 
 	for (int32 OutputIdx = 0; OutputIdx < LastExpression->GetOutputs().Num(); ++OutputIdx)
 	{
diff --git a/Engine/Plugins/Experimental/DynamicMaterial/Source/DynamicMaterialEditor/Private/Components/DMMaterialStageThroughput.cpp b/Engine/Plugins/Experimental/DynamicMaterial/Source/DynamicMaterialEditor/Private/Components/DMMaterialStageThroughput.cpp
--- a/Engine/Plugins/Experimental/DynamicMaterial/Source/DynamicMaterialEditor/Private/Components/DMMaterialStageThroughput.cpp
+++ b/Engine/Plugins/Experimental/DynamicMaterial/Source/DynamicMaterialEditor/Private/Components/DMMaterialStageThroughput.cpp
@@ -25,6 +25,7 @@
 #include "Utils/DMInputNodeBuilder.h"
 #include "Utils/DMMaterialFunctionLibrary.h"
 #include "Utils/DMPrivate.h"
+#include "Utils/DMUtils.h"
 
 #define LOCTEXT_NAMESPACE "DMMaterialStageThroughput"
 
diff --git a/Engine/Plugins/Experimental/DynamicMaterial/Source/DynamicMaterialEditor/Private/Components/DMMaterialStageThroughputLayerBlend.cpp b/Engine/Plugins/Experimental/DynamicMaterial/Source/DynamicMaterialEditor/Private/Components/DMMaterialStageThroughputLayerBlend.cpp
--- a/Engine/Plugins/Experimental/DynamicMaterial/Source/DynamicMaterialEditor/Private/Components/DMMaterialStageThroughputLayerBlend.cpp
+++ b/Engine/Plugins/Experimental/DynamicMaterial/Source/DynamicMaterialEditor/Private/Components/DMMaterialStageThroughputLayerBlend.cpp
@@ -25,6 +25,7 @@
 #include "Utils/DMInputNodeBuilder.h"
 #include "Utils/DMMaterialFunctionLibrary.h"
 #include "Utils/DMPrivate.h"
+#include "Utils/DMUtils.h"
 
 #define LOCTEXT_NAMESPACE "DMMaterialStageThroughputLayerBlend"
 
@@ -328,7 +329,7 @@ void UDMMaterialStageThroughputLayerBlend::GenerateMainExpressions(const TShared
 
 	check(GetAlphaBlend());
 
-	UMaterialExpressionMaterialFunctionCall* FunctionCall = InBuildState->GetBuildUtils().CreateExpression<UMaterialExpressionMaterialFunctionCall>(UE_DM_NodeComment_Default);;
+	UMaterialExpressionMaterialFunctionCall* FunctionCall = InBuildState->GetBuildUtils().CreateExpression<UMaterialExpressionMaterialFunctionCall>(UE_DM_NodeComment_Default);
 	FunctionCall->SetMaterialFunction(GetAlphaBlend());
 	FunctionCall->UpdateFromFunctionResource();
 
@@ -363,17 +364,15 @@ void UDMMaterialStageThroughputLayerBlend::GenerateMainExpressions(const TShared
 
 					if (ModelEditorOnlyData)
 					{
-						switch (ModelEditorOnlyData->GetShadingModel())
+						if (UDMMaterialSlot* BaseColorSlot = ModelEditorOnlyData->GetSlotForMaterialProperty(EDMMaterialPropertyType::BaseColor))
 						{
-							case EDMMaterialShadingModel::DefaultLit:
-								RGBProperty = EDMMaterialPropertyType::BaseColor;
-								RGBSlot = ModelEditorOnlyData->GetSlotForMaterialProperty(RGBProperty);
-								break;
-
-							case EDMMaterialShadingModel::Unlit:
-								RGBProperty = EDMMaterialPropertyType::EmissiveColor;
-								RGBSlot = ModelEditorOnlyData->GetSlotForMaterialProperty(RGBProperty);
-								break;
+							RGBProperty = EDMMaterialPropertyType::BaseColor;
+							RGBSlot = BaseColorSlot;
+						}
+						else if (UDMMaterialSlot* EmissiveColorSlot = ModelEditorOnlyData->GetSlotForMaterialProperty(EDMMaterialPropertyType::EmissiveColor))
+						{
+							RGBProperty = EDMMaterialPropertyType::EmissiveColor;
+							RGBSlot = EmissiveColorSlot;
 						}
 					}
 				}
diff --git a/Engine/Plugins/Experimental/DynamicMaterial/Source/DynamicMaterialEditor/Private/Components/MaterialProperties/DMMPAmbientOcclusion.cpp b/Engine/Plugins/Experimental/DynamicMaterial/Source/DynamicMaterialEditor/Private/Components/MaterialProperties/DMMPAmbientOcclusion.cpp
--- a/Engine/Plugins/Experimental/DynamicMaterial/Source/DynamicMaterialEditor/Private/Components/MaterialProperties/DMMPAmbientOcclusion.cpp
+++ b/Engine/Plugins/Experimental/DynamicMaterial/Source/DynamicMaterialEditor/Private/Components/MaterialProperties/DMMPAmbientOcclusion.cpp
@@ -10,11 +10,6 @@ UDMMaterialPropertyAmbientOcclusion::UDMMaterialPropertyAmbientOcclusion()
 {
 }
 
-bool UDMMaterialPropertyAmbientOcclusion::IsValidForModel(UDynamicMaterialModelEditorOnlyData& InModelEditorOnlyData) const
-{
-	return (InModelEditorOnlyData.GetShadingModel() != EDMMaterialShadingModel::Unlit);
-}
-
 UMaterialExpression* UDMMaterialPropertyAmbientOcclusion::GetDefaultInput(
 	const TSharedRef<FDMMaterialBuildState>& InBuildState) const
 {
diff --git a/Engine/Plugins/Experimental/DynamicMaterial/Source/DynamicMaterialEditor/Private/Components/MaterialProperties/DMMPAmbientOcclusion.h b/Engine/Plugins/Experimental/DynamicMaterial/Source/DynamicMaterialEditor/Private/Components/MaterialProperties/DMMPAmbientOcclusion.h
--- a/Engine/Plugins/Experimental/DynamicMaterial/Source/DynamicMaterialEditor/Private/Components/MaterialProperties/DMMPAmbientOcclusion.h
+++ b/Engine/Plugins/Experimental/DynamicMaterial/Source/DynamicMaterialEditor/Private/Components/MaterialProperties/DMMPAmbientOcclusion.h
@@ -16,7 +16,6 @@ public:
 	UDMMaterialPropertyAmbientOcclusion();
 
 	//~ Begin UDMMaterialProperty
-	virtual bool IsValidForModel(UDynamicMaterialModelEditorOnlyData& InModelEditorOnlyData) const override;
 	virtual UMaterialExpression* GetDefaultInput(const TSharedRef<FDMMaterialBuildState>& InBuildState) const override;
 	//~ End UDMMaterialProperty
 };
diff --git a/Engine/Plugins/Experimental/DynamicMaterial/Source/DynamicMaterialEditor/Private/Components/MaterialProperties/DMMPAnisotropy.cpp b/Engine/Plugins/Experimental/DynamicMaterial/Source/DynamicMaterialEditor/Private/Components/MaterialProperties/DMMPAnisotropy.cpp
--- a/Engine/Plugins/Experimental/DynamicMaterial/Source/DynamicMaterialEditor/Private/Components/MaterialProperties/DMMPAnisotropy.cpp
+++ b/Engine/Plugins/Experimental/DynamicMaterial/Source/DynamicMaterialEditor/Private/Components/MaterialProperties/DMMPAnisotropy.cpp
@@ -10,24 +10,6 @@ UDMMaterialPropertyAnisotropy::UDMMaterialPropertyAnisotropy()
 {
 }
 
-bool UDMMaterialPropertyAnisotropy::IsValidForModel(UDynamicMaterialModelEditorOnlyData& InModelEditorOnlyData) const
-{
-	if (InModelEditorOnlyData.GetShadingModel() == EDMMaterialShadingModel::Unlit)
-	{
-		return false;
-	}
-
-	switch (InModelEditorOnlyData.GetBlendMode())
-	{
-		case EBlendMode::BLEND_Opaque:
-		case EBlendMode::BLEND_Masked:
-			return true;
-
-		default:
-			return false;
-	}
-}
-
 UMaterialExpression* UDMMaterialPropertyAnisotropy::GetDefaultInput(const TSharedRef<FDMMaterialBuildState>& InBuildState) const
 {
 	return CreateConstant(InBuildState, 1.f);
diff --git a/Engine/Plugins/Experimental/DynamicMaterial/Source/DynamicMaterialEditor/Private/Components/MaterialProperties/DMMPAnisotropy.h b/Engine/Plugins/Experimental/DynamicMaterial/Source/DynamicMaterialEditor/Private/Components/MaterialProperties/DMMPAnisotropy.h
--- a/Engine/Plugins/Experimental/DynamicMaterial/Source/DynamicMaterialEditor/Private/Components/MaterialProperties/DMMPAnisotropy.h
+++ b/Engine/Plugins/Experimental/DynamicMaterial/Source/DynamicMaterialEditor/Private/Components/MaterialProperties/DMMPAnisotropy.h
@@ -16,7 +16,6 @@ public:
 	UDMMaterialPropertyAnisotropy();
 
 	//~ Begin UDMMaterialProperty
-	virtual bool IsValidForModel(UDynamicMaterialModelEditorOnlyData& InModelEditorOnlyData) const override;
 	virtual UMaterialExpression* GetDefaultInput(const TSharedRef<FDMMaterialBuildState>& InBuildState) const override;
 	//~ End UDMMaterialProperty
 };
diff --git a/Engine/Plugins/Experimental/DynamicMaterial/Source/DynamicMaterialEditor/Private/Components/MaterialProperties/DMMPBaseColor.cpp b/Engine/Plugins/Experimental/DynamicMaterial/Source/DynamicMaterialEditor/Private/Components/MaterialProperties/DMMPBaseColor.cpp
--- a/Engine/Plugins/Experimental/DynamicMaterial/Source/DynamicMaterialEditor/Private/Components/MaterialProperties/DMMPBaseColor.cpp
+++ b/Engine/Plugins/Experimental/DynamicMaterial/Source/DynamicMaterialEditor/Private/Components/MaterialProperties/DMMPBaseColor.cpp
@@ -11,11 +11,6 @@ UDMMaterialPropertyBaseColor::UDMMaterialPropertyBaseColor()
 {
 }
 
-bool UDMMaterialPropertyBaseColor::IsValidForModel(UDynamicMaterialModelEditorOnlyData& InModelEditorOnlyData) const
-{
-	return (InModelEditorOnlyData.GetShadingModel() != EDMMaterialShadingModel::Unlit);
-}
-
 UMaterialExpression* UDMMaterialPropertyBaseColor::GetDefaultInput(const TSharedRef<FDMMaterialBuildState>& InBuildState) const
 {
 	return CreateConstant(InBuildState, FVector::OneVector);
diff --git a/Engine/Plugins/Experimental/DynamicMaterial/Source/DynamicMaterialEditor/Private/Components/MaterialProperties/DMMPBaseColor.h b/Engine/Plugins/Experimental/DynamicMaterial/Source/DynamicMaterialEditor/Private/Components/MaterialProperties/DMMPBaseColor.h
--- a/Engine/Plugins/Experimental/DynamicMaterial/Source/DynamicMaterialEditor/Private/Components/MaterialProperties/DMMPBaseColor.h
+++ b/Engine/Plugins/Experimental/DynamicMaterial/Source/DynamicMaterialEditor/Private/Components/MaterialProperties/DMMPBaseColor.h
@@ -16,7 +16,6 @@ public:
 	UDMMaterialPropertyBaseColor();
 
 	//~ Begin UDMMaterialProperty
-	virtual bool IsValidForModel(UDynamicMaterialModelEditorOnlyData& InModelEditorOnlyData) const override;
 	virtual UMaterialExpression* GetDefaultInput(const TSharedRef<FDMMaterialBuildState>& InBuildState) const override;
 	virtual TEnumAsByte<EMaterialSamplerType> GetTextureSamplerType() const override;
 	//~ End UDMMaterialProperty
diff --git a/Engine/Plugins/Experimental/DynamicMaterial/Source/DynamicMaterialEditor/Private/Components/MaterialProperties/DMMPEmissiveColor.cpp b/Engine/Plugins/Experimental/DynamicMaterial/Source/DynamicMaterialEditor/Private/Components/MaterialProperties/DMMPEmissiveColor.cpp
--- a/Engine/Plugins/Experimental/DynamicMaterial/Source/DynamicMaterialEditor/Private/Components/MaterialProperties/DMMPEmissiveColor.cpp
+++ b/Engine/Plugins/Experimental/DynamicMaterial/Source/DynamicMaterialEditor/Private/Components/MaterialProperties/DMMPEmissiveColor.cpp
@@ -12,18 +12,6 @@ UDMMaterialPropertyEmissiveColor::UDMMaterialPropertyEmissiveColor()
 {
 }
 
-bool UDMMaterialPropertyEmissiveColor::IsValidForModel(UDynamicMaterialModelEditorOnlyData& InModelEditorOnlyData) const
-{
-	switch (InModelEditorOnlyData.GetBlendMode())
-	{
-		case EBlendMode::BLEND_AlphaHoldout:
-			return false;
-
-		default:
-			return true;
-	}
-}
-
 UMaterialExpression* UDMMaterialPropertyEmissiveColor::GetDefaultInput(
 	const TSharedRef<FDMMaterialBuildState>& InBuildState) const
 {
diff --git a/Engine/Plugins/Experimental/DynamicMaterial/Source/DynamicMaterialEditor/Private/Components/MaterialProperties/DMMPEmissiveColor.h b/Engine/Plugins/Experimental/DynamicMaterial/Source/DynamicMaterialEditor/Private/Components/MaterialProperties/DMMPEmissiveColor.h
--- a/Engine/Plugins/Experimental/DynamicMaterial/Source/DynamicMaterialEditor/Private/Components/MaterialProperties/DMMPEmissiveColor.h
+++ b/Engine/Plugins/Experimental/DynamicMaterial/Source/DynamicMaterialEditor/Private/Components/MaterialProperties/DMMPEmissiveColor.h
@@ -16,7 +16,6 @@ public:
 	UDMMaterialPropertyEmissiveColor();
 
 	//~ Begin UDMMaterialProperty
-	virtual bool IsValidForModel(UDynamicMaterialModelEditorOnlyData& InModelEditorOnlyData) const override;
 	virtual UMaterialExpression* GetDefaultInput(const TSharedRef<FDMMaterialBuildState>& InBuildState) const override;
 	virtual TEnumAsByte<EMaterialSamplerType> GetTextureSamplerType() const override;
 	//~ End UDMMaterialProperty
diff --git a/Engine/Plugins/Experimental/DynamicMaterial/Source/DynamicMaterialEditor/Private/Components/MaterialProperties/DMMPMetallic.cpp b/Engine/Plugins/Experimental/DynamicMaterial/Source/DynamicMaterialEditor/Private/Components/MaterialProperties/DMMPMetallic.cpp
--- a/Engine/Plugins/Experimental/DynamicMaterial/Source/DynamicMaterialEditor/Private/Components/MaterialProperties/DMMPMetallic.cpp
+++ b/Engine/Plugins/Experimental/DynamicMaterial/Source/DynamicMaterialEditor/Private/Components/MaterialProperties/DMMPMetallic.cpp
@@ -10,24 +10,6 @@ UDMMaterialPropertyMetallic::UDMMaterialPropertyMetallic()
 {
 }
 
-bool UDMMaterialPropertyMetallic::IsValidForModel(UDynamicMaterialModelEditorOnlyData& InModelEditorOnlyData) const
-{
-	if (InModelEditorOnlyData.GetShadingModel() == EDMMaterialShadingModel::Unlit)
-	{
-		return false;
-	}
-
-	switch (InModelEditorOnlyData.GetBlendMode())
-	{
-		case EBlendMode::BLEND_Opaque:
-		case EBlendMode::BLEND_Masked:
-			return true;
-
-		default:
-			return false;
-	}
-}
-
 UMaterialExpression* UDMMaterialPropertyMetallic::GetDefaultInput(const TSharedRef<FDMMaterialBuildState>& InBuildState) const
 {
 	return CreateConstant(InBuildState, 0.f);
diff --git a/Engine/Plugins/Experimental/DynamicMaterial/Source/DynamicMaterialEditor/Private/Components/MaterialProperties/DMMPMetallic.h b/Engine/Plugins/Experimental/DynamicMaterial/Source/DynamicMaterialEditor/Private/Components/MaterialProperties/DMMPMetallic.h
--- a/Engine/Plugins/Experimental/DynamicMaterial/Source/DynamicMaterialEditor/Private/Components/MaterialProperties/DMMPMetallic.h
+++ b/Engine/Plugins/Experimental/DynamicMaterial/Source/DynamicMaterialEditor/Private/Components/MaterialProperties/DMMPMetallic.h
@@ -16,7 +16,6 @@ public:
 	UDMMaterialPropertyMetallic();
 
 	//~ Begin UDMMaterialProperty
-	virtual bool IsValidForModel(UDynamicMaterialModelEditorOnlyData& InModelEditorOnlyData) const override;
 	virtual UMaterialExpression* GetDefaultInput(const TSharedRef<FDMMaterialBuildState>& InBuildState) const override;
 	//~ End UDMMaterialProperty
 };
diff --git a/Engine/Plugins/Experimental/DynamicMaterial/Source/DynamicMaterialEditor/Private/Components/MaterialProperties/DMMPNormal.cpp b/Engine/Plugins/Experimental/DynamicMaterial/Source/DynamicMaterialEditor/Private/Components/MaterialProperties/DMMPNormal.cpp
--- a/Engine/Plugins/Experimental/DynamicMaterial/Source/DynamicMaterialEditor/Private/Components/MaterialProperties/DMMPNormal.cpp
+++ b/Engine/Plugins/Experimental/DynamicMaterial/Source/DynamicMaterialEditor/Private/Components/MaterialProperties/DMMPNormal.cpp
@@ -7,6 +7,7 @@
 #include "Model/DMMaterialBuildState.h"
 #include "Model/DynamicMaterialModelEditorOnlyData.h"
 #include "Utils/DMMaterialFunctionLibrary.h"
+#include "Utils/DMUtils.h"
 
 namespace UE::DynamicMaterialEditor::Private
 {
@@ -38,24 +39,6 @@ UDMMaterialPropertyNormal::UDMMaterialPropertyNormal()
 {
 }
 
-bool UDMMaterialPropertyNormal::IsValidForModel(UDynamicMaterialModelEditorOnlyData& InModelEditorOnlyData) const
-{
-	if (InModelEditorOnlyData.GetShadingModel() == EDMMaterialShadingModel::Unlit)
-	{
-		return false;
-	}
-
-	switch (InModelEditorOnlyData.GetBlendMode())
-	{
-		case EBlendMode::BLEND_Opaque:
-		case EBlendMode::BLEND_Masked:
-			return true;
-
-		default:
-			return false;
-	}
-}
-
 UMaterialExpression* UDMMaterialPropertyNormal::GetDefaultInput(const TSharedRef<FDMMaterialBuildState>& InBuildState) const
 {
 	return CreateConstant(InBuildState, FVector::OneVector);
diff --git a/Engine/Plugins/Experimental/DynamicMaterial/Source/DynamicMaterialEditor/Private/Components/MaterialProperties/DMMPNormal.h b/Engine/Plugins/Experimental/DynamicMaterial/Source/DynamicMaterialEditor/Private/Components/MaterialProperties/DMMPNormal.h
--- a/Engine/Plugins/Experimental/DynamicMaterial/Source/DynamicMaterialEditor/Private/Components/MaterialProperties/DMMPNormal.h
+++ b/Engine/Plugins/Experimental/DynamicMaterial/Source/DynamicMaterialEditor/Private/Components/MaterialProperties/DMMPNormal.h
@@ -16,7 +16,6 @@ public:
 	UDMMaterialPropertyNormal();
 
 	//~ Begin UDMMaterialProperty
-	virtual bool IsValidForModel(UDynamicMaterialModelEditorOnlyData& InModelEditorOnlyData) const override;
 	virtual UMaterialExpression* GetDefaultInput(const TSharedRef<FDMMaterialBuildState>& InBuildState) const override;
 	virtual TEnumAsByte<EMaterialSamplerType> GetTextureSamplerType() const override;
 
diff --git a/Engine/Plugins/Experimental/DynamicMaterial/Source/DynamicMaterialEditor/Private/Components/MaterialProperties/DMMPOpacity.cpp b/Engine/Plugins/Experimental/DynamicMaterial/Source/DynamicMaterialEditor/Private/Components/MaterialProperties/DMMPOpacity.cpp
--- a/Engine/Plugins/Experimental/DynamicMaterial/Source/DynamicMaterialEditor/Private/Components/MaterialProperties/DMMPOpacity.cpp
+++ b/Engine/Plugins/Experimental/DynamicMaterial/Source/DynamicMaterialEditor/Private/Components/MaterialProperties/DMMPOpacity.cpp
@@ -11,21 +11,6 @@ UDMMaterialPropertyOpacity::UDMMaterialPropertyOpacity()
 {
 }
 
-bool UDMMaterialPropertyOpacity::IsValidForModel(UDynamicMaterialModelEditorOnlyData& InModelEditorOnlyData) const
-{
-	switch (InModelEditorOnlyData.GetBlendMode())
-	{
-		case EBlendMode::BLEND_Translucent:
-		case EBlendMode::BLEND_Additive:
-		case EBlendMode::BLEND_AlphaComposite:
-		case EBlendMode::BLEND_AlphaHoldout:
-			return true;
-
-		default:
-			return false;
-	}
-}
-
 UMaterialExpression* UDMMaterialPropertyOpacity::GetDefaultInput(const TSharedRef<FDMMaterialBuildState>& InBuildState) const
 {
 	return CreateConstant(InBuildState, 1.f);
diff --git a/Engine/Plugins/Experimental/DynamicMaterial/Source/DynamicMaterialEditor/Private/Components/MaterialProperties/DMMPOpacity.h b/Engine/Plugins/Experimental/DynamicMaterial/Source/DynamicMaterialEditor/Private/Components/MaterialProperties/DMMPOpacity.h
--- a/Engine/Plugins/Experimental/DynamicMaterial/Source/DynamicMaterialEditor/Private/Components/MaterialProperties/DMMPOpacity.h
+++ b/Engine/Plugins/Experimental/DynamicMaterial/Source/DynamicMaterialEditor/Private/Components/MaterialProperties/DMMPOpacity.h
@@ -16,7 +16,6 @@ public:
 	UDMMaterialPropertyOpacity();
 
 	//~ Begin UDMMaterialProperty
-	virtual bool IsValidForModel(UDynamicMaterialModelEditorOnlyData& InModelEditorOnlyData) const override;
 	virtual UMaterialExpression* GetDefaultInput(const TSharedRef<FDMMaterialBuildState>& InBuildState) const override;
 	//~ End UDMMaterialProperty
 };
diff --git a/Engine/Plugins/Experimental/DynamicMaterial/Source/DynamicMaterialEditor/Private/Components/MaterialProperties/DMMPOpacityMask.cpp b/Engine/Plugins/Experimental/DynamicMaterial/Source/DynamicMaterialEditor/Private/Components/MaterialProperties/DMMPOpacityMask.cpp
--- a/Engine/Plugins/Experimental/DynamicMaterial/Source/DynamicMaterialEditor/Private/Components/MaterialProperties/DMMPOpacityMask.cpp
+++ b/Engine/Plugins/Experimental/DynamicMaterial/Source/DynamicMaterialEditor/Private/Components/MaterialProperties/DMMPOpacityMask.cpp
@@ -12,11 +12,6 @@ UDMMaterialPropertyOpacityMask::UDMMaterialPropertyOpacityMask()
 {
 }
 
-bool UDMMaterialPropertyOpacityMask::IsValidForModel(UDynamicMaterialModelEditorOnlyData& InModelEditorOnlyData) const
-{
-	return (InModelEditorOnlyData.GetBlendMode() == EBlendMode::BLEND_Masked);
-}
-
 UMaterialExpression* UDMMaterialPropertyOpacityMask::GetDefaultInput(
 	const TSharedRef<FDMMaterialBuildState>& InBuildState) const
 {
diff --git a/Engine/Plugins/Experimental/DynamicMaterial/Source/DynamicMaterialEditor/Private/Components/MaterialProperties/DMMPOpacityMask.h b/Engine/Plugins/Experimental/DynamicMaterial/Source/DynamicMaterialEditor/Private/Components/MaterialProperties/DMMPOpacityMask.h
--- a/Engine/Plugins/Experimental/DynamicMaterial/Source/DynamicMaterialEditor/Private/Components/MaterialProperties/DMMPOpacityMask.h
+++ b/Engine/Plugins/Experimental/DynamicMaterial/Source/DynamicMaterialEditor/Private/Components/MaterialProperties/DMMPOpacityMask.h
@@ -16,7 +16,6 @@ public:
 	UDMMaterialPropertyOpacityMask();
 
 	//~ Begin UDMMaterialProperty
-	virtual bool IsValidForModel(UDynamicMaterialModelEditorOnlyData& InModelEditorOnlyData) const override;
 	virtual UMaterialExpression* GetDefaultInput(const TSharedRef<FDMMaterialBuildState>& InBuildState) const override;
 	virtual TEnumAsByte<EMaterialSamplerType> GetTextureSamplerType() const override;
 	//~ End UDMMaterialProperty
diff --git a/Engine/Plugins/Experimental/DynamicMaterial/Source/DynamicMaterialEditor/Private/Components/MaterialProperties/DMMPPixelDepthOffset.cpp b/Engine/Plugins/Experimental/DynamicMaterial/Source/DynamicMaterialEditor/Private/Components/MaterialProperties/DMMPPixelDepthOffset.cpp
--- a/Engine/Plugins/Experimental/DynamicMaterial/Source/DynamicMaterialEditor/Private/Components/MaterialProperties/DMMPPixelDepthOffset.cpp
+++ b/Engine/Plugins/Experimental/DynamicMaterial/Source/DynamicMaterialEditor/Private/Components/MaterialProperties/DMMPPixelDepthOffset.cpp
@@ -10,19 +10,6 @@ UDMMaterialPropertyPixelDepthOffset::UDMMaterialPropertyPixelDepthOffset()
 {
 }
 
-bool UDMMaterialPropertyPixelDepthOffset::IsValidForModel(UDynamicMaterialModelEditorOnlyData& InModelEditorOnlyData) const
-{
-	switch (InModelEditorOnlyData.GetBlendMode())
-	{
-		case EBlendMode::BLEND_Opaque:
-		case EBlendMode::BLEND_Masked:
-			return true;
-
-		default:
-			return false;
-	}
-}
-
 UMaterialExpression* UDMMaterialPropertyPixelDepthOffset::GetDefaultInput(
 	const TSharedRef<FDMMaterialBuildState>& InBuildState) const
 {
diff --git a/Engine/Plugins/Experimental/DynamicMaterial/Source/DynamicMaterialEditor/Private/Components/MaterialProperties/DMMPPixelDepthOffset.h b/Engine/Plugins/Experimental/DynamicMaterial/Source/DynamicMaterialEditor/Private/Components/MaterialProperties/DMMPPixelDepthOffset.h
--- a/Engine/Plugins/Experimental/DynamicMaterial/Source/DynamicMaterialEditor/Private/Components/MaterialProperties/DMMPPixelDepthOffset.h
+++ b/Engine/Plugins/Experimental/DynamicMaterial/Source/DynamicMaterialEditor/Private/Components/MaterialProperties/DMMPPixelDepthOffset.h
@@ -16,7 +16,6 @@ public:
 	UDMMaterialPropertyPixelDepthOffset();
 
 	//~ Begin UDMMaterialProperty
-	virtual bool IsValidForModel(UDynamicMaterialModelEditorOnlyData& InModelEditorOnlyData) const override;
 	virtual UMaterialExpression* GetDefaultInput(const TSharedRef<FDMMaterialBuildState>& InBuildState) const override;
 	//~ End UDMMaterialProperty
 };
diff --git a/Engine/Plugins/Experimental/DynamicMaterial/Source/DynamicMaterialEditor/Private/Components/MaterialProperties/DMMPRefraction.cpp b/Engine/Plugins/Experimental/DynamicMaterial/Source/DynamicMaterialEditor/Private/Components/MaterialProperties/DMMPRefraction.cpp
--- a/Engine/Plugins/Experimental/DynamicMaterial/Source/DynamicMaterialEditor/Private/Components/MaterialProperties/DMMPRefraction.cpp
+++ b/Engine/Plugins/Experimental/DynamicMaterial/Source/DynamicMaterialEditor/Private/Components/MaterialProperties/DMMPRefraction.cpp
@@ -10,20 +10,6 @@ UDMMaterialPropertyRefraction::UDMMaterialPropertyRefraction()
 {
 }
 
-bool UDMMaterialPropertyRefraction::IsValidForModel(UDynamicMaterialModelEditorOnlyData& InModelEditorOnlyData) const
-{
-	switch (InModelEditorOnlyData.GetBlendMode())
-	{
-		case EBlendMode::BLEND_Translucent:
-		case EBlendMode::BLEND_Additive:
-		case EBlendMode::BLEND_AlphaComposite:
-			return true;
-
-		default:
-			return false;
-	}
-}
-
 UMaterialExpression* UDMMaterialPropertyRefraction::GetDefaultInput(
 	const TSharedRef<FDMMaterialBuildState>& InBuildState) const
 {
diff --git a/Engine/Plugins/Experimental/DynamicMaterial/Source/DynamicMaterialEditor/Private/Components/MaterialProperties/DMMPRefraction.h b/Engine/Plugins/Experimental/DynamicMaterial/Source/DynamicMaterialEditor/Private/Components/MaterialProperties/DMMPRefraction.h
--- a/Engine/Plugins/Experimental/DynamicMaterial/Source/DynamicMaterialEditor/Private/Components/MaterialProperties/DMMPRefraction.h
+++ b/Engine/Plugins/Experimental/DynamicMaterial/Source/DynamicMaterialEditor/Private/Components/MaterialProperties/DMMPRefraction.h
@@ -16,7 +16,6 @@ public:
 	UDMMaterialPropertyRefraction();
 
 	//~ Begin UDMMaterialProperty
-	virtual bool IsValidForModel(UDynamicMaterialModelEditorOnlyData& InModelEditorOnlyData) const override;
 	virtual UMaterialExpression* GetDefaultInput(const TSharedRef<FDMMaterialBuildState>& InBuildState) const override;
 	//~ End UDMMaterialProperty
 };
diff --git a/Engine/Plugins/Experimental/DynamicMaterial/Source/DynamicMaterialEditor/Private/Components/MaterialProperties/DMMPRoughness.cpp b/Engine/Plugins/Experimental/DynamicMaterial/Source/DynamicMaterialEditor/Private/Components/MaterialProperties/DMMPRoughness.cpp
--- a/Engine/Plugins/Experimental/DynamicMaterial/Source/DynamicMaterialEditor/Private/Components/MaterialProperties/DMMPRoughness.cpp
+++ b/Engine/Plugins/Experimental/DynamicMaterial/Source/DynamicMaterialEditor/Private/Components/MaterialProperties/DMMPRoughness.cpp
@@ -10,24 +10,6 @@ UDMMaterialPropertyRoughness::UDMMaterialPropertyRoughness()
 {
 }
 
-bool UDMMaterialPropertyRoughness::IsValidForModel(UDynamicMaterialModelEditorOnlyData& InModelEditorOnlyData) const
-{
-	if (InModelEditorOnlyData.GetShadingModel() == EDMMaterialShadingModel::Unlit)
-	{
-		return false;
-	}
-
-	switch (InModelEditorOnlyData.GetBlendMode())
-	{
-		case EBlendMode::BLEND_Opaque:
-		case EBlendMode::BLEND_Masked:
-			return true;
-
-		default:
-			return false;
-	}
-}
-
 UMaterialExpression* UDMMaterialPropertyRoughness::GetDefaultInput(const TSharedRef<FDMMaterialBuildState>& InBuildState) const
 {
 	return CreateConstant(InBuildState, 0.f);
diff --git a/Engine/Plugins/Experimental/DynamicMaterial/Source/DynamicMaterialEditor/Private/Components/MaterialProperties/DMMPRoughness.h b/Engine/Plugins/Experimental/DynamicMaterial/Source/DynamicMaterialEditor/Private/Components/MaterialProperties/DMMPRoughness.h
--- a/Engine/Plugins/Experimental/DynamicMaterial/Source/DynamicMaterialEditor/Private/Components/MaterialProperties/DMMPRoughness.h
+++ b/Engine/Plugins/Experimental/DynamicMaterial/Source/DynamicMaterialEditor/Private/Components/MaterialProperties/DMMPRoughness.h
@@ -16,7 +16,6 @@ public:
 	UDMMaterialPropertyRoughness();
 
 	//~ Begin UDMMaterialProperty
-	virtual bool IsValidForModel(UDynamicMaterialModelEditorOnlyData& InModelEditorOnlyData) const override;
 	virtual UMaterialExpression* GetDefaultInput(const TSharedRef<FDMMaterialBuildState>& InBuildState) const override;
 	//~ End UDMMaterialProperty
 };
diff --git a/Engine/Plugins/Experimental/DynamicMaterial/Source/DynamicMaterialEditor/Private/Components/MaterialProperties/DMMPSpecular.cpp b/Engine/Plugins/Experimental/DynamicMaterial/Source/DynamicMaterialEditor/Private/Components/MaterialProperties/DMMPSpecular.cpp
--- a/Engine/Plugins/Experimental/DynamicMaterial/Source/DynamicMaterialEditor/Private/Components/MaterialProperties/DMMPSpecular.cpp
+++ b/Engine/Plugins/Experimental/DynamicMaterial/Source/DynamicMaterialEditor/Private/Components/MaterialProperties/DMMPSpecular.cpp
@@ -10,24 +10,6 @@ UDMMaterialPropertySpecular::UDMMaterialPropertySpecular()
 {
 }
 
-bool UDMMaterialPropertySpecular::IsValidForModel(UDynamicMaterialModelEditorOnlyData& InModelEditorOnlyData) const
-{
-	if (InModelEditorOnlyData.GetShadingModel() == EDMMaterialShadingModel::Unlit)
-	{
-		return false;
-	}
-
-	switch (InModelEditorOnlyData.GetBlendMode())
-	{
-		case EBlendMode::BLEND_Opaque:
-		case EBlendMode::BLEND_Masked:
-			return true;
-
-		default:
-			return false;
-	}
-}
-
 UMaterialExpression* UDMMaterialPropertySpecular::GetDefaultInput(const TSharedRef<FDMMaterialBuildState>& InBuildState) const
 {
 	return CreateConstant(InBuildState, FVector::OneVector);
diff --git a/Engine/Plugins/Experimental/DynamicMaterial/Source/DynamicMaterialEditor/Private/Components/MaterialProperties/DMMPSpecular.h b/Engine/Plugins/Experimental/DynamicMaterial/Source/DynamicMaterialEditor/Private/Components/MaterialProperties/DMMPSpecular.h
--- a/Engine/Plugins/Experimental/DynamicMaterial/Source/DynamicMaterialEditor/Private/Components/MaterialProperties/DMMPSpecular.h
+++ b/Engine/Plugins/Experimental/DynamicMaterial/Source/DynamicMaterialEditor/Private/Components/MaterialProperties/DMMPSpecular.h
@@ -16,7 +16,6 @@ public:
 	UDMMaterialPropertySpecular();
 
 	//~ Begin UDMMaterialProperty
-	virtual bool IsValidForModel(UDynamicMaterialModelEditorOnlyData& InModelEditorOnlyData) const override;
 	virtual UMaterialExpression* GetDefaultInput(const TSharedRef<FDMMaterialBuildState>& InBuildState) const override;
 	//~ End UDMMaterialProperty
 };
diff --git a/Engine/Plugins/Experimental/DynamicMaterial/Source/DynamicMaterialEditor/Private/Components/MaterialProperties/DMMPTangent.cpp b/Engine/Plugins/Experimental/DynamicMaterial/Source/DynamicMaterialEditor/Private/Components/MaterialProperties/DMMPTangent.cpp
--- a/Engine/Plugins/Experimental/DynamicMaterial/Source/DynamicMaterialEditor/Private/Components/MaterialProperties/DMMPTangent.cpp
+++ b/Engine/Plugins/Experimental/DynamicMaterial/Source/DynamicMaterialEditor/Private/Components/MaterialProperties/DMMPTangent.cpp
@@ -10,24 +10,6 @@ UDMMaterialPropertyTangent::UDMMaterialPropertyTangent()
 {
 }
 
-bool UDMMaterialPropertyTangent::IsValidForModel(UDynamicMaterialModelEditorOnlyData& InModelEditorOnlyData) const
-{
-	if (InModelEditorOnlyData.GetShadingModel() == EDMMaterialShadingModel::Unlit)
-	{
-		return false;
-	}
-
-	switch (InModelEditorOnlyData.GetBlendMode())
-	{
-		case EBlendMode::BLEND_Opaque:
-		case EBlendMode::BLEND_Masked:
-			return true;
-
-		default:
-			return false;
-	}
-}
-
 UMaterialExpression* UDMMaterialPropertyTangent::GetDefaultInput(const TSharedRef<FDMMaterialBuildState>& InBuildState) const
 {
 	return CreateConstant(InBuildState, FVector::OneVector);
diff --git a/Engine/Plugins/Experimental/DynamicMaterial/Source/DynamicMaterialEditor/Private/Components/MaterialProperties/DMMPTangent.h b/Engine/Plugins/Experimental/DynamicMaterial/Source/DynamicMaterialEditor/Private/Components/MaterialProperties/DMMPTangent.h
--- a/Engine/Plugins/Experimental/DynamicMaterial/Source/DynamicMaterialEditor/Private/Components/MaterialProperties/DMMPTangent.h
+++ b/Engine/Plugins/Experimental/DynamicMaterial/Source/DynamicMaterialEditor/Private/Components/MaterialProperties/DMMPTangent.h
@@ -16,7 +16,6 @@ public:
 	UDMMaterialPropertyTangent();
 
 	//~ Begin UDMMaterialProperty
-	virtual bool IsValidForModel(UDynamicMaterialModelEditorOnlyData& InModelEditorOnlyData) const override;
 	virtual UMaterialExpression* GetDefaultInput(const TSharedRef<FDMMaterialBuildState>& InBuildState) const override;
 	//~ End UDMMaterialProperty
 };
diff --git a/Engine/Plugins/Experimental/DynamicMaterial/Source/DynamicMaterialEditor/Private/Components/MaterialProperties/DMMPWorldPositionOffset.cpp b/Engine/Plugins/Experimental/DynamicMaterial/Source/DynamicMaterialEditor/Private/Components/MaterialProperties/DMMPWorldPositionOffset.cpp
--- a/Engine/Plugins/Experimental/DynamicMaterial/Source/DynamicMaterialEditor/Private/Components/MaterialProperties/DMMPWorldPositionOffset.cpp
+++ b/Engine/Plugins/Experimental/DynamicMaterial/Source/DynamicMaterialEditor/Private/Components/MaterialProperties/DMMPWorldPositionOffset.cpp
@@ -10,11 +10,6 @@ UDMMaterialPropertyWorldPositionOffset::UDMMaterialPropertyWorldPositionOffset()
 {
 }
 
-bool UDMMaterialPropertyWorldPositionOffset::IsValidForModel(UDynamicMaterialModelEditorOnlyData& InModelEditorOnlyData) const
-{
-	return true;
-}
-
 UMaterialExpression* UDMMaterialPropertyWorldPositionOffset::GetDefaultInput(
 	const TSharedRef<FDMMaterialBuildState>& InBuildState) const
 {
diff --git a/Engine/Plugins/Experimental/DynamicMaterial/Source/DynamicMaterialEditor/Private/Components/MaterialProperties/DMMPWorldPositionOffset.h b/Engine/Plugins/Experimental/DynamicMaterial/Source/DynamicMaterialEditor/Private/Components/MaterialProperties/DMMPWorldPositionOffset.h
--- a/Engine/Plugins/Experimental/DynamicMaterial/Source/DynamicMaterialEditor/Private/Components/MaterialProperties/DMMPWorldPositionOffset.h
+++ b/Engine/Plugins/Experimental/DynamicMaterial/Source/DynamicMaterialEditor/Private/Components/MaterialProperties/DMMPWorldPositionOffset.h
@@ -16,7 +16,6 @@ public:
 	UDMMaterialPropertyWorldPositionOffset();
 
 	//~ Begin UDMMaterialProperty
-	virtual bool IsValidForModel(UDynamicMaterialModelEditorOnlyData& InModelEditorOnlyData) const override;
 	virtual UMaterialExpression* GetDefaultInput(const TSharedRef<FDMMaterialBuildState>& InBuildState) const override;
 	//~ End UDMMaterialProperty
 };
diff --git a/Engine/Plugins/Experimental/DynamicMaterial/Source/DynamicMaterialEditor/Private/Components/MaterialStageExpressions/DMMSETextureSampleBase.cpp b/Engine/Plugins/Experimental/DynamicMaterial/Source/DynamicMaterialEditor/Private/Components/MaterialStageExpressions/DMMSETextureSampleBase.cpp
--- a/Engine/Plugins/Experimental/DynamicMaterial/Source/DynamicMaterialEditor/Private/Components/MaterialStageExpressions/DMMSETextureSampleBase.cpp
+++ b/Engine/Plugins/Experimental/DynamicMaterial/Source/DynamicMaterialEditor/Private/Components/MaterialStageExpressions/DMMSETextureSampleBase.cpp
@@ -17,6 +17,7 @@
 #include "Materials/MaterialExpressionTextureSample.h"
 #include "Model/DMMaterialBuildState.h"
 #include "Model/DynamicMaterialModelEditorOnlyData.h"
+#include "Utils/DMUtils.h"
 
 #define LOCTEXT_NAMESPACE "DMMaterialStageExpressionTextureSample"
 
diff --git a/Engine/Plugins/Experimental/DynamicMaterial/Source/DynamicMaterialEditor/Private/Components/MaterialStageExpressions/DMMSEWorldPositionNoise.cpp b/Engine/Plugins/Experimental/DynamicMaterial/Source/DynamicMaterialEditor/Private/Components/MaterialStageExpressions/DMMSEWorldPositionNoise.cpp
--- a/Engine/Plugins/Experimental/DynamicMaterial/Source/DynamicMaterialEditor/Private/Components/MaterialStageExpressions/DMMSEWorldPositionNoise.cpp
+++ b/Engine/Plugins/Experimental/DynamicMaterial/Source/DynamicMaterialEditor/Private/Components/MaterialStageExpressions/DMMSEWorldPositionNoise.cpp
@@ -16,6 +16,7 @@
 #include "Materials/MaterialExpressionWorldPosition.h"
 #include "Model/DMMaterialBuildState.h"
 #include "Model/DMMaterialBuildUtils.h"
+#include "Utils/DMUtils.h"
 
 #define LOCTEXT_NAMESPACE "DMMaterialStageExpressionTextureSample"
 
diff --git a/Engine/Plugins/Experimental/DynamicMaterial/Source/DynamicMaterialEditor/Private/Components/MaterialStageGradients/DMMSGLinear.cpp b/Engine/Plugins/Experimental/DynamicMaterial/Source/DynamicMaterialEditor/Private/Components/MaterialStageGradients/DMMSGLinear.cpp
--- a/Engine/Plugins/Experimental/DynamicMaterial/Source/DynamicMaterialEditor/Private/Components/MaterialStageGradients/DMMSGLinear.cpp
+++ b/Engine/Plugins/Experimental/DynamicMaterial/Source/DynamicMaterialEditor/Private/Components/MaterialStageGradients/DMMSGLinear.cpp
@@ -11,6 +11,7 @@
 #include "Model/DMMaterialBuildState.h"
 #include "Model/DMMaterialBuildUtils.h"
 #include "Utils/DMMaterialFunctionLibrary.h"
+#include "Utils/DMUtils.h"
 
 #define LOCTEXT_NAMESPACE "DMMaterialStageGradientLinear"
 
diff --git a/Engine/Plugins/Experimental/DynamicMaterial/Source/DynamicMaterialEditor/Private/Components/MaterialStageGradients/DMMSGRadial.cpp b/Engine/Plugins/Experimental/DynamicMaterial/Source/DynamicMaterialEditor/Private/Components/MaterialStageGradients/DMMSGRadial.cpp
--- a/Engine/Plugins/Experimental/DynamicMaterial/Source/DynamicMaterialEditor/Private/Components/MaterialStageGradients/DMMSGRadial.cpp
+++ b/Engine/Plugins/Experimental/DynamicMaterial/Source/DynamicMaterialEditor/Private/Components/MaterialStageGradients/DMMSGRadial.cpp
@@ -5,6 +5,7 @@
 #include "Materials/MaterialExpressionMaterialFunctionCall.h"
 #include "Model/DMMaterialBuildState.h"
 #include "Model/DMMaterialBuildUtils.h"
+#include "Utils/DMUtils.h"
 
 #define LOCTEXT_NAMESPACE "DMMaterialStageGradientRadial"
 
diff --git a/Engine/Plugins/Experimental/DynamicMaterial/Source/DynamicMaterialEditor/Private/Components/MaterialStageInputs/DMMSITextureUV.cpp b/Engine/Plugins/Experimental/DynamicMaterial/Source/DynamicMaterialEditor/Private/Components/MaterialStageInputs/DMMSITextureUV.cpp
--- a/Engine/Plugins/Experimental/DynamicMaterial/Source/DynamicMaterialEditor/Private/Components/MaterialStageInputs/DMMSITextureUV.cpp
+++ b/Engine/Plugins/Experimental/DynamicMaterial/Source/DynamicMaterialEditor/Private/Components/MaterialStageInputs/DMMSITextureUV.cpp
@@ -23,6 +23,7 @@
 #include "Model/DynamicMaterialModel.h"
 #include "Model/DynamicMaterialModelEditorOnlyData.h"
 #include "Utils/DMMaterialFunctionLibrary.h"
+#include "Utils/DMUtils.h"
 
 #define LOCTEXT_NAMESPACE "DMMaterialStageInputTextureUV"
 
diff --git a/Engine/Plugins/Experimental/DynamicMaterial/Source/DynamicMaterialEditor/Private/DynamicMaterialEditorSettings.cpp b/Engine/Plugins/Experimental/DynamicMaterial/Source/DynamicMaterialEditor/Private/DynamicMaterialEditorSettings.cpp
--- a/Engine/Plugins/Experimental/DynamicMaterial/Source/DynamicMaterialEditor/Private/DynamicMaterialEditorSettings.cpp
+++ b/Engine/Plugins/Experimental/DynamicMaterial/Source/DynamicMaterialEditor/Private/DynamicMaterialEditorSettings.cpp
@@ -56,8 +56,10 @@ bool FDMMaterialChannelListPreset::IsPropertyEnabled(EDMMaterialPropertyType InP
 	switch (InProperty)
 	{
 		case EDMMaterialPropertyType::BaseColor:
+			return bBaseColor;
+
 		case EDMMaterialPropertyType::EmissiveColor:
-			return bRGB;
+			return bEmissive;
 
 		case EDMMaterialPropertyType::Opacity:
 		case EDMMaterialPropertyType::OpacityMask:
@@ -112,8 +114,8 @@ UDynamicMaterialEditorSettings::UDynamicMaterialEditorSettings()
 
 	FDMMaterialChannelListPreset Opaque;
 	Opaque.Name = TEXT("Opaque");
-	Opaque.bRGB = true;
-	Opaque.bOpacity = false;
+	Opaque.bBaseColor = true;
+	Opaque.bEmissive = true;
 	Opaque.DefaultBlendMode = BLEND_Opaque;
 	Opaque.DefaultShadingModel = EDMMaterialShadingModel::DefaultLit;
 	Opaque.bDefaultAnimated = false;
@@ -121,8 +123,7 @@ UDynamicMaterialEditorSettings::UDynamicMaterialEditorSettings()
 
 	FDMMaterialChannelListPreset Emissive;
 	Emissive.Name = TEXT("Emissive");
-	Emissive.bRGB = true;
-	Emissive.bOpacity = false;
+	Emissive.bEmissive = true;
 	Emissive.DefaultBlendMode = BLEND_Opaque;
 	Emissive.DefaultShadingModel = EDMMaterialShadingModel::Unlit;
 	Emissive.bDefaultAnimated = false;
@@ -130,7 +131,7 @@ UDynamicMaterialEditorSettings::UDynamicMaterialEditorSettings()
 
 	FDMMaterialChannelListPreset Translucent;
 	Translucent.Name = TEXT("Translucent");
-	Translucent.bRGB = true;
+	Translucent.bEmissive = true;
 	Translucent.bOpacity = true;
 	Translucent.DefaultBlendMode = BLEND_Translucent;
 	Translucent.DefaultShadingModel = EDMMaterialShadingModel::Unlit;
@@ -139,7 +140,8 @@ UDynamicMaterialEditorSettings::UDynamicMaterialEditorSettings()
 
 	FDMMaterialChannelListPreset PBR;
 	PBR.Name = TEXT("PBR");
-	PBR.bRGB = true;
+	PBR.bBaseColor = true;
+	PBR.bEmissive = true;
 	PBR.bOpacity = true;
 	PBR.bMetallic = true;
 	PBR.bSpecular = true;
@@ -153,7 +155,8 @@ UDynamicMaterialEditorSettings::UDynamicMaterialEditorSettings()
 
 	FDMMaterialChannelListPreset All;
 	All.Name = TEXT("All");
-	All.bRGB = true;
+	All.bBaseColor = true;
+	All.bEmissive = true;
 	All.bOpacity = true;
 	All.bMetallic = true;
 	All.bSpecular = true;
@@ -199,6 +202,24 @@ bool UDynamicMaterialEditorSettings::IsUseLinearColorForVectorsEnabled()
 	return true;
 }
 
+void UDynamicMaterialEditorSettings::PostInitProperties()
+{
+	Super::PostInitProperties();
+
+	if (bValidatedPresets)
+	{
+		return;
+	}
+
+	for (FDMMaterialChannelListPreset& Preset : MaterialChannelPresets)
+	{
+		Preset.bBaseColor = Preset.Name != TEXT("Emissive");
+		Preset.bEmissive = true;
+	}
+
+	bValidatedPresets = true;
+}
+
 void UDynamicMaterialEditorSettings::PreEditChange(FEditPropertyChain& InPropertyAboutToChange)
 {
 	Super::PreEditChange(InPropertyAboutToChange);
diff --git a/Engine/Plugins/Experimental/DynamicMaterial/Source/DynamicMaterialEditor/Private/DynamicMaterialEditorSettings.h b/Engine/Plugins/Experimental/DynamicMaterial/Source/DynamicMaterialEditor/Private/DynamicMaterialEditorSettings.h
--- a/Engine/Plugins/Experimental/DynamicMaterial/Source/DynamicMaterialEditor/Private/DynamicMaterialEditorSettings.h
+++ b/Engine/Plugins/Experimental/DynamicMaterial/Source/DynamicMaterialEditor/Private/DynamicMaterialEditorSettings.h
@@ -65,9 +65,11 @@ struct FDMMaterialChannelListPreset
 	UPROPERTY(Config, EditAnywhere, BlueprintReadOnly, Category = "Channel")
 	FName Name;
 
-	/** This cannot be disabled. It is here for illustrative purposes only. */
-	UPROPERTY(Config, VisibleAnywhere, BlueprintReadOnly, Category = "Channel")
-	bool bRGB = true;
+	UPROPERTY(Config, EditAnywhere, BlueprintReadOnly, Category = "Channel")
+	bool bBaseColor = false;
+
+	UPROPERTY(Config, EditAnywhere, BlueprintReadOnly, Category = "Channel")
+	bool bEmissive = false;
 
 	UPROPERTY(Config, EditAnywhere, BlueprintReadWrite, Category = "Channel")
 	bool bOpacity = false;
@@ -232,12 +234,16 @@ public:
 	UPROPERTY(Config, EditAnywhere, BlueprintReadWrite, Category = "Material Functions")
 	bool bUseLinearColorForVectors;
 
+	UPROPERTY(Config)
+	bool bValidatedPresets = false;
+
 	/** This variable is accessed in multiple places, so this is a quick accessor. */
 	static bool IsUseLinearColorForVectorsEnabled();
 
 	FOnFinishedChangingProperties OnSettingsChanged;
 
 	//~ Begin UObject
+	virtual void PostInitProperties() override;
 	virtual void PreEditChange(FEditPropertyChain& InPropertyAboutToChange) override;
 	virtual void PostEditChangeProperty(FPropertyChangedEvent& InPropertyChangedEvent) override;
 	//~ End UObject
diff --git a/Engine/Plugins/Experimental/DynamicMaterial/Source/DynamicMaterialEditor/Private/Menus/DMMaterialSlotMenus.cpp b/Engine/Plugins/Experimental/DynamicMaterial/Source/DynamicMaterialEditor/Private/Menus/DMMaterialSlotMenus.cpp
--- a/Engine/Plugins/Experimental/DynamicMaterial/Source/DynamicMaterialEditor/Private/Menus/DMMaterialSlotMenus.cpp
+++ b/Engine/Plugins/Experimental/DynamicMaterial/Source/DynamicMaterialEditor/Private/Menus/DMMaterialSlotMenus.cpp
@@ -12,44 +12,6 @@
 namespace UE::DynamicMaterialEditor::Private
 {
 	static const FName SlotAddLayerMenuName = TEXT("MaterialDesigner.MaterialSlot.AddLayer");
-	static const FName SlotSectionName = TEXT("Slot");
-
-	void AddSlotMenu(UToolMenu* InMenu)
-	{
-		if (!IsValid(InMenu) || InMenu->ContainsSection(SlotSectionName))
-		{
-			return;
-		}
-
-		const UDMMenuContext* const MenuContext = InMenu->FindContext<UDMMenuContext>();
-
-		if (!MenuContext)
-		{
-			return;
-		}
-
-		TSharedPtr<SDMEditor> EditorWidget = MenuContext->GetEditorWidget();
-
-		if (!EditorWidget.IsValid() || !EditorWidget->GetActiveSlotWidget().IsValid())
-		{
-			return;
-		}
-
-		FToolMenuSection& NewSection = InMenu->AddSection(SlotSectionName, LOCTEXT("MaterialSlotMenu", "Material Slot"));
-		{
-			NewSection.AddMenuEntry(
-				NAME_None,
-				LOCTEXT("RemoveSlot", "Remove"),
-				LOCTEXT("RemoveSlotTooltip", "Remove this stage from its Material Slot.\n\n"
-					"Warning: Removing a stage off may result in inputs being reset where incompatibilities are found."),
-				FSlateIcon(),
-				FUIAction(FExecuteAction::CreateSP(
-					EditorWidget->GetActiveSlotWidget().Get(),
-					&SDMSlot::RemoveSlot
-				))
-			);
-		}
-	}
 }
 
 TSharedRef<SWidget> FDMMaterialSlotMenus::MakeAddLayerButtonMenu(const TSharedPtr<SDMSlot>& InSlotWidget)
diff --git a/Engine/Plugins/Experimental/DynamicMaterial/Source/DynamicMaterialEditor/Private/Model/DMMaterialBuildUtils.cpp b/Engine/Plugins/Experimental/DynamicMaterial/Source/DynamicMaterialEditor/Private/Model/DMMaterialBuildUtils.cpp
--- a/Engine/Plugins/Experimental/DynamicMaterial/Source/DynamicMaterialEditor/Private/Model/DMMaterialBuildUtils.cpp
+++ b/Engine/Plugins/Experimental/DynamicMaterial/Source/DynamicMaterialEditor/Private/Model/DMMaterialBuildUtils.cpp
@@ -14,6 +14,7 @@
 #include "Materials/MaterialExpressionTextureSample.h"
 #include "Model/DMMaterialBuildState.h"
 #include "Utils/DMPrivate.h"
+#include "Utils/DMUtils.h"
 
 FDMMaterialBuildUtils::FDMMaterialBuildUtils(FDMMaterialBuildState& InBuildState)
 	: BuildState(InBuildState)
diff --git a/Engine/Plugins/Experimental/DynamicMaterial/Source/DynamicMaterialEditor/Private/Model/DynamicMaterialModelEditorOnlyData.cpp b/Engine/Plugins/Experimental/DynamicMaterial/Source/DynamicMaterialEditor/Private/Model/DynamicMaterialModelEditorOnlyData.cpp
--- a/Engine/Plugins/Experimental/DynamicMaterial/Source/DynamicMaterialEditor/Private/Model/DynamicMaterialModelEditorOnlyData.cpp
+++ b/Engine/Plugins/Experimental/DynamicMaterial/Source/DynamicMaterialEditor/Private/Model/DynamicMaterialModelEditorOnlyData.cpp
@@ -48,6 +48,7 @@
 #include "Model/DynamicMaterialModel.h"
 #include "Model/DynamicMaterialModelEditorOnlyDataVersion.h"
 #include "UObject/Package.h"
+#include "Utils/DMUtils.h"
 
 #define LOCTEXT_NAMESPACE "MaterialDesignerModel"
 
@@ -224,10 +225,27 @@ void UDynamicMaterialModelEditorOnlyData::OnChannelListPresetChanged()
 		return;
 	}
 
+	EnsurePresetSlots();
+
+	SetBlendMode(Preset->DefaultBlendMode);
+	SetShadingModel(Preset->DefaultShadingModel);
+	SetPixelAnimationFlag(Preset->bDefaultAnimated);
+	SetTwoSidedFlag(Preset->bDefaultTwoSided);
+}
+
+void UDynamicMaterialModelEditorOnlyData::EnsurePresetSlots()
+{
+	const FDMMaterialChannelListPreset* Preset = GetDefault<UDynamicMaterialEditorSettings>()->GetPresetByName(ChannelListPreset);
+
+	if (!Preset)
+	{
+		return;
+	}
+
 	ForEachMaterialPropertyType(
 		[this, Preset](EDMMaterialPropertyType InProperty)
 		{
-			if (InProperty == EDMMaterialPropertyType::EmissiveColor || InProperty == EDMMaterialPropertyType::OpacityMask)
+			if (InProperty == EDMMaterialPropertyType::OpacityMask)
 			{
 				return EDMIterationResult::Continue;
 			}
@@ -243,11 +261,6 @@ void UDynamicMaterialModelEditorOnlyData::OnChannelListPresetChanged()
 
 			return EDMIterationResult::Continue;
 		});
-
-	SetBlendMode(Preset->DefaultBlendMode);
-	SetShadingModel(Preset->DefaultShadingModel);
-	SetPixelAnimationFlag(Preset->bDefaultAnimated);
-	SetTwoSidedFlag(Preset->bDefaultTwoSided);
 }
 
 void UDynamicMaterialModelEditorOnlyData::OnDomainChanged()
@@ -256,55 +269,30 @@ void UDynamicMaterialModelEditorOnlyData::OnDomainChanged()
 	{
 		const FDMUpdateGuard Guard;
 
-		// Post process only supports emissive.
-		ForEachMaterialPropertyType(
-			[this](EDMMaterialPropertyType InProperty)
-			{
-				switch (InProperty)
-				{
-					case EDMMaterialPropertyType::BaseColor:
-					case EDMMaterialPropertyType::EmissiveColor:
-						RemoveSlotForMaterialProperty(InProperty);
-						break;
-
-					default:
-						// Do nothing
-						break;
-				}
+		// Make sure we have the correct slots, then adjust for domains.
+		EnsurePresetSlots();
 
-				return EDMIterationResult::Continue;
-			});
+		// Post process only supports emissive.
+		UDMMaterialSlot* BaseColorSlot = GetSlotForMaterialProperty(EDMMaterialPropertyType::BaseColor);
+		UDMMaterialSlot* EmissiveSlot = GetSlotForMaterialProperty(EDMMaterialPropertyType::EmissiveColor);
 
-		EnsureSwapSlotMaterialProperty(EDMMaterialPropertyType::BaseColor, EDMMaterialPropertyType::EmissiveColor);
+		if (!EmissiveSlot)
+		{
+			if (BaseColorSlot)
+			{
+				EnsureSwapSlotMaterialProperty(EDMMaterialPropertyType::BaseColor, EDMMaterialPropertyType::EmissiveColor);
+			}
+			else
+			{
+				AddSlotForMaterialProperty(EDMMaterialPropertyType::EmissiveColor);
+			}
+		}
 
 		SetShadingModel(EDMMaterialShadingModel::Unlit);
 		SetBlendMode(EBlendMode::BLEND_Opaque);
 	}
 	else if (Domain == EMaterialDomain::MD_DeferredDecal)
 	{
-		// Post process only supports basic types.
-		ForEachMaterialPropertyType(
-			[this](EDMMaterialPropertyType InProperty)
-			{
-				switch (InProperty)
-				{
-					case EDMMaterialPropertyType::BaseColor:
-					case EDMMaterialPropertyType::EmissiveColor:
-					case EDMMaterialPropertyType::Opacity:
-					case EDMMaterialPropertyType::OpacityMask:
-						RemoveSlotForMaterialProperty(InProperty);
-						break;
-
-					default:
-						// Do nothing
-						break;
-				}
-
-				return EDMIterationResult::Continue;
-			});
-
-		EnsureSwapSlotMaterialProperty(EDMMaterialPropertyType::EmissiveColor, EDMMaterialPropertyType::BaseColor);
-
 		SetShadingModel(EDMMaterialShadingModel::DefaultLit);
 		SetBlendMode(EBlendMode::BLEND_Translucent);
 	}
@@ -314,6 +302,8 @@ void UDynamicMaterialModelEditorOnlyData::OnDomainChanged()
 
 void UDynamicMaterialModelEditorOnlyData::OnBlendModeChanged()
 {
+	EnsurePresetSlots();
+
 	switch (BlendMode)
 	{
 		case EBlendMode::BLEND_Opaque:
@@ -339,63 +329,16 @@ void UDynamicMaterialModelEditorOnlyData::OnBlendModeChanged()
 
 void UDynamicMaterialModelEditorOnlyData::OnShadingModelChanged()
 {
-	EDMMaterialPropertyType FromProperty;
-	EDMMaterialPropertyType ToProperty;
-
-	switch (ShadingModel)
-	{
-		case EDMMaterialShadingModel::Unlit:
-			FromProperty = EDMMaterialPropertyType::BaseColor;
-			ToProperty = EDMMaterialPropertyType::EmissiveColor;
-			break;
-
-		case EDMMaterialShadingModel::DefaultLit:
-			FromProperty = EDMMaterialPropertyType::EmissiveColor;
-			ToProperty = EDMMaterialPropertyType::BaseColor;
-			break;
-
-		default:
-			return;
-	}
-
-	EnsureSwapSlotMaterialProperty(FromProperty, ToProperty);
-
 	RequestMaterialBuild();
 }
 
 void UDynamicMaterialModelEditorOnlyData::OnPixelAnimationFlagChanged()
 {
-	if (MaterialModel)
-	{
-		if (UMaterial* Material = MaterialModel->GetGeneratedMaterial())
-		{
-			if (GUndo)
-			{
-				Material->Modify();
-			}
-
-			Material->bHasPixelAnimation = bPixelAnimationFlag;
-		}
-	}
-
 	RequestMaterialBuild();
 }
 
 void UDynamicMaterialModelEditorOnlyData::OnTwoSidedFlagChanged()
 {
-	if (MaterialModel)
-	{
-		if (UMaterial* Material = MaterialModel->GetGeneratedMaterial())
-		{
-			if (GUndo)
-			{
-				Material->Modify();
-			}
-
-			Material->TwoSided = bTwoSidedFlag;
-		}
-	}
-
 	RequestMaterialBuild();
 }
 
@@ -820,7 +763,7 @@ bool UDynamicMaterialModelEditorOnlyData::AddTextureSet(UDMTextureSet* InTexture
 
 	for (const TPair<EDMTextureSetMaterialProperty, FDMMaterialTexture>& MaterialTexture : InTextureSet->GetTextures())
 	{
-		const EDMMaterialPropertyType PropertyType = UE::DynamicMaterial::MaterialPropertyToMaterialPropertyType(MaterialTexture.Key);
+		const EDMMaterialPropertyType PropertyType = FDMUtils::TextureSetMaterialPropertyToMaterialPropertyType(MaterialTexture.Key);
 
 		if (PropertyType == EDMMaterialPropertyType::None)
 		{
@@ -1279,35 +1222,11 @@ UDMMaterialSlot* UDynamicMaterialModelEditorOnlyData::AddSlot()
 
 UDMMaterialSlot* UDynamicMaterialModelEditorOnlyData::AddSlotForMaterialProperty(EDMMaterialPropertyType InType)
 {
-	if (InType == EDMMaterialPropertyType::BaseColor || InType == EDMMaterialPropertyType::EmissiveColor)
-	{
-		if (Domain == EMaterialDomain::MD_DeferredDecal)
-		{
-			InType = EDMMaterialPropertyType::EmissiveColor;
-		}
-		else
-		{
-			switch (ShadingModel)
-			{
-				case EDMMaterialShadingModel::DefaultLit:
-					InType = EDMMaterialPropertyType::BaseColor;
-					break;
-
-				case EDMMaterialShadingModel::Unlit:
-					InType = EDMMaterialPropertyType::EmissiveColor;
-					break;
-
-				default:
-					checkNoEntry();
-					break;
-			}
-		}
-	}
-	else if (Domain == EMaterialDomain::MD_PostProcess)
+	if (InType == EDMMaterialPropertyType::EmissiveColor && Domain == EMaterialDomain::MD_PostProcess)
 	{
 		return nullptr;
 	}
-	else if (InType == EDMMaterialPropertyType::Opacity || InType == EDMMaterialPropertyType::OpacityMask)
+	if (InType == EDMMaterialPropertyType::Opacity || InType == EDMMaterialPropertyType::OpacityMask)
 	{
 		switch (BlendMode)
 		{
@@ -1330,11 +1249,6 @@ UDMMaterialSlot* UDynamicMaterialModelEditorOnlyData::AddSlotForMaterialProperty
 		}
 	}
 
-	if (Domain == EMaterialDomain::MD_DeferredDecal && InType != EDMMaterialPropertyType::EmissiveColor)
-	{
-		return nullptr;
-	}
-
 	if (UDMMaterialSlot* ExistingSlot = GetSlotForMaterialProperty(InType))
 	{
 		return ExistingSlot;
@@ -1344,22 +1258,6 @@ UDMMaterialSlot* UDynamicMaterialModelEditorOnlyData::AddSlotForMaterialProperty
 	// be checked. If it is found, it is converted and returned. The same goes for Opacity and OpacityMask.
 	switch (InType)
 	{
-		case EDMMaterialPropertyType::BaseColor:
-			if (UDMMaterialSlot* ExistingSlot = GetSlotForMaterialProperty(EDMMaterialPropertyType::EmissiveColor))
-			{
-				SwapSlotMaterialProperty(EDMMaterialPropertyType::EmissiveColor, EDMMaterialPropertyType::BaseColor);
-				return ExistingSlot;
-			}
-			break;
-
-		case EDMMaterialPropertyType::EmissiveColor:
-			if (UDMMaterialSlot* ExistingSlot = GetSlotForMaterialProperty(EDMMaterialPropertyType::BaseColor))
-			{
-				SwapSlotMaterialProperty(EDMMaterialPropertyType::BaseColor, EDMMaterialPropertyType::EmissiveColor);
-				return ExistingSlot;
-			}
-			break;
-
 		case EDMMaterialPropertyType::Opacity:
 			if (UDMMaterialSlot* ExistingSlot = GetSlotForMaterialProperty(EDMMaterialPropertyType::OpacityMask))
 			{
diff --git a/Engine/Plugins/Experimental/DynamicMaterial/Source/DynamicMaterialEditor/Private/Slate/SDMEditor.cpp b/Engine/Plugins/Experimental/DynamicMaterial/Source/DynamicMaterialEditor/Private/Slate/SDMEditor.cpp
--- a/Engine/Plugins/Experimental/DynamicMaterial/Source/DynamicMaterialEditor/Private/Slate/SDMEditor.cpp
+++ b/Engine/Plugins/Experimental/DynamicMaterial/Source/DynamicMaterialEditor/Private/Slate/SDMEditor.cpp
@@ -976,16 +976,6 @@ TSharedRef<SWidget> SDMEditor::CreateSlotPickerWidget()
 
 	for (const TPair<EDMMaterialPropertyType, UDMMaterialProperty*>& Property : ModelEditorOnlyData->GetMaterialProperties())
 	{
-		if (Property.Key == EDMMaterialPropertyType::BaseColor && ModelEditorOnlyData->GetShadingModel() == EDMMaterialShadingModel::Unlit)
-		{
-			continue;
-		}
-
-		if (Property.Key == EDMMaterialPropertyType::EmissiveColor && ModelEditorOnlyData->GetShadingModel() == EDMMaterialShadingModel::DefaultLit)
-		{
-			continue;
-		}		
-
 		// Always create opacity, not opacity mask - Will be sorted out by the button itself.
 		if (Property.Key == EDMMaterialPropertyType::OpacityMask)
 		{
@@ -1203,14 +1193,6 @@ bool SDMEditor::IsPropertyValidForModel(EDMMaterialPropertyType InProperty) cons
 		}
 	}
 
-	if (InProperty == EDMMaterialPropertyType::BaseColor)
-	{
-		if (UDMMaterialProperty* Property = EditorOnlyData->GetMaterialProperty(EDMMaterialPropertyType::EmissiveColor))
-		{
-			return Property->IsValidForModel(*EditorOnlyData);
-		}
-	}
-
 	if (InProperty == EDMMaterialPropertyType::Opacity)
 	{
 		if (UDMMaterialProperty* Property = EditorOnlyData->GetMaterialProperty(EDMMaterialPropertyType::OpacityMask))
@@ -1240,7 +1222,7 @@ ECheckBoxState SDMEditor::GetSlotCheckState(EDMMaterialPropertyType InProperty)
 
 	if (!EditorOnlyData)
 	{
-		return ECheckBoxState::Undetermined;
+		return ECheckBoxState::Unchecked;
 	}
 
 	if (!EditorOnlyData->GetSlots().IsValidIndex(ActiveSlotIndex))
@@ -1634,14 +1616,6 @@ UDMMaterialSlot* SDMEditor::GetSlotForMaterialProperty(EDMMaterialPropertyType I
 	{
 		switch (InProperty)
 		{
-			case EDMMaterialPropertyType::BaseColor:
-				PropertySlot = ModelEditorOnlyData->GetSlotForMaterialProperty(EDMMaterialPropertyType::EmissiveColor);
-				break;
-
-			case EDMMaterialPropertyType::EmissiveColor:
-				PropertySlot = ModelEditorOnlyData->GetSlotForMaterialProperty(EDMMaterialPropertyType::BaseColor);
-				break;
-
 			case EDMMaterialPropertyType::Opacity:
 				PropertySlot = ModelEditorOnlyData->GetSlotForMaterialProperty(EDMMaterialPropertyType::OpacityMask);
 				break;
diff --git a/Engine/Plugins/Experimental/DynamicMaterial/Source/DynamicMaterialEditor/Private/Slate/SDMMaterialWizard.cpp b/Engine/Plugins/Experimental/DynamicMaterial/Source/DynamicMaterialEditor/Private/Slate/SDMMaterialWizard.cpp
--- a/Engine/Plugins/Experimental/DynamicMaterial/Source/DynamicMaterialEditor/Private/Slate/SDMMaterialWizard.cpp
+++ b/Engine/Plugins/Experimental/DynamicMaterial/Source/DynamicMaterialEditor/Private/Slate/SDMMaterialWizard.cpp
@@ -182,7 +182,7 @@ TSharedRef<SWidget> SDMMaterialWizard::CreateChannelList()
 	{
 		for (const TPair<EDMMaterialPropertyType, UDMMaterialProperty*>& Property : ModelEditorOnlyData->GetMaterialProperties())
 		{
-			if (Property.Key == EDMMaterialPropertyType::EmissiveColor || Property.Key == EDMMaterialPropertyType::OpacityMask)
+			if (Property.Key == EDMMaterialPropertyType::OpacityMask)
 			{
 				continue;
 			}
diff --git a/Engine/Plugins/Experimental/DynamicMaterial/Source/DynamicMaterialEditor/Private/Slate/SDMSlot.cpp b/Engine/Plugins/Experimental/DynamicMaterial/Source/DynamicMaterialEditor/Private/Slate/SDMSlot.cpp
--- a/Engine/Plugins/Experimental/DynamicMaterial/Source/DynamicMaterialEditor/Private/Slate/SDMSlot.cpp
+++ b/Engine/Plugins/Experimental/DynamicMaterial/Source/DynamicMaterialEditor/Private/Slate/SDMSlot.cpp
@@ -323,39 +323,6 @@ bool SDMSlot::IsOpacityPropertyEnabled() const
 	return LayerView->GetLayerItems().IsValidIndex(LayerView->GetSelectedLayerIndex());
 }
 
-FReply SDMSlot::OnRemoveSlotClicked()
-{
-	if (CanRemoveSlot())
-	{
-		RemoveSlot();
-		return FReply::Handled();
-	}
-
-	return FReply::Unhandled();
-}
-
-bool SDMSlot::CanRemoveSlot() const
-{
-	if (UDMMaterialSlot* Slot = SlotWeak.Get())
-	{
-		// Cannot remove RGB slot.
-		const TArray<TObjectPtr<UDMMaterialLayerObject>>& Layers = Slot->GetLayers();
-
-		for (const TObjectPtr<UDMMaterialLayerObject>& Layer : Layers)
-		{
-			if (Layer->GetMaterialProperty() == EDMMaterialPropertyType::BaseColor
-				|| Layer->GetMaterialProperty() == EDMMaterialPropertyType::EmissiveColor)
-			{
-				return false;
-			}
-		}
-
-		return true;
-	}
-
-	return false;
-}
-
 void SDMSlot::RefreshMainWidget()
 {
 	bInvalidateMainWidget = false;
@@ -545,24 +512,6 @@ TSharedRef<SWidget> SDMSlot::CreateLayerButtonsRowWidget()
 	return 
 		SNew(SHorizontalBox)
 		+ SHorizontalBox::Slot()
-		.AutoWidth()
-		.HAlign(HAlign_Left)
-		.VAlign(VAlign_Top)
-		.Padding(5.0f, 2.0f, 10.0f, 2.0f)
-		[
-			SNew(SButton)
-			.ContentPadding(4.0f)
-			.ButtonStyle(FDynamicMaterialEditorStyle::Get(), "HoverHintOnly.Bordered.Dark")
-			.ToolTipText(LOCTEXT("RemoveSlotTooltip", "Remove Slot\n\nRGB Slots cannot be removed."))
-			.IsEnabled(this, &SDMSlot::CanRemoveSlot)
-			.OnClicked(this, &SDMSlot::OnRemoveSlotClicked)
-			[
-				SNew(SImage)
-				.Image(FAppStyle::Get().GetBrush("GenericCommands.Delete"))
-				.DesiredSizeOverride(FVector2D(16.0f))
-			]
-		]
-		+ SHorizontalBox::Slot()
 		.FillWidth(1.0f)
 		.HAlign(HAlign_Left)
 		.VAlign(VAlign_Center)
@@ -1097,25 +1046,6 @@ void SDMSlot::ClearSelection()
 	}
 }
 
-void SDMSlot::RemoveSlot()
-{
-	if (UDMMaterialSlot* Slot = SlotWeak.Get())
-	{
-		if (UDynamicMaterialModelEditorOnlyData* ModelEditorOnlyData = Slot->GetMaterialModelEditorOnlyData())
-		{
-			FDMScopedUITransaction Transaction(LOCTEXT("RemoveSlot", "Material Designer Remove Slot"));
-			ModelEditorOnlyData->Modify();
-			ModelEditorOnlyData->RemoveSlot(Slot->GetIndex());
-
-			if (TSharedPtr<SDMEditor> Editor = EditorWidgetWeak.Pin())
-			{
-				Editor->SetActiveSlotIndex(0);
-				Editor->RefreshSlotPickerList();
-			}
-		}
-	}
-}
-
 void SDMSlot::AddPropertyToSlot(EDMMaterialPropertyType Property)
 {
 	UDMMaterialSlot* Slot = SlotWeak.Get();
diff --git a/Engine/Plugins/Experimental/DynamicMaterial/Source/DynamicMaterialEditor/Private/Slate/SDMSlot.h b/Engine/Plugins/Experimental/DynamicMaterial/Source/DynamicMaterialEditor/Private/Slate/SDMSlot.h
--- a/Engine/Plugins/Experimental/DynamicMaterial/Source/DynamicMaterialEditor/Private/Slate/SDMSlot.h
+++ b/Engine/Plugins/Experimental/DynamicMaterial/Source/DynamicMaterialEditor/Private/Slate/SDMSlot.h
@@ -76,8 +76,6 @@ public:
 	void AddNewLayer_SceneTexture();
 	void AddNewLayer_Renderer(TSubclassOf<UDMRenderTargetRenderer> InRendererClass);
 
-	void RemoveSlot();
-
 	bool CheckValidity();
 
 	bool GetLayerRowsButtonsCanRemove() const;
@@ -161,7 +159,4 @@ protected:
 
 	bool CanRemoveLayerByIndex(const int32 InLayerIndex) const;
 	void RemoveLayerByIndex(const int32 InLayerIndex, const bool bInSelectNextMaskIfMask = false);
-
-	bool CanRemoveSlot() const;
-	FReply OnRemoveSlotClicked();
 };
diff --git a/Engine/Plugins/Experimental/DynamicMaterial/Source/DynamicMaterialEditor/Private/Utils/DMInputNodeBuilder.cpp b/Engine/Plugins/Experimental/DynamicMaterial/Source/DynamicMaterialEditor/Private/Utils/DMInputNodeBuilder.cpp
--- a/Engine/Plugins/Experimental/DynamicMaterial/Source/DynamicMaterialEditor/Private/Utils/DMInputNodeBuilder.cpp
+++ b/Engine/Plugins/Experimental/DynamicMaterial/Source/DynamicMaterialEditor/Private/Utils/DMInputNodeBuilder.cpp
@@ -8,6 +8,7 @@
 #include "Model/DMMaterialBuildState.h"
 #include "Model/DMMaterialBuildUtils.h"
 #include "Utils/DMPrivate.h"
+#include "Utils/DMUtils.h"
 
 namespace UE::DynamicMaterialEditor::Private
 {
diff --git a/Engine/Plugins/Experimental/DynamicMaterial/Source/DynamicMaterialEditor/Private/Utils/DMMaterialUtils.cpp b/Engine/Plugins/Experimental/DynamicMaterial/Source/DynamicMaterialEditor/Private/Utils/DMMaterialUtils.cpp
new file mode 100644
--- /dev/null
+++ b/Engine/Plugins/Experimental/DynamicMaterial/Source/DynamicMaterialEditor/Private/Utils/DMMaterialUtils.cpp
@@ -0,0 +1,307 @@
+// Copyright Epic Games, Inc. All Rights Reserved.
+
+#include "Utils/DMMaterialUtils.h"
+#include "MaterialShared.h"
+#include "RenderUtils.h"
+
+bool FDMMaterialUtils::IsMaterialPropertyActive(const FParams& InParams)
+{
+	const bool bSubstrateEnabled = Substrate::IsSubstrateEnabled();
+	const bool bSubstrateOpacityOverrideAllowed = InParams.BlendMode == BLEND_AlphaComposite; // Should we always have it enabled to be able to be plugged in an fed when blend mode is toggled later on a material instance?
+
+	switch (InParams.Domain)
+	{
+		case EMaterialDomain::MD_PostProcess:
+			return IsMaterialPropertyActive_PostProcess(InParams, bSubstrateEnabled, bSubstrateOpacityOverrideAllowed);
+
+		case EMaterialDomain::MD_LightFunction:
+			return IsMaterialPropertyActive_LightFunction(InParams, bSubstrateEnabled, bSubstrateOpacityOverrideAllowed);
+
+		case EMaterialDomain::MD_DeferredDecal:
+			return IsMaterialPropertyActive_DeferredDecal(InParams, bSubstrateEnabled, bSubstrateOpacityOverrideAllowed);
+
+		case EMaterialDomain::MD_Volume:
+			return IsMaterialPropertyActive_Volume(InParams, bSubstrateEnabled, bSubstrateOpacityOverrideAllowed);
+
+		case EMaterialDomain::MD_UI:
+			return IsMaterialPropertyActive_UI(InParams, bSubstrateEnabled, bSubstrateOpacityOverrideAllowed);
+
+		case EMaterialDomain::MD_Surface:
+			return IsMaterialPropertyActive_Surface(InParams, bSubstrateEnabled, bSubstrateOpacityOverrideAllowed);
+
+		default:
+			return false;
+	}
+}
+
+bool FDMMaterialUtils::IsMaterialPropertyActive_PostProcess(const FParams& InParams, bool bInSubstrateEnabled, bool bInSubstrateOpacityOverrideAllowed)
+{
+	if (bInSubstrateEnabled)
+	{
+		return InParams.Property == MP_FrontMaterial || (InParams.Property == MP_Opacity && bInSubstrateOpacityOverrideAllowed);
+	}
+
+	return InParams.Property == MP_EmissiveColor || (InParams.bBlendableOutputAlpha && InParams.Property == MP_Opacity);
+}
+
+bool FDMMaterialUtils::IsMaterialPropertyActive_LightFunction(const FParams& InParams, bool bInSubstrateEnabled, bool bInSubstrateOpacityOverrideAllowed)
+{
+	// light functions should already use MSM_Unlit but we also we don't want WorldPosOffset
+	if (bInSubstrateEnabled)
+	{
+		return InParams.Property == MP_FrontMaterial;
+	}
+
+	return InParams.Property == MP_EmissiveColor;
+}
+
+bool FDMMaterialUtils::IsMaterialPropertyActive_DeferredDecal(const FParams& InParams, bool bInSubstrateEnabled, bool bInSubstrateOpacityOverrideAllowed)
+{
+	if (bInSubstrateEnabled)
+	{
+		return InParams.Property == MP_FrontMaterial
+			|| InParams.Property == MP_AmbientOcclusion
+			|| (InParams.Property == MP_Opacity && bInSubstrateOpacityOverrideAllowed);
+	}
+
+	if (InParams.Property >= MP_CustomizedUVs0 && InParams.Property <= MP_CustomizedUVs7)
+	{
+		return true;
+	}
+	
+	if (InParams.Property == MP_MaterialAttributes)
+	{
+		// todo: MaterialAttruibutes would not return true, should it? Why we don't check for the checkbox in the material
+		return true;
+	}
+
+	if (InParams.Property == MP_WorldPositionOffset)
+	{
+		// Note: DeferredDecals don't support this but MeshDecals do
+		return true;
+	}
+
+	switch (InParams.BlendMode)
+	{
+		case BLEND_Translucent:
+			return InParams.Property == MP_EmissiveColor
+				|| InParams.Property == MP_Normal
+				|| InParams.Property == MP_Metallic
+				|| InParams.Property == MP_Specular
+				|| InParams.Property == MP_BaseColor
+				|| InParams.Property == MP_Roughness
+				|| InParams.Property == MP_Opacity
+				|| InParams.Property == MP_AmbientOcclusion;
+
+		case BLEND_AlphaComposite:
+			// AlphaComposite decals never write normal.
+			return InParams.Property == MP_EmissiveColor
+				|| InParams.Property == MP_Metallic
+				|| InParams.Property == MP_Specular
+				|| InParams.Property == MP_BaseColor
+				|| InParams.Property == MP_Roughness
+				|| InParams.Property == MP_Opacity;
+
+		case BLEND_Modulate:
+			return InParams.Property == MP_EmissiveColor
+				|| InParams.Property == MP_Normal
+				|| InParams.Property == MP_Metallic
+				|| InParams.Property == MP_Specular
+				|| InParams.Property == MP_BaseColor
+				|| InParams.Property == MP_Roughness
+				|| InParams.Property == MP_Opacity;
+
+		default:
+			break;
+	}
+
+	return false;
+}
+
+bool FDMMaterialUtils::IsMaterialPropertyActive_Volume(const FParams& InParams, bool bInSubstrateEnabled, bool bInSubstrateOpacityOverrideAllowed)
+{
+	if (bInSubstrateEnabled)
+	{
+		return InParams.Property == MP_FrontMaterial;
+	}
+
+	return InParams.Property == MP_EmissiveColor
+		|| InParams.Property == MP_SubsurfaceColor
+		|| InParams.Property == MP_BaseColor
+		|| InParams.Property == MP_AmbientOcclusion;
+}
+
+bool FDMMaterialUtils::IsMaterialPropertyActive_UI(const FParams& InParams, bool bInSubstrateEnabled, bool bInSubstrateOpacityOverrideAllowed)
+{
+	if (bInSubstrateEnabled)
+	{
+		return InParams.Property == MP_FrontMaterial
+			|| (InParams.Property == MP_WorldPositionOffset)
+			|| (InParams.Property == MP_OpacityMask && IsMaskedBlendMode(InParams.BlendMode))
+			|| (InParams.Property == MP_Opacity && bInSubstrateOpacityOverrideAllowed)
+			|| (InParams.Property >= MP_CustomizedUVs0 && InParams.Property <= MP_CustomizedUVs7);
+	}
+
+	return InParams.Property == MP_EmissiveColor
+		|| (InParams.Property == MP_WorldPositionOffset)
+		|| (InParams.Property == MP_OpacityMask && IsMaskedBlendMode(InParams.BlendMode))
+		|| (InParams.Property == MP_Opacity && IsTranslucentBlendMode(InParams.BlendMode) && InParams.BlendMode != BLEND_Modulate)
+		|| (InParams.Property >= MP_CustomizedUVs0 && InParams.Property <= MP_CustomizedUVs7);
+}
+
+bool FDMMaterialUtils::IsMaterialPropertyActive_Surface(const FParams& InParams, bool bInSubstrateEnabled, bool bInSubstrateOpacityOverrideAllowed)
+{
+	const bool bIsTranslucentBlendMode = IsTranslucentBlendMode(InParams.BlendMode);
+	const bool bIsNonDirectionalTranslucencyLightingMode = InParams.TranslucencyLightingMode == TLM_VolumetricNonDirectional || InParams.TranslucencyLightingMode == TLM_VolumetricPerVertexNonDirectional;
+	const bool bIsVolumetricTranslucencyLightingMode = InParams.TranslucencyLightingMode == TLM_VolumetricNonDirectional
+		|| InParams.TranslucencyLightingMode == TLM_VolumetricDirectional
+		|| InParams.TranslucencyLightingMode == TLM_VolumetricPerVertexNonDirectional
+		|| InParams.TranslucencyLightingMode == TLM_VolumetricPerVertexDirectional;
+
+	bool Active = true;
+
+	if (bInSubstrateEnabled)
+	{
+		Active = false;
+
+		if (InParams.bIsSupported)
+		{
+			switch (InParams.Property)
+			{
+				case MP_Refraction:
+					Active = (bIsTranslucentBlendMode && !IsAlphaHoldoutBlendMode(InParams.BlendMode) && !IsModulateBlendMode(InParams.BlendMode) && InParams.bUsesDistortion) || InParams.ShadingModels.HasShadingModel(MSM_SingleLayerWater);
+					break;
+
+				case MP_Opacity:
+					// Opacity is used as alpha override for alpha composite blending. 
+					Active = bInSubstrateOpacityOverrideAllowed;
+					break;
+
+				case MP_OpacityMask:
+					Active = IsMaskedBlendMode(InParams.BlendMode);
+					break;
+
+				case MP_AmbientOcclusion:
+					Active = InParams.ShadingModels.IsLit();
+					break;
+
+				case MP_Displacement:
+					Active = InParams.bIsTessellationEnabled;
+					break;
+
+				case MP_PixelDepthOffset:
+					Active = (!bIsTranslucentBlendMode) || (InParams.bIsTranslucencyWritingVelocity);
+					break;
+
+				case MP_SurfaceThickness:
+					Active = InParams.bIsThinSurface;
+					break;
+
+				case MP_WorldPositionOffset:
+				case MP_FrontMaterial:
+				case MP_MaterialAttributes:
+					Active = true;
+					break;
+
+				default:
+					if (InParams.Property >= MP_CustomizedUVs0 && InParams.Property <= MP_CustomizedUVs7)
+					{
+						Active = true;
+					}
+			}
+		}
+	}
+	else
+	{
+		switch (InParams.Property)
+		{
+			case MP_Refraction:
+				Active = (bIsTranslucentBlendMode && !IsAlphaHoldoutBlendMode(InParams.BlendMode) && !IsModulateBlendMode(InParams.BlendMode) && InParams.bUsesDistortion) || InParams.ShadingModels.HasShadingModel(MSM_SingleLayerWater);
+				break;
+
+			case MP_Opacity:
+				Active = (bIsTranslucentBlendMode && !IsModulateBlendMode(InParams.BlendMode)) || InParams.ShadingModels.HasShadingModel(MSM_SingleLayerWater);
+				if (IsSubsurfaceShadingModel(InParams.ShadingModels))
+				{
+					Active = true;
+				}
+				break;
+
+			case MP_OpacityMask:
+				Active = IsMaskedBlendMode(InParams.BlendMode);
+				break;
+
+			case MP_BaseColor:
+			case MP_AmbientOcclusion:
+				Active = InParams.ShadingModels.IsLit();
+				break;
+
+			case MP_Specular:
+			case MP_Roughness:
+				Active = InParams.ShadingModels.IsLit() && (!bIsTranslucentBlendMode || !bIsVolumetricTranslucencyLightingMode);
+				break;
+
+			case MP_Anisotropy:
+				Active = InParams.ShadingModels.HasAnyShadingModel({MSM_DefaultLit, MSM_ClearCoat}) && (!bIsTranslucentBlendMode || !bIsVolumetricTranslucencyLightingMode);
+				break;
+
+			case MP_Metallic:
+				// Subsurface models store opacity in place of Metallic in the GBuffer
+				Active = InParams.ShadingModels.IsLit() && (!bIsTranslucentBlendMode || !bIsVolumetricTranslucencyLightingMode);
+				break;
+
+			case MP_Normal:
+				Active = (InParams.ShadingModels.IsLit() && (!bIsTranslucentBlendMode || !bIsNonDirectionalTranslucencyLightingMode)) || InParams.bUsesDistortion;
+				break;
+
+			case MP_Tangent:
+				Active = InParams.ShadingModels.HasAnyShadingModel({MSM_DefaultLit, MSM_ClearCoat}) && (!bIsTranslucentBlendMode || !bIsVolumetricTranslucencyLightingMode);
+				break;
+
+			case MP_SubsurfaceColor:
+				Active = InParams.ShadingModels.HasAnyShadingModel({MSM_Subsurface, MSM_PreintegratedSkin, MSM_TwoSidedFoliage, MSM_Cloth});
+				break;
+
+			case MP_CustomData0:
+				Active = InParams.ShadingModels.HasAnyShadingModel({MSM_ClearCoat, MSM_Hair, MSM_Cloth, MSM_Eye, MSM_SubsurfaceProfile});
+				break;
+
+			case MP_CustomData1:
+				Active = InParams.ShadingModels.HasAnyShadingModel({MSM_ClearCoat, MSM_Eye});
+				break;
+
+			case MP_EmissiveColor:
+				// Emissive is always active, even for light functions and post process materials, 
+				// but not for AlphaHoldout
+				Active = !IsAlphaHoldoutBlendMode(InParams.BlendMode);
+				break;
+
+			case MP_Displacement:
+				Active = InParams.bIsTessellationEnabled;
+				break;
+
+			case MP_PixelDepthOffset:
+				Active = (!bIsTranslucentBlendMode) || (InParams.bIsTranslucencyWritingVelocity);
+				break;
+
+			case MP_ShadingModel:
+				Active = InParams.bUsesShadingModelFromMaterialExpression;
+				break;
+
+			case MP_DiffuseColor:
+			case MP_SpecularColor:
+			case MP_SurfaceThickness:
+			case MP_FrontMaterial:
+				Active = false;
+				break;
+
+			case MP_WorldPositionOffset:
+			case MP_MaterialAttributes:
+			default:
+				Active = true;
+				break;
+		}
+	}
+
+	return Active;
+}
diff --git a/Engine/Plugins/Experimental/DynamicMaterial/Source/DynamicMaterialEditor/Private/Utils/DMMaterialUtils.h b/Engine/Plugins/Experimental/DynamicMaterial/Source/DynamicMaterialEditor/Private/Utils/DMMaterialUtils.h
new file mode 100644
--- /dev/null
+++ b/Engine/Plugins/Experimental/DynamicMaterial/Source/DynamicMaterialEditor/Private/Utils/DMMaterialUtils.h
@@ -0,0 +1,42 @@
+// Copyright Epic Games, Inc. All Rights Reserved.
+
+#pragma once
+
+#include "Engine/EngineTypes.h"
+#include "HAL/Platform.h"
+#include "MaterialDomain.h"
+#include "SceneTypes.h"
+
+struct FDMMaterialUtils
+{
+	struct FParams
+	{
+		EMaterialProperty Property;
+		EMaterialDomain Domain;
+		EBlendMode BlendMode;
+		FMaterialShadingModelField ShadingModels;
+		ETranslucencyLightingMode TranslucencyLightingMode;
+		bool bIsTessellationEnabled;
+		bool bBlendableOutputAlpha;
+		bool bUsesDistortion;
+		bool bUsesShadingModelFromMaterialExpression;
+		bool bIsTranslucencyWritingVelocity;
+		bool bIsThinSurface;
+		bool bIsSupported;
+	};
+
+	static bool IsMaterialPropertyActive(const FParams& InParams);		
+
+protected:
+	static bool IsMaterialPropertyActive_PostProcess(const FParams& InParams, bool bSubstrateEnabled, bool bSubstrateOpacityOverrideAllowed);
+
+	static bool IsMaterialPropertyActive_LightFunction(const FParams& InParams, bool bSubstrateEnabled, bool bSubstrateOpacityOverrideAllowed);
+
+	static bool IsMaterialPropertyActive_DeferredDecal(const FParams& InParams, bool bSubstrateEnabled, bool bSubstrateOpacityOverrideAllowed);
+
+	static bool IsMaterialPropertyActive_Volume(const FParams& InParams, bool bSubstrateEnabled, bool bSubstrateOpacityOverrideAllowed);
+
+	static bool IsMaterialPropertyActive_UI(const FParams& InParams, bool bSubstrateEnabled, bool bSubstrateOpacityOverrideAllowed);
+
+	static bool IsMaterialPropertyActive_Surface(const FParams& InParams, bool bSubstrateEnabled, bool bSubstrateOpacityOverrideAllowed);
+};
diff --git a/Engine/Plugins/Experimental/DynamicMaterial/Source/DynamicMaterialEditor/Public/Components/DMMaterialProperty.h b/Engine/Plugins/Experimental/DynamicMaterial/Source/DynamicMaterialEditor/Public/Components/DMMaterialProperty.h
--- a/Engine/Plugins/Experimental/DynamicMaterial/Source/DynamicMaterialEditor/Public/Components/DMMaterialProperty.h
+++ b/Engine/Plugins/Experimental/DynamicMaterial/Source/DynamicMaterialEditor/Public/Components/DMMaterialProperty.h
@@ -93,7 +93,7 @@ public:
 	 * @param InMaterialModel The model with the settings.
 	 * @return True if it is valid.
 	 */
-	virtual bool IsValidForModel(UDynamicMaterialModelEditorOnlyData& InMaterialModel) const { return true; }
+	virtual bool IsValidForModel(UDynamicMaterialModelEditorOnlyData& InMaterialModel) const;
 
 	/**
 	 * Empty the connection map and recalculate the input types.
diff --git a/Engine/Plugins/Experimental/DynamicMaterial/Source/DynamicMaterialEditor/Public/Model/DynamicMaterialModelEditorOnlyData.h b/Engine/Plugins/Experimental/DynamicMaterial/Source/DynamicMaterialEditor/Public/Model/DynamicMaterialModelEditorOnlyData.h
--- a/Engine/Plugins/Experimental/DynamicMaterial/Source/DynamicMaterialEditor/Public/Model/DynamicMaterialModelEditorOnlyData.h
+++ b/Engine/Plugins/Experimental/DynamicMaterial/Source/DynamicMaterialEditor/Public/Model/DynamicMaterialModelEditorOnlyData.h
@@ -306,6 +306,7 @@ protected:
 	TArray<FName> GetPresetOptions() const;
 
 	void OnChannelListPresetChanged();
+	void EnsurePresetSlots();
 	void OnDomainChanged();
 	void OnBlendModeChanged();
 	void OnShadingModelChanged();
