Material Designer
- Slight redesign of the toolbar. Save buttons are now on the left. Content browser buttons moved.

diff --git a/Engine/Plugins/Experimental/DynamicMaterial/Source/DynamicMaterialEditor/Private/Slate/SDMToolBar.cpp b/Engine/Plugins/Experimental/DynamicMaterial/Source/DynamicMaterialEditor/Private/Slate/SDMToolBar.cpp
--- a/Engine/Plugins/Experimental/DynamicMaterial/Source/DynamicMaterialEditor/Private/Slate/SDMToolBar.cpp
+++ b/Engine/Plugins/Experimental/DynamicMaterial/Source/DynamicMaterialEditor/Private/Slate/SDMToolBar.cpp
@@ -75,6 +75,7 @@ TSharedRef<SWidget> SDMToolBar::CreateToolBarEntries()
 			.UseAllottedSize(true)
 			.HAlign(HAlign_Left)
 			.InnerSlotPadding(FVector2D(5.0f))
+
 			+ SWrapBox::Slot()
 			.HAlign(HAlign_Left)
 			.VAlign(VAlign_Center)
@@ -84,33 +85,95 @@ TSharedRef<SWidget> SDMToolBar::CreateToolBarEntries()
 				.Visibility(this, &SDMToolBar::GetActorVisibility)
 				+ SHorizontalBox::Slot()
 				.AutoWidth()
-				.Padding(0.f, 0.f, 5.f, 0.f)
+				.Padding(5.0f, 0.0f, 0.0f, 0.0f)
+				.VAlign(EVerticalAlignment::VAlign_Center)
 				[
-					SNew(STextBlock)
-					.TextStyle(FDynamicMaterialEditorStyle::Get(), "ActorName")
-					.Text(LOCTEXT("Actor", "Actor"))
-				]
+					SAssignNew(SaveButton, SButton)
+					.IsEnabled(this, &SDMToolBar::CanSave)
+					.ContentPadding(GetLargeIconToolBarButtonContentPadding())
+					.ButtonStyle(FDynamicMaterialEditorStyle::Get(), "HoverHintOnly")
+					.ToolTipText(LOCTEXT("MaterialDesignerSaveTooltip", "Save the Material Designer asset\n\nCaution: If this asset lives inside an actor, the actor/level will be saved."))
+					.OnClicked(this, &SDMToolBar::OnSaveClicked)
+					[
+						SNew(SImage)
+						.Image(this, &SDMToolBar::GetSaveIcon)
+						.DesiredSizeOverride(GetLargeIconToolBarButtonSize())
+					]
+				]		
 				+ SHorizontalBox::Slot()
 				.AutoWidth()
-				.Padding(0.f, 0.f, 0.f, 0.f)
+				.Padding(5.0f, 0.0f, 0.0f, 0.0f)
+				.VAlign(EVerticalAlignment::VAlign_Center)
 				[
-					SNew(STextBlock)
-					.TextStyle(FDynamicMaterialEditorStyle::Get(), "ActorName")
-					.Text(this, &SDMToolBar::GetActorName)
+					SNew(SButton)
+					.ContentPadding(GetLargeIconToolBarButtonContentPadding())
+					.ButtonStyle(FDynamicMaterialEditorStyle::Get(), "HoverHintOnly")
+					.ToolTipText(LOCTEXT("ExportMaterialInstance", "Save As"))
+					.Visibility(this, &SDMToolBar::GetExportMaterialInstanceButtonVisibility)
+					.OnClicked(this, &SDMToolBar::OnExportMaterialInstanceButtonClicked)
+					[
+						SNew(SImage)
+						.Image(FAppStyle::GetBrush("AssetEditor.SaveAssetAs"))
+						.DesiredSizeOverride(GetLargeIconToolBarButtonSize())
+					]
 				]
 			]
+
 			+ SWrapBox::Slot()
-			.FillEmptySpace(false)
 			.HAlign(HAlign_Left)
 			.VAlign(VAlign_Center)
 			.Padding(5.0f, 0.0f, 0.0f, 0.0f)
 			[
-				SAssignNew(SlotSelectorContainer, SBox)
+				SNew(SHorizontalBox)
 				.Visibility(this, &SDMToolBar::GetActorVisibility)
+				+ SHorizontalBox::Slot()
+				.AutoWidth()
+				.Padding(0.f, 0.f, 5.f, 0.f)
+				.VAlign(EVerticalAlignment::VAlign_Center)
 				[
-					CreateSlotsComboBoxWidget()
+					SNew(SImage)
+					.Image(FAppStyle::GetBrush("ClassIcon.Actor"))
+					.DesiredSizeOverride(GetLargeIconToolBarButtonSize())
+				]
+				+ SHorizontalBox::Slot()
+				.AutoWidth()
+				.Padding(0.f, 0.f, 0.f, 0.f)
+				.VAlign(EVerticalAlignment::VAlign_Center)
+				[
+					SNew(STextBlock)
+					.TextStyle(FDynamicMaterialEditorStyle::Get(), "ActorName")
+					.Text(this, &SDMToolBar::GetActorName)
+				]	
+				+ SHorizontalBox::Slot()
+				.AutoWidth()
+				.VAlign(EVerticalAlignment::VAlign_Center)
+				.Padding(5.0f, 0.0f, 0.0f, 0.0f)
+				[
+					SAssignNew(SlotSelectorContainer, SBox)
+					.Visibility(this, &SDMToolBar::GetActorVisibility)
+					[
+						CreateSlotsComboBoxWidget()
+					]
+				]	
+				+ SHorizontalBox::Slot()
+				.AutoWidth()
+				.VAlign(EVerticalAlignment::VAlign_Center)
+				.Padding(5.0f, 0.0f, 0.0f, 0.0f)
+				[
+					SAssignNew(UseButton, SButton)
+					.Visibility(EVisibility::Collapsed)
+					.ContentPadding(GetLargeIconToolBarButtonContentPadding())
+					.ButtonStyle(FDynamicMaterialEditorStyle::Get(), "HoverHintOnly")
+					.ToolTipText(LOCTEXT("MaterialDesignerUseTooltip", "Replace the material in this slot with the one selected in the content browser."))
+					.OnClicked(this, &SDMToolBar::OnUseClicked)
+					[
+						SNew(SImage)
+						.Image(FAppStyle::GetBrush(TEXT("Icons.Use")))
+						.DesiredSizeOverride(GetLargeIconToolBarButtonSize())
+					]
 				]
 			]
+
 			+ SWrapBox::Slot()
 			.HAlign(HAlign_Left)
 			.VAlign(VAlign_Center)
@@ -120,14 +183,24 @@ TSharedRef<SWidget> SDMToolBar::CreateToolBarEntries()
 				.Visibility(this, &SDMToolBar::GetAssetVisibility)
 				+ SHorizontalBox::Slot()
 				.AutoWidth()
-				.Padding(0.f, 0.f, 5.f, 0.f)
+				.VAlign(VAlign_Center)
+				.Padding(5.0f, 0.0f, 0.0f, 0.0f)
 				[
-					SNew(STextBlock)
-					.TextStyle(FDynamicMaterialEditorStyle::Get(), "ActorName")
-					.Text(LOCTEXT("Asset", "Asset"))
+					SAssignNew(BrowseButton, SButton)
+					.Visibility(EVisibility::Collapsed)
+					.ContentPadding(GetLargeIconToolBarButtonContentPadding())
+					.ButtonStyle(FDynamicMaterialEditorStyle::Get(), "HoverHintOnly")
+					.ToolTipText(LOCTEXT("MaterialDesignerBrowseTooltip", "Browse to the selected asset in the content browser."))
+					.OnClicked(this, &SDMToolBar::OnBrowseClicked)
+					[
+						SNew(SImage)
+						.Image(FAppStyle::GetBrush(TEXT("Icons.BrowseContent")))
+						.DesiredSizeOverride(GetLargeIconToolBarButtonSize())
+					]
 				]
 				+ SHorizontalBox::Slot()
 				.AutoWidth()
+				.VAlign(VAlign_Center)
 				.Padding(0.f, 0.f, 0.f, 0.f)
 				[
 					SNew(STextBlock)
@@ -138,63 +211,6 @@ TSharedRef<SWidget> SDMToolBar::CreateToolBarEntries()
 			]
 		]
 		
-		+ SHorizontalBox::Slot()
-		.AutoWidth()
-		.HAlign(HAlign_Right)
-		.VAlign(VAlign_Top)
-		.Padding(5.0f, 0.0f, 0.0f, 0.0f)
-		[
-			SAssignNew(BrowseButton, SButton)
-			.Visibility(EVisibility::Collapsed)
-			.ContentPadding(GetLargeIconToolBarButtonContentPadding())
-			.ButtonStyle(FDynamicMaterialEditorStyle::Get(), "HoverHintOnly")
-			.ToolTipText(LOCTEXT("MaterialDesignerBrowseTooltip", "Browse to the selected asset in the content browser."))
-			.OnClicked(this, &SDMToolBar::OnBrowseClicked)
-			[
-				SNew(SImage)
-				.Image(FAppStyle::GetBrush(TEXT("Icons.BrowseContent")))
-				.DesiredSizeOverride(GetLargeIconToolBarButtonSize())
-			]
-		]
-
-		+ SHorizontalBox::Slot()
-		.AutoWidth()
-		.HAlign(HAlign_Right)
-		.VAlign(VAlign_Top)
-		.Padding(5.0f, 0.0f, 0.0f, 0.0f)
-		[
-			SAssignNew(UseButton, SButton)
-			.Visibility(EVisibility::Collapsed)
-			.ContentPadding(GetLargeIconToolBarButtonContentPadding())
-			.ButtonStyle(FDynamicMaterialEditorStyle::Get(), "HoverHintOnly")
-			.ToolTipText(LOCTEXT("MaterialDesignerUseTooltip", "Replace the material in this slot with the one selected in the content browser."))
-			.OnClicked(this, &SDMToolBar::OnUseClicked)
-			[
-				SNew(SImage)
-				.Image(FAppStyle::GetBrush(TEXT("Icons.Use")))
-				.DesiredSizeOverride(GetLargeIconToolBarButtonSize())
-			]
-		]
-		
-		+ SHorizontalBox::Slot()
-		.AutoWidth()
-		.HAlign(HAlign_Right)
-		.VAlign(VAlign_Top)
-		.Padding(5.0f, 0.0f, 0.0f, 0.0f)
-		[
-			SAssignNew(SaveButton, SButton)
-			.IsEnabled(this, &SDMToolBar::CanSave)
-			.ContentPadding(GetLargeIconToolBarButtonContentPadding())
-			.ButtonStyle(FDynamicMaterialEditorStyle::Get(), "HoverHintOnly")
-			.ToolTipText(LOCTEXT("MaterialDesignerSaveTooltip", "Save the Material Designer asset\n\nCaution: If this asset lives inside an actor, the actor/level will be saved."))
-			.OnClicked(this, &SDMToolBar::OnSaveClicked)
-			[
-				SNew(SImage)
-				.Image(this, &SDMToolBar::GetSaveIcon)
-				.DesiredSizeOverride(GetLargeIconToolBarButtonSize())
-			]
-		]
-		
 		+ SHorizontalBox::Slot()
 		.AutoWidth()
 		.HAlign(HAlign_Right)
@@ -214,25 +230,6 @@ TSharedRef<SWidget> SDMToolBar::CreateToolBarEntries()
 			]
 		]
 		
-		+ SHorizontalBox::Slot()
-		.AutoWidth()
-		.HAlign(HAlign_Right)
-		.VAlign(VAlign_Top)
-		.Padding(5.0f, 0.0f, 0.0f, 0.0f)
-		[
-			SNew(SButton)
-			.ContentPadding(GetLargeIconToolBarButtonContentPadding())
-			.ButtonStyle(FDynamicMaterialEditorStyle::Get(), "HoverHintOnly")
-			.ToolTipText(LOCTEXT("ExportMaterialInstance", "Export Material Designer Instance"))
-			.Visibility(this, &SDMToolBar::GetExportMaterialInstanceButtonVisibility)
-			.OnClicked(this, &SDMToolBar::OnExportMaterialInstanceButtonClicked)
-			[
-				SNew(SImage)
-				.Image(FAppStyle::Get().GetBrush(TEXT("Icons.Toolbar.Export")))
-				.DesiredSizeOverride(GetLargeIconToolBarButtonSize())
-			]
-		]
-		
 		+ SHorizontalBox::Slot()
 		.AutoWidth()
 		.HAlign(HAlign_Right)
