Material Designer
- The slot/element selector is now always shown, even if there's only a single available element.

diff --git a/Engine/Plugins/Experimental/DynamicMaterial/Source/DynamicMaterialEditor/Private/Slate/SDMToolBar.cpp b/Engine/Plugins/Experimental/DynamicMaterial/Source/DynamicMaterialEditor/Private/Slate/SDMToolBar.cpp
--- a/Engine/Plugins/Experimental/DynamicMaterial/Source/DynamicMaterialEditor/Private/Slate/SDMToolBar.cpp
+++ b/Engine/Plugins/Experimental/DynamicMaterial/Source/DynamicMaterialEditor/Private/Slate/SDMToolBar.cpp
@@ -269,7 +269,7 @@ TSharedRef<SWidget> SDMToolBar::CreateToolBarButton(TAttribute<const FSlateBrush
 
 TSharedRef<SWidget> SDMToolBar::CreateSlotsComboBoxWidget()
 {
-	if (!MaterialActorWeak.IsValid() || !MaterialModelWeak.IsValid() || ActorMaterialProperties.Num() <= 1)
+	if (!MaterialActorWeak.IsValid() || !MaterialModelWeak.IsValid())
 	{
 		return SNullWidget::NullWidget;
 	}
@@ -279,6 +279,7 @@ TSharedRef<SWidget> SDMToolBar::CreateSlotsComboBoxWidget()
 
 	return 
 		SNew(SComboBox<TSharedPtr<FDMObjectMaterialProperty>>)
+		.IsEnabled(ActorMaterialProperties.Num() > 1)
 		.InitiallySelectedItem(InitiallySelectedItem)
 		.OptionsSource(&ActorMaterialProperties)
 		.OnGenerateWidget(this, &SDMToolBar::GenerateSelectedMaterialSlotRow)
