Material Designer
- When selecting an actor, the first material will be selected, not the last.

diff --git a/Engine/Plugins/Experimental/DynamicMaterial/Source/DynamicMaterialEditor/Private/Slate/SDMEditor.cpp b/Engine/Plugins/Experimental/DynamicMaterial/Source/DynamicMaterialEditor/Private/Slate/SDMEditor.cpp
--- a/Engine/Plugins/Experimental/DynamicMaterial/Source/DynamicMaterialEditor/Private/Slate/SDMEditor.cpp
+++ b/Engine/Plugins/Experimental/DynamicMaterial/Source/DynamicMaterialEditor/Private/Slate/SDMEditor.cpp
@@ -609,11 +609,11 @@ void SDMEditor::OnActorSelected(AActor* InActor)
 	for (int32 MaterialPropertyIdx = 0; MaterialPropertyIdx < ActorProperties.Num(); ++MaterialPropertyIdx)
 	{
 		const FDMObjectMaterialProperty& MaterialProperty = ActorProperties[MaterialPropertyIdx];
-		UDynamicMaterialModel* Model = MaterialProperty.GetMaterialModel();
 
-		if (UDynamicMaterialModel* MaterialModel = Model)
+		if (MaterialProperty.GetMaterialModel())
 		{
 			SetMaterialObjectProperty(MaterialProperty);
+			break;
 		}
 	}
 
