Material Designer
- Slot selector is now visible when switching between slots in the same actor.

diff --git a/Engine/Plugins/Experimental/DynamicMaterial/Source/DynamicMaterialEditor/Private/Slate/SDMEditor.cpp b/Engine/Plugins/Experimental/DynamicMaterial/Source/DynamicMaterialEditor/Private/Slate/SDMEditor.cpp
--- a/Engine/Plugins/Experimental/DynamicMaterial/Source/DynamicMaterialEditor/Private/Slate/SDMEditor.cpp
+++ b/Engine/Plugins/Experimental/DynamicMaterial/Source/DynamicMaterialEditor/Private/Slate/SDMEditor.cpp
@@ -1315,7 +1315,7 @@ void SDMEditor::OnToolBarPropertyChanged(TSharedPtr<FDMObjectMaterialProperty> I
 {
 	if (InNewSelectedProperty.IsValid())
 	{
-		SetMaterialModel(InNewSelectedProperty->GetMaterialModel());
+		SetMaterialObjectProperty(*InNewSelectedProperty);
 	}
 }
 
