Material Designer
- Change to the way value/textureuv update events are triggered and propagated.

diff --git a/Engine/Plugins/Experimental/DynamicMaterial/Source/DynamicMaterial/Private/Components/DMMaterialComponent.cpp b/Engine/Plugins/Experimental/DynamicMaterial/Source/DynamicMaterial/Private/Components/DMMaterialComponent.cpp
--- a/Engine/Plugins/Experimental/DynamicMaterial/Source/DynamicMaterial/Private/Components/DMMaterialComponent.cpp
+++ b/Engine/Plugins/Experimental/DynamicMaterial/Source/DynamicMaterial/Private/Components/DMMaterialComponent.cpp
@@ -4,7 +4,6 @@
 #include "DMComponentPath.h"
 #include "DynamicMaterialModule.h"
 #include "HAL/PlatformTime.h"
-#include "Model/DynamicMaterialModel.h"
 
 #define LOCTEXT_NAMESPACE "DMMaterialComponent"
 
@@ -264,11 +263,8 @@ void UDMMaterialComponent::OnComponentRemoved()
 
 void UDMMaterialComponent::MarkComponentDirty()
 {
-#if WITH_EDITORONLY_DATA
-	bComponentDirty = true;
-#endif
-
 #if WITH_EDITOR
+	bComponentDirty = true;
 	PreventClean();
 #endif
 }
diff --git a/Engine/Plugins/Experimental/DynamicMaterial/Source/DynamicMaterial/Private/Components/DMMaterialValue.cpp b/Engine/Plugins/Experimental/DynamicMaterial/Source/DynamicMaterial/Private/Components/DMMaterialValue.cpp
--- a/Engine/Plugins/Experimental/DynamicMaterial/Source/DynamicMaterial/Private/Components/DMMaterialValue.cpp
+++ b/Engine/Plugins/Experimental/DynamicMaterial/Source/DynamicMaterial/Private/Components/DMMaterialValue.cpp
@@ -387,7 +387,7 @@ void UDMMaterialValue::PostEditUndo()
  
 	MarkComponentDirty();
  
-	OnValueUpdated(/*bForceStructureUpdate*/ true); // Just in case - Undos are not meant to be quick and easy.
+	OnValueChanged(EDMUpdateType::Structure, /* bInAllowPropagate */ true); // Just in case - Undos are not meant to be quick and easy.
 }
  
 void UDMMaterialValue::PostEditChangeProperty(FPropertyChangedEvent& PropertyChangedEvent)
@@ -410,7 +410,12 @@ void UDMMaterialValue::PostEditChangeProperty(FPropertyChangedEvent& PropertyCha
 	{
 		if (EditableProperty == MemberPropertyName)
 		{
-			OnValueUpdated(/* bForceStructureUpdate */ EditableProperty != ValueName);
+			OnValueChanged(
+				EditableProperty == ValueName
+					? EDMUpdateType::Value
+					: EDMUpdateType::Structure, 
+				/* bInAllowPropagate */ true
+			);
 			return;
 		}
 	}
@@ -430,7 +435,7 @@ UDMMaterialValue::UDMMaterialValue(EDMValueType InType)
 #endif
 }
  
-void UDMMaterialValue::OnValueUpdated(bool bInForceStructureUpdate)
+void UDMMaterialValue::OnValueChanged(EDMUpdateType InUpdateType, bool bInUpdateParent)
 {
 	if (!IsComponentValid())
 	{
@@ -439,7 +444,14 @@ void UDMMaterialValue::OnValueUpdated(bool bInForceStructureUpdate)
 
 	if (FDMUpdateGuard::CanUpdate())
 	{
-		Update(bInForceStructureUpdate ? EDMUpdateType::Structure : EDMUpdateType::Value);
+		Update(InUpdateType);
+
+#if WITH_EDITOR
+		if (bInUpdateParent && ParentComponent)
+		{
+			ParentComponent->Update(InUpdateType);
+		}
+#endif
 	}
 }
  
@@ -455,16 +467,8 @@ void UDMMaterialValue::Update(EDMUpdateType InUpdateType)
 	{
 		return;
 	}
- 
-	if (InUpdateType == EDMUpdateType::Structure)
-	{
-		MarkComponentDirty();
-	}
 
-	if (ParentComponent)
-	{
-		ParentComponent->Update(InUpdateType);
-	}
+	MarkComponentDirty();
 #endif
 
 	Super::Update(InUpdateType);
diff --git a/Engine/Plugins/Experimental/DynamicMaterial/Source/DynamicMaterial/Private/Components/DMRenderTargetRenderer.cpp b/Engine/Plugins/Experimental/DynamicMaterial/Source/DynamicMaterial/Private/Components/DMRenderTargetRenderer.cpp
--- a/Engine/Plugins/Experimental/DynamicMaterial/Source/DynamicMaterial/Private/Components/DMRenderTargetRenderer.cpp
+++ b/Engine/Plugins/Experimental/DynamicMaterial/Source/DynamicMaterial/Private/Components/DMRenderTargetRenderer.cpp
@@ -89,16 +89,13 @@ void UDMRenderTargetRenderer::Update(EDMUpdateType InUpdateType)
 		return;
 	}
 
-	if (InUpdateType == EDMUpdateType::Structure)
-	{
-		MarkComponentDirty();
-	}
+	MarkComponentDirty();
+#endif
 
 	if (UDMMaterialValueRenderTarget* RenderTarget = GetRenderTargetValue())
 	{
 		RenderTarget->Update(InUpdateType);
 	}
-#endif
 
 	Super::Update(InUpdateType);
 }
diff --git a/Engine/Plugins/Experimental/DynamicMaterial/Source/DynamicMaterial/Private/Components/DMTextureUV.cpp b/Engine/Plugins/Experimental/DynamicMaterial/Source/DynamicMaterial/Private/Components/DMTextureUV.cpp
--- a/Engine/Plugins/Experimental/DynamicMaterial/Source/DynamicMaterial/Private/Components/DMTextureUV.cpp
+++ b/Engine/Plugins/Experimental/DynamicMaterial/Source/DynamicMaterial/Private/Components/DMTextureUV.cpp
@@ -7,10 +7,6 @@
 #include "Model/DynamicMaterialModel.h"
 #include "Serialization/CustomVersion.h"
 
-#if WITH_EDITOR
-#include "Model/IDynamicMaterialModelEditorOnlyDataInterface.h"
-#endif
-
 namespace UE::DynamicMaterial::Private
 {
 	TMap<int32, FName> BaseParameterNames = {
@@ -91,7 +87,7 @@ void UDMTextureUV::SetUVSource(EDMUVSource InUVSource)
 
 	UVSource = InUVSource;
 
-	Update(EDMUpdateType::Structure);
+	OnTextureUVChanged(EDMUpdateType::Structure, /* bInAllowPropagate */ true);
 }
 #endif
 
@@ -110,7 +106,7 @@ void UDMTextureUV::SetOffset(const FVector2D& InOffset)
 
 	Offset = InOffset;
 
-	Update(EDMUpdateType::Value);
+	OnTextureUVChanged(EDMUpdateType::Value, /* bInAllowPropagate */ false);
 }
 
 void UDMTextureUV::SetPivot(const FVector2D& InPivot)
@@ -128,7 +124,7 @@ void UDMTextureUV::SetPivot(const FVector2D& InPivot)
 
 	Pivot = InPivot;
 
-	Update(EDMUpdateType::Value);
+	OnTextureUVChanged(EDMUpdateType::Value, /* bInAllowPropagate */ false);
 }
 
 void UDMTextureUV::SetRotation(float InRotation)
@@ -145,7 +141,7 @@ void UDMTextureUV::SetRotation(float InRotation)
 
 	Rotation = InRotation;
 
-	Update(EDMUpdateType::Value);
+	OnTextureUVChanged(EDMUpdateType::Value, /* bInAllowPropagate */ false);
 }
 
 void UDMTextureUV::SetTiling(const FVector2D& InTiling)
@@ -163,7 +159,7 @@ void UDMTextureUV::SetTiling(const FVector2D& InTiling)
 
 	Tiling = InTiling;
 
-	Update(EDMUpdateType::Value);
+	OnTextureUVChanged(EDMUpdateType::Value, /* bInAllowPropagate */ false);
 }
 
 #if WITH_EDITOR
@@ -181,7 +177,7 @@ void UDMTextureUV::SetMirrorOnX(bool bInMirrorOnX)
 
 	bMirrorOnX = bInMirrorOnX;
 
-	Update(EDMUpdateType::Structure);
+	OnTextureUVChanged(EDMUpdateType::Value, /* bInAllowPropagate */ false);
 }
 
 void UDMTextureUV::SetMirrorOnY(bool bInMirrorOnY)
@@ -198,7 +194,7 @@ void UDMTextureUV::SetMirrorOnY(bool bInMirrorOnY)
 
 	bMirrorOnY = bInMirrorOnY;
 
-	Update(EDMUpdateType::Structure);
+	OnTextureUVChanged(EDMUpdateType::Structure, /* bInAllowPropagate */ true);
 }
 #endif
 
@@ -386,15 +382,7 @@ void UDMTextureUV::Update(EDMUpdateType InUpdateType)
 		return;
 	}
 
-	if (InUpdateType == EDMUpdateType::Structure)
-	{
-		MarkComponentDirty();
-	}
-
-	if (ParentComponent)
-	{
-		ParentComponent->Update(InUpdateType);
-	}
+	MarkComponentDirty();
 #endif
 
 	Super::Update(InUpdateType);
@@ -545,6 +533,26 @@ void UDMTextureUV::RemoveParameterNames()
 	}
 }
 
+void UDMTextureUV::OnTextureUVChanged(EDMUpdateType InUpdateType, bool bInUpdateParent)
+{
+	if (!IsComponentValid())
+	{
+		return;
+	}
+
+	if (FDMUpdateGuard::CanUpdate())
+	{
+		Update(InUpdateType);
+
+#if WITH_EDITOR
+		if (bInUpdateParent && ParentComponent)
+		{
+			ParentComponent->Update(InUpdateType);
+		}
+#endif
+	}
+}
+
 void UDMTextureUV::OnComponentAdded()
 {
 	if (!IsComponentValid())
@@ -593,13 +601,13 @@ void UDMTextureUV::PostEditChangeProperty(struct FPropertyChangedEvent& Property
 		|| PropertyChangedEvent.MemberProperty->GetFName() == NAME_Rotation
 		|| PropertyChangedEvent.MemberProperty->GetFName() == NAME_Tiling)
 	{
-		Update(EDMUpdateType::Value);
+		OnTextureUVChanged(EDMUpdateType::Value, /* bInAllowPropagate */ true);
 	}
 	else if (PropertyChangedEvent.MemberProperty->GetFName() == NAME_UVSource
 		|| PropertyChangedEvent.MemberProperty->GetFName() == NAME_bMirrorOnX
 		|| PropertyChangedEvent.MemberProperty->GetFName() == NAME_bMirrorOnY)
 	{
-		Update(EDMUpdateType::Structure);
+		OnTextureUVChanged(EDMUpdateType::Structure, /* bInAllowPropagate */ true);
 	}
 }
 
@@ -620,11 +628,11 @@ void UDMTextureUV::PostEditUndo()
 		|| bMirrorOnX != bMirrorOnX_PreUndo
 		|| bMirrorOnY != bMirrorOnY_PreUndo)
 	{
-		Update(EDMUpdateType::Structure);
+		OnTextureUVChanged(EDMUpdateType::Structure, /* bInAllowPropagate */ true);
 	}
 	else
 	{
-		Update(EDMUpdateType::Value);
+		OnTextureUVChanged(EDMUpdateType::Value, /* bInAllowPropagate */ true);
 	}
 }
 
@@ -660,11 +668,11 @@ void UDMTextureUV::PostLoad()
 
 	if (bNeedsPostLoadStructureUpdate)
 	{
-		Update(EDMUpdateType::Structure);
+		OnTextureUVChanged(EDMUpdateType::Structure, /* bInAllowPropagate */ false);
 	}
 	else if (bNeedsPostLoadValueUpdate)
 	{
-		Update(EDMUpdateType::Value);
+		OnTextureUVChanged(EDMUpdateType::Value, /* bInAllowPropagate */ false);
 	}
 
 	bNeedsPostLoadStructureUpdate = false;
diff --git a/Engine/Plugins/Experimental/DynamicMaterial/Source/DynamicMaterial/Private/Components/MaterialValues/DMMaterialValueBool.cpp b/Engine/Plugins/Experimental/DynamicMaterial/Source/DynamicMaterial/Private/Components/MaterialValues/DMMaterialValueBool.cpp
--- a/Engine/Plugins/Experimental/DynamicMaterial/Source/DynamicMaterial/Private/Components/MaterialValues/DMMaterialValueBool.cpp
+++ b/Engine/Plugins/Experimental/DynamicMaterial/Source/DynamicMaterial/Private/Components/MaterialValues/DMMaterialValueBool.cpp
@@ -112,7 +112,7 @@ void UDMMaterialValueBool::SetValue(bool InValue)
 
 	Value = InValue;
  
-	OnValueUpdated(/* bForceStructureUpdate */ false);
+	OnValueChanged(EDMUpdateType::Value, /* bInAllowPropagate */ false);
 }
  
 #undef LOCTEXT_NAMESPACE
diff --git a/Engine/Plugins/Experimental/DynamicMaterial/Source/DynamicMaterial/Private/Components/MaterialValues/DMMaterialValueColorAtlas.cpp b/Engine/Plugins/Experimental/DynamicMaterial/Source/DynamicMaterial/Private/Components/MaterialValues/DMMaterialValueColorAtlas.cpp
--- a/Engine/Plugins/Experimental/DynamicMaterial/Source/DynamicMaterial/Private/Components/MaterialValues/DMMaterialValueColorAtlas.cpp
+++ b/Engine/Plugins/Experimental/DynamicMaterial/Source/DynamicMaterial/Private/Components/MaterialValues/DMMaterialValueColorAtlas.cpp
@@ -47,7 +47,7 @@ void UDMMaterialValueColorAtlas::SetValue(float InValue)
 
 	Value = InValue;
 
-	OnValueUpdated(/* bForceStructureUpdate */ false);
+	OnValueChanged(EDMUpdateType::Value, /* bInAllowPropagate */ false);
 }
 
 void UDMMaterialValueColorAtlas::SetMIDParameter(UMaterialInstanceDynamic* InMID) const
@@ -77,7 +77,7 @@ void UDMMaterialValueColorAtlas::SetAtlas(UCurveLinearColorAtlas* InAtlas)
 
 	Atlas = InAtlas;
 
-	OnValueUpdated(/* bForceStructureUpdate */ true);
+	OnValueChanged(EDMUpdateType::Structure, /* bInAllowPropagate */ true);
 }
 
 void UDMMaterialValueColorAtlas::SetCurve(UCurveLinearColor* InCurve)
@@ -94,7 +94,7 @@ void UDMMaterialValueColorAtlas::SetCurve(UCurveLinearColor* InCurve)
 
 	Curve = InCurve;
 
-	OnValueUpdated(/* bForceStructureUpdate */ true);
+	OnValueChanged(EDMUpdateType::Structure, /* bInAllowPropagate */ true);
 }
 
 void UDMMaterialValueColorAtlas::GenerateExpression(const TSharedRef<IDMMaterialBuildStateInterface>& InBuildState) const
@@ -198,7 +198,7 @@ bool UDMMaterialValueColorAtlas::JsonDeserialize(const TSharedPtr<FJsonValue>& I
 
 	if (bSuccess)
 	{
-		Update(UpdateType);
+		OnValueChanged(UpdateType, /* bInAllowPropagate */ true);
 	}
 
 	return bSuccess;
diff --git a/Engine/Plugins/Experimental/DynamicMaterial/Source/DynamicMaterial/Private/Components/MaterialValues/DMMaterialValueFloat1.cpp b/Engine/Plugins/Experimental/DynamicMaterial/Source/DynamicMaterial/Private/Components/MaterialValues/DMMaterialValueFloat1.cpp
--- a/Engine/Plugins/Experimental/DynamicMaterial/Source/DynamicMaterial/Private/Components/MaterialValues/DMMaterialValueFloat1.cpp
+++ b/Engine/Plugins/Experimental/DynamicMaterial/Source/DynamicMaterial/Private/Components/MaterialValues/DMMaterialValueFloat1.cpp
@@ -99,7 +99,7 @@ void UDMMaterialValueFloat1::SetValue(float InValue)
  
 	Value = InValue;
  
-	OnValueUpdated(/* bForceStructureUpdate */ false);
+	OnValueChanged(EDMUpdateType::Value, /* bInAllowPropagate */ false);
 }
  
 void UDMMaterialValueFloat1::SetMIDParameter(UMaterialInstanceDynamic* InMID) const
diff --git a/Engine/Plugins/Experimental/DynamicMaterial/Source/DynamicMaterial/Private/Components/MaterialValues/DMMaterialValueFloat2.cpp b/Engine/Plugins/Experimental/DynamicMaterial/Source/DynamicMaterial/Private/Components/MaterialValues/DMMaterialValueFloat2.cpp
--- a/Engine/Plugins/Experimental/DynamicMaterial/Source/DynamicMaterial/Private/Components/MaterialValues/DMMaterialValueFloat2.cpp
+++ b/Engine/Plugins/Experimental/DynamicMaterial/Source/DynamicMaterial/Private/Components/MaterialValues/DMMaterialValueFloat2.cpp
@@ -112,7 +112,7 @@ void UDMMaterialValueFloat2::SetValue(const FVector2D& InValue)
  
 	Value = ValueClamped;
  
-	OnValueUpdated(/* bForceStructureUpdate */ false);
+	OnValueChanged(EDMUpdateType::Value, /* bInAllowPropagate */ false);
 }
  
 #if WITH_EDITOR
diff --git a/Engine/Plugins/Experimental/DynamicMaterial/Source/DynamicMaterial/Private/Components/MaterialValues/DMMaterialValueFloat3RGB.cpp b/Engine/Plugins/Experimental/DynamicMaterial/Source/DynamicMaterial/Private/Components/MaterialValues/DMMaterialValueFloat3RGB.cpp
--- a/Engine/Plugins/Experimental/DynamicMaterial/Source/DynamicMaterial/Private/Components/MaterialValues/DMMaterialValueFloat3RGB.cpp
+++ b/Engine/Plugins/Experimental/DynamicMaterial/Source/DynamicMaterial/Private/Components/MaterialValues/DMMaterialValueFloat3RGB.cpp
@@ -112,7 +112,7 @@ void UDMMaterialValueFloat3RGB::SetValue(const FLinearColor& InValue)
 
 	Value = ValueClamped;
 
-	OnValueUpdated(/* bForceStructureUpdate */ false);
+	OnValueChanged(EDMUpdateType::Value, /* bInAllowPropagate */ false);
 }
  
 #if WITH_EDITOR
diff --git a/Engine/Plugins/Experimental/DynamicMaterial/Source/DynamicMaterial/Private/Components/MaterialValues/DMMaterialValueFloat3RPY.cpp b/Engine/Plugins/Experimental/DynamicMaterial/Source/DynamicMaterial/Private/Components/MaterialValues/DMMaterialValueFloat3RPY.cpp
--- a/Engine/Plugins/Experimental/DynamicMaterial/Source/DynamicMaterial/Private/Components/MaterialValues/DMMaterialValueFloat3RPY.cpp
+++ b/Engine/Plugins/Experimental/DynamicMaterial/Source/DynamicMaterial/Private/Components/MaterialValues/DMMaterialValueFloat3RPY.cpp
@@ -114,7 +114,7 @@ void UDMMaterialValueFloat3RPY::SetValue(const FRotator& InValue)
  
 	Value = ValueClamped;
  
-	OnValueUpdated(/* bForceStructureUpdate */ false);
+	OnValueChanged(EDMUpdateType::Value, /* bInAllowPropagate */ false);
 }
  
 #if WITH_EDITOR
diff --git a/Engine/Plugins/Experimental/DynamicMaterial/Source/DynamicMaterial/Private/Components/MaterialValues/DMMaterialValueFloat3XYZ.cpp b/Engine/Plugins/Experimental/DynamicMaterial/Source/DynamicMaterial/Private/Components/MaterialValues/DMMaterialValueFloat3XYZ.cpp
--- a/Engine/Plugins/Experimental/DynamicMaterial/Source/DynamicMaterial/Private/Components/MaterialValues/DMMaterialValueFloat3XYZ.cpp
+++ b/Engine/Plugins/Experimental/DynamicMaterial/Source/DynamicMaterial/Private/Components/MaterialValues/DMMaterialValueFloat3XYZ.cpp
@@ -111,7 +111,7 @@ void UDMMaterialValueFloat3XYZ::SetValue(const FVector& InValue)
  
 	Value = ValueClamped;
  
-	OnValueUpdated(/* bForceStructureUpdate */ false);
+	OnValueChanged(EDMUpdateType::Value, /* bInAllowPropagate */ false);
 }
  
 #if WITH_EDITOR
diff --git a/Engine/Plugins/Experimental/DynamicMaterial/Source/DynamicMaterial/Private/Components/MaterialValues/DMMaterialValueFloat4.cpp b/Engine/Plugins/Experimental/DynamicMaterial/Source/DynamicMaterial/Private/Components/MaterialValues/DMMaterialValueFloat4.cpp
--- a/Engine/Plugins/Experimental/DynamicMaterial/Source/DynamicMaterial/Private/Components/MaterialValues/DMMaterialValueFloat4.cpp
+++ b/Engine/Plugins/Experimental/DynamicMaterial/Source/DynamicMaterial/Private/Components/MaterialValues/DMMaterialValueFloat4.cpp
@@ -114,7 +114,7 @@ void UDMMaterialValueFloat4::SetValue(const FLinearColor& InValue)
  
 	Value = ValueClamped;
  
-	OnValueUpdated(/* bForceStructureUpdate */ false);
+	OnValueChanged(EDMUpdateType::Value, /* bInAllowPropagate */ false);
 }
  
 #if WITH_EDITOR
diff --git a/Engine/Plugins/Experimental/DynamicMaterial/Source/DynamicMaterial/Private/Components/MaterialValues/DMMaterialValueRenderTarget.cpp b/Engine/Plugins/Experimental/DynamicMaterial/Source/DynamicMaterial/Private/Components/MaterialValues/DMMaterialValueRenderTarget.cpp
--- a/Engine/Plugins/Experimental/DynamicMaterial/Source/DynamicMaterial/Private/Components/MaterialValues/DMMaterialValueRenderTarget.cpp
+++ b/Engine/Plugins/Experimental/DynamicMaterial/Source/DynamicMaterial/Private/Components/MaterialValues/DMMaterialValueRenderTarget.cpp
@@ -91,19 +91,21 @@ UDMRenderTargetRenderer* UDMMaterialValueRenderTarget::GetRenderer() const
 
 void UDMMaterialValueRenderTarget::SetRenderer(UDMRenderTargetRenderer* InRenderer)
 {
-#if WITH_EDITOR
 	if (Renderer == InRenderer)
 	{
 		return;
 	}
 
+#if WITH_EDITOR
 	if (Renderer)
 	{
 		Renderer->SetComponentState(EDMComponentLifetimeState::Removed);
 	}
+#endif
 
 	Renderer = InRenderer;
 
+#if WITH_EDITOR
 	if (IsComponentAdded())
 	{
 		Renderer->SetComponentState(EDMComponentLifetimeState::Added);
@@ -222,7 +224,7 @@ bool UDMMaterialValueRenderTarget::JsonDeserialize(const TSharedPtr<FJsonValue>&
 
 	if (bSuccess)
 	{
-		Update(EDMUpdateType::Value);
+		OnValueChanged(EDMUpdateType::Structure, /* bInAllowPropagate */ true);
 	}
 
 	return bSuccess;
diff --git a/Engine/Plugins/Experimental/DynamicMaterial/Source/DynamicMaterial/Private/Components/MaterialValues/DMMaterialValueTexture.cpp b/Engine/Plugins/Experimental/DynamicMaterial/Source/DynamicMaterial/Private/Components/MaterialValues/DMMaterialValueTexture.cpp
--- a/Engine/Plugins/Experimental/DynamicMaterial/Source/DynamicMaterial/Private/Components/MaterialValues/DMMaterialValueTexture.cpp
+++ b/Engine/Plugins/Experimental/DynamicMaterial/Source/DynamicMaterial/Private/Components/MaterialValues/DMMaterialValueTexture.cpp
@@ -146,11 +146,15 @@ void UDMMaterialValueTexture::PostEditChangeProperty(FPropertyChangedEvent& InPr
 				const TextureCompressionSettings CurrentCompression = OldValue ? OldValue->CompressionSettings.GetValue() : TextureCompressionSettings::TC_MAX;
 				const TextureCompressionSettings NewCompression = Value ? Value->CompressionSettings.GetValue() : TextureCompressionSettings::TC_MAX;
 
-				OnValueUpdated(CurrentCompression != NewCompression);
+				const EDMUpdateType UpdateType = CurrentCompression == NewCompression
+					? EDMUpdateType::Structure
+					: EDMUpdateType::Value;
+
+				OnValueChanged(UpdateType, /* bInAllowPropagate */ true);
 				return;
 			}
 
-			OnValueUpdated(/* bForceStructureUpdate */ true);
+			OnValueChanged(EDMUpdateType::Structure, /* bInAllowPropagate */ true);
 			return;
 		}
 	}
@@ -221,12 +225,18 @@ void UDMMaterialValueTexture::SetValue(UTexture* InValue)
 		return;
 	}
 
+	/*
 	const TextureCompressionSettings CurrentCompression = Value ? Value->CompressionSettings.GetValue() : TextureCompressionSettings::TC_MAX;
 	const TextureCompressionSettings NewCompression = InValue ? InValue->CompressionSettings.GetValue() : TextureCompressionSettings::TC_MAX;
 
+	const EDMUpdateType UpdateType = CurrentCompression == NewCompression
+		? EDMUpdateType::Structure
+		: EDMUpdateType::Value;
+	*/
+
 	Value = InValue;
 
-	OnValueUpdated(CurrentCompression != NewCompression);
+	OnValueChanged(EDMUpdateType::Value, /* bInAllowPropagate */ false);
 }
 
 void UDMMaterialValueTexture::SetMIDParameter(UMaterialInstanceDynamic* InMID) const
diff --git a/Engine/Plugins/Experimental/DynamicMaterial/Source/DynamicMaterial/Private/Model/DynamicMaterialModel.cpp b/Engine/Plugins/Experimental/DynamicMaterial/Source/DynamicMaterial/Private/Model/DynamicMaterialModel.cpp
--- a/Engine/Plugins/Experimental/DynamicMaterial/Source/DynamicMaterial/Private/Model/DynamicMaterialModel.cpp
+++ b/Engine/Plugins/Experimental/DynamicMaterial/Source/DynamicMaterial/Private/Model/DynamicMaterialModel.cpp
@@ -14,7 +14,6 @@
 #include "Model/IDynamicMaterialModelEditorOnlyDataInterface.h"
 
 #if WITH_EDITOR
-#include "Components/DMTextureUV.h"
 #include "Materials/Material.h"
 #endif
 
diff --git a/Engine/Plugins/Experimental/DynamicMaterial/Source/DynamicMaterial/Private/Utils/DMUtils.cpp b/Engine/Plugins/Experimental/DynamicMaterial/Source/DynamicMaterial/Private/Utils/DMUtils.cpp
--- a/Engine/Plugins/Experimental/DynamicMaterial/Source/DynamicMaterial/Private/Utils/DMUtils.cpp
+++ b/Engine/Plugins/Experimental/DynamicMaterial/Source/DynamicMaterial/Private/Utils/DMUtils.cpp
@@ -7,6 +7,7 @@
 #include "SceneTypes.h"
 
 #if WITH_EDITOR
+#include "Misc/CoreMisc.h"
 #include "Misc/Paths.h"
 #endif
 
diff --git a/Engine/Plugins/Experimental/DynamicMaterial/Source/DynamicMaterial/Public/Components/DMMaterialValue.h b/Engine/Plugins/Experimental/DynamicMaterial/Source/DynamicMaterial/Public/Components/DMMaterialValue.h
--- a/Engine/Plugins/Experimental/DynamicMaterial/Source/DynamicMaterial/Public/Components/DMMaterialValue.h
+++ b/Engine/Plugins/Experimental/DynamicMaterial/Source/DynamicMaterial/Public/Components/DMMaterialValue.h
@@ -174,7 +174,7 @@ protected:
  
 	UDMMaterialValue(EDMValueType InType);
  
-	virtual void OnValueUpdated(bool bInForceStructureUpdate);
+	virtual void OnValueChanged(EDMUpdateType InUpdateType, bool bInUpdateParent);
  
 	//~ Begin UDMMaterialComponent
 	virtual UDMMaterialComponent* GetSubComponentByPath(FDMComponentPath& InPath, const FDMComponentPathSegment& InPathSegment) const override;
diff --git a/Engine/Plugins/Experimental/DynamicMaterial/Source/DynamicMaterial/Public/Components/DMTextureUV.h b/Engine/Plugins/Experimental/DynamicMaterial/Source/DynamicMaterial/Public/Components/DMTextureUV.h
--- a/Engine/Plugins/Experimental/DynamicMaterial/Source/DynamicMaterial/Public/Components/DMTextureUV.h
+++ b/Engine/Plugins/Experimental/DynamicMaterial/Source/DynamicMaterial/Public/Components/DMTextureUV.h
@@ -219,6 +219,8 @@ protected:
 	void RemoveParameterNames();
 #endif
 
+	void OnTextureUVChanged(EDMUpdateType InUpdateType, bool bInUpdateParent);
+
 	//~ Begin UDMMaterialComponent
 	virtual void Update(EDMUpdateType InUpdateType) override;
 	virtual UDMMaterialComponent* GetSubComponentByPath(FDMComponentPath& InPath, const FDMComponentPathSegment& InPathSegment) const override;
diff --git a/Engine/Plugins/Experimental/DynamicMaterial/Source/DynamicMaterial/Public/Components/MaterialValues/DMMaterialValueTexture.h b/Engine/Plugins/Experimental/DynamicMaterial/Source/DynamicMaterial/Public/Components/MaterialValues/DMMaterialValueTexture.h
--- a/Engine/Plugins/Experimental/DynamicMaterial/Source/DynamicMaterial/Public/Components/MaterialValues/DMMaterialValueTexture.h
+++ b/Engine/Plugins/Experimental/DynamicMaterial/Source/DynamicMaterial/Public/Components/MaterialValues/DMMaterialValueTexture.h
@@ -76,7 +76,7 @@ protected:
 	UDMMaterialValueTexture();
  
 	UPROPERTY(EditAnywhere, BlueprintReadWrite, NoClear, Getter = GetValue, Setter = SetValue, BlueprintSetter = SetValue, Category = "Material Designer",
-		meta = (DisplayThumbnail = true, DisplayName = "Texture", AllowPrivateAccess = "true", HighPriority, NotKeyframeable, NoCreate))
+		meta = (DisplayThumbnail = true, DisplayName = "Texture", AllowPrivateAccess = "true", HighPriority, NoCreate))
 	TObjectPtr<UTexture> Value;
 
 #if WITH_EDITORONLY_DATA
diff --git a/Engine/Plugins/Experimental/DynamicMaterial/Source/DynamicMaterialEditor/Private/Components/MaterialStageInputs/DMMSIValue.cpp b/Engine/Plugins/Experimental/DynamicMaterial/Source/DynamicMaterialEditor/Private/Components/MaterialStageInputs/DMMSIValue.cpp
--- a/Engine/Plugins/Experimental/DynamicMaterial/Source/DynamicMaterialEditor/Private/Components/MaterialStageInputs/DMMSIValue.cpp
+++ b/Engine/Plugins/Experimental/DynamicMaterial/Source/DynamicMaterialEditor/Private/Components/MaterialStageInputs/DMMSIValue.cpp
@@ -615,8 +615,6 @@ void UDMMaterialStageInputValue::InitInputValue()
 
 			Value->SetParentComponent(this);
 		}
-
-		Value->GetOnUpdate().AddUObject(this, &UDMMaterialStageInputValue::OnValueUpdated);
 	}
 }
 
@@ -624,7 +622,15 @@ void UDMMaterialStageInputValue::DeinitInputValue()
 {
 	if (Value)
 	{
-		Value->GetOnUpdate().RemoveAll(this);
+		if (Value->IsLocal())
+		{
+			if (GUndo)
+			{
+				Value->Modify();
+			}
+
+			Value->SetParentComponent(nullptr);
+		}
 	}
 }
 
