Material Designer
- When exporting material designer instances (now called "Save As"), the current folder in the content browser is selected by default.

diff --git a/Engine/Plugins/Experimental/DynamicMaterial/Source/DynamicMaterialEditor/Private/Slate/SDMToolBar.cpp b/Engine/Plugins/Experimental/DynamicMaterial/Source/DynamicMaterialEditor/Private/Slate/SDMToolBar.cpp
--- a/Engine/Plugins/Experimental/DynamicMaterial/Source/DynamicMaterialEditor/Private/Slate/SDMToolBar.cpp
+++ b/Engine/Plugins/Experimental/DynamicMaterial/Source/DynamicMaterialEditor/Private/Slate/SDMToolBar.cpp
@@ -497,9 +497,13 @@ FReply SDMToolBar::OnExportMaterialInstanceButtonClicked()
 		return FReply::Handled();
 	}
 
+	IContentBrowserSingleton& ContentBrowser = FModuleManager::LoadModuleChecked<FContentBrowserModule>("ContentBrowser").Get();
+	const FContentBrowserItemPath CurrentPath = ContentBrowser.GetCurrentPath();
+	const FString PathStr = CurrentPath.HasInternalPath() ? CurrentPath.GetInternalPathString() : "/Game";
+
 	FSaveAssetDialogConfig SaveAssetDialogConfig;
 	SaveAssetDialogConfig.DialogTitleOverride = LOCTEXT("SaveAssetDialogTitle", "Save Asset As");
-	SaveAssetDialogConfig.DefaultPath = "/Game";
+	SaveAssetDialogConfig.DefaultPath = PathStr;
 	SaveAssetDialogConfig.DefaultAssetName = MaterialInstance->GetName();
 	SaveAssetDialogConfig.ExistingAssetPolicy = ESaveAssetDialogExistingAssetPolicy::AllowButWarn;
 
