Material Designer
- Refactored the update type enum to make it flag based.

diff --git a/Engine/Plugins/Experimental/DynamicMaterial/Source/DynamicMaterial/Private/Components/DMMaterialValue.cpp b/Engine/Plugins/Experimental/DynamicMaterial/Source/DynamicMaterial/Private/Components/DMMaterialValue.cpp
--- a/Engine/Plugins/Experimental/DynamicMaterial/Source/DynamicMaterial/Private/Components/DMMaterialValue.cpp
+++ b/Engine/Plugins/Experimental/DynamicMaterial/Source/DynamicMaterial/Private/Components/DMMaterialValue.cpp
@@ -387,7 +387,7 @@ void UDMMaterialValue::PostEditUndo()
  
 	MarkComponentDirty();
  
-	OnValueChanged(EDMUpdateType::Structure, /* bInAllowPropagate */ true); // Just in case - Undos are not meant to be quick and easy.
+	OnValueChanged(EDMUpdateType::Structure | EDMUpdateType::AllowParentUpdate); // Just in case - Undos are not meant to be quick and easy.
 }
  
 void UDMMaterialValue::PostEditChangeProperty(FPropertyChangedEvent& PropertyChangedEvent)
@@ -410,12 +410,10 @@ void UDMMaterialValue::PostEditChangeProperty(FPropertyChangedEvent& PropertyCha
 	{
 		if (EditableProperty == MemberPropertyName)
 		{
-			OnValueChanged(
-				EditableProperty == ValueName
-					? EDMUpdateType::Value
-					: EDMUpdateType::Structure, 
-				/* bInAllowPropagate */ true
-			);
+			const EDMUpdateType UpdateType = (EditableProperty == ValueName ? EDMUpdateType::Value : EDMUpdateType::Structure)
+				| EDMUpdateType::AllowParentUpdate;
+
+			OnValueChanged(UpdateType);
 			return;
 		}
 	}
@@ -435,7 +433,7 @@ UDMMaterialValue::UDMMaterialValue(EDMValueType InType)
 #endif
 }
  
-void UDMMaterialValue::OnValueChanged(EDMUpdateType InUpdateType, bool bInUpdateParent)
+void UDMMaterialValue::OnValueChanged(EDMUpdateType InUpdateType)
 {
 	if (!IsComponentValid())
 	{
@@ -447,7 +445,7 @@ void UDMMaterialValue::OnValueChanged(EDMUpdateType InUpdateType, bool bInUpdate
 		Update(InUpdateType);
 
 #if WITH_EDITOR
-		if (bInUpdateParent && ParentComponent)
+		if (EnumHasAnyFlags(InUpdateType, EDMUpdateType::AllowParentUpdate) && ParentComponent)
 		{
 			ParentComponent->Update(InUpdateType);
 		}
diff --git a/Engine/Plugins/Experimental/DynamicMaterial/Source/DynamicMaterial/Private/Components/DMTextureUV.cpp b/Engine/Plugins/Experimental/DynamicMaterial/Source/DynamicMaterial/Private/Components/DMTextureUV.cpp
--- a/Engine/Plugins/Experimental/DynamicMaterial/Source/DynamicMaterial/Private/Components/DMTextureUV.cpp
+++ b/Engine/Plugins/Experimental/DynamicMaterial/Source/DynamicMaterial/Private/Components/DMTextureUV.cpp
@@ -87,7 +87,7 @@ void UDMTextureUV::SetUVSource(EDMUVSource InUVSource)
 
 	UVSource = InUVSource;
 
-	OnTextureUVChanged(EDMUpdateType::Structure, /* bInAllowPropagate */ true);
+	OnTextureUVChanged(EDMUpdateType::Structure | EDMUpdateType::AllowParentUpdate);
 }
 #endif
 
@@ -106,7 +106,7 @@ void UDMTextureUV::SetOffset(const FVector2D& InOffset)
 
 	Offset = InOffset;
 
-	OnTextureUVChanged(EDMUpdateType::Value, /* bInAllowPropagate */ false);
+	OnTextureUVChanged(EDMUpdateType::Value);
 }
 
 void UDMTextureUV::SetPivot(const FVector2D& InPivot)
@@ -124,7 +124,7 @@ void UDMTextureUV::SetPivot(const FVector2D& InPivot)
 
 	Pivot = InPivot;
 
-	OnTextureUVChanged(EDMUpdateType::Value, /* bInAllowPropagate */ false);
+	OnTextureUVChanged(EDMUpdateType::Value);
 }
 
 void UDMTextureUV::SetRotation(float InRotation)
@@ -141,7 +141,7 @@ void UDMTextureUV::SetRotation(float InRotation)
 
 	Rotation = InRotation;
 
-	OnTextureUVChanged(EDMUpdateType::Value, /* bInAllowPropagate */ false);
+	OnTextureUVChanged(EDMUpdateType::Value);
 }
 
 void UDMTextureUV::SetTiling(const FVector2D& InTiling)
@@ -159,7 +159,7 @@ void UDMTextureUV::SetTiling(const FVector2D& InTiling)
 
 	Tiling = InTiling;
 
-	OnTextureUVChanged(EDMUpdateType::Value, /* bInAllowPropagate */ false);
+	OnTextureUVChanged(EDMUpdateType::Value);
 }
 
 #if WITH_EDITOR
@@ -177,7 +177,7 @@ void UDMTextureUV::SetMirrorOnX(bool bInMirrorOnX)
 
 	bMirrorOnX = bInMirrorOnX;
 
-	OnTextureUVChanged(EDMUpdateType::Value, /* bInAllowPropagate */ false);
+	OnTextureUVChanged(EDMUpdateType::Structure | EDMUpdateType::AllowParentUpdate);
 }
 
 void UDMTextureUV::SetMirrorOnY(bool bInMirrorOnY)
@@ -194,7 +194,7 @@ void UDMTextureUV::SetMirrorOnY(bool bInMirrorOnY)
 
 	bMirrorOnY = bInMirrorOnY;
 
-	OnTextureUVChanged(EDMUpdateType::Structure, /* bInAllowPropagate */ true);
+	OnTextureUVChanged(EDMUpdateType::Structure | EDMUpdateType::AllowParentUpdate);
 }
 #endif
 
@@ -533,7 +533,7 @@ void UDMTextureUV::RemoveParameterNames()
 	}
 }
 
-void UDMTextureUV::OnTextureUVChanged(EDMUpdateType InUpdateType, bool bInUpdateParent)
+void UDMTextureUV::OnTextureUVChanged(EDMUpdateType InUpdateType)
 {
 	if (!IsComponentValid())
 	{
@@ -545,7 +545,7 @@ void UDMTextureUV::OnTextureUVChanged(EDMUpdateType InUpdateType, bool bInUpdate
 		Update(InUpdateType);
 
 #if WITH_EDITOR
-		if (bInUpdateParent && ParentComponent)
+		if (EnumHasAnyFlags(InUpdateType, EDMUpdateType::AllowParentUpdate) && ParentComponent)
 		{
 			ParentComponent->Update(InUpdateType);
 		}
@@ -601,13 +601,13 @@ void UDMTextureUV::PostEditChangeProperty(struct FPropertyChangedEvent& Property
 		|| PropertyChangedEvent.MemberProperty->GetFName() == NAME_Rotation
 		|| PropertyChangedEvent.MemberProperty->GetFName() == NAME_Tiling)
 	{
-		OnTextureUVChanged(EDMUpdateType::Value, /* bInAllowPropagate */ true);
+		OnTextureUVChanged(EDMUpdateType::Value | EDMUpdateType::AllowParentUpdate);
 	}
 	else if (PropertyChangedEvent.MemberProperty->GetFName() == NAME_UVSource
 		|| PropertyChangedEvent.MemberProperty->GetFName() == NAME_bMirrorOnX
 		|| PropertyChangedEvent.MemberProperty->GetFName() == NAME_bMirrorOnY)
 	{
-		OnTextureUVChanged(EDMUpdateType::Structure, /* bInAllowPropagate */ true);
+		OnTextureUVChanged(EDMUpdateType::Structure | EDMUpdateType::AllowParentUpdate);
 	}
 }
 
@@ -628,11 +628,11 @@ void UDMTextureUV::PostEditUndo()
 		|| bMirrorOnX != bMirrorOnX_PreUndo
 		|| bMirrorOnY != bMirrorOnY_PreUndo)
 	{
-		OnTextureUVChanged(EDMUpdateType::Structure, /* bInAllowPropagate */ true);
+		OnTextureUVChanged(EDMUpdateType::Structure | EDMUpdateType::AllowParentUpdate);
 	}
 	else
 	{
-		OnTextureUVChanged(EDMUpdateType::Value, /* bInAllowPropagate */ true);
+		OnTextureUVChanged(EDMUpdateType::Value | EDMUpdateType::AllowParentUpdate);
 	}
 }
 
@@ -668,11 +668,11 @@ void UDMTextureUV::PostLoad()
 
 	if (bNeedsPostLoadStructureUpdate)
 	{
-		OnTextureUVChanged(EDMUpdateType::Structure, /* bInAllowPropagate */ false);
+		OnTextureUVChanged(EDMUpdateType::Structure);
 	}
 	else if (bNeedsPostLoadValueUpdate)
 	{
-		OnTextureUVChanged(EDMUpdateType::Value, /* bInAllowPropagate */ false);
+		OnTextureUVChanged(EDMUpdateType::Value);
 	}
 
 	bNeedsPostLoadStructureUpdate = false;
diff --git a/Engine/Plugins/Experimental/DynamicMaterial/Source/DynamicMaterial/Private/Components/MaterialValues/DMMaterialValueBool.cpp b/Engine/Plugins/Experimental/DynamicMaterial/Source/DynamicMaterial/Private/Components/MaterialValues/DMMaterialValueBool.cpp
--- a/Engine/Plugins/Experimental/DynamicMaterial/Source/DynamicMaterial/Private/Components/MaterialValues/DMMaterialValueBool.cpp
+++ b/Engine/Plugins/Experimental/DynamicMaterial/Source/DynamicMaterial/Private/Components/MaterialValues/DMMaterialValueBool.cpp
@@ -112,7 +112,7 @@ void UDMMaterialValueBool::SetValue(bool InValue)
 
 	Value = InValue;
  
-	OnValueChanged(EDMUpdateType::Value, /* bInAllowPropagate */ false);
+	OnValueChanged(EDMUpdateType::Value);
 }
  
 #undef LOCTEXT_NAMESPACE
diff --git a/Engine/Plugins/Experimental/DynamicMaterial/Source/DynamicMaterial/Private/Components/MaterialValues/DMMaterialValueColorAtlas.cpp b/Engine/Plugins/Experimental/DynamicMaterial/Source/DynamicMaterial/Private/Components/MaterialValues/DMMaterialValueColorAtlas.cpp
--- a/Engine/Plugins/Experimental/DynamicMaterial/Source/DynamicMaterial/Private/Components/MaterialValues/DMMaterialValueColorAtlas.cpp
+++ b/Engine/Plugins/Experimental/DynamicMaterial/Source/DynamicMaterial/Private/Components/MaterialValues/DMMaterialValueColorAtlas.cpp
@@ -47,7 +47,7 @@ void UDMMaterialValueColorAtlas::SetValue(float InValue)
 
 	Value = InValue;
 
-	OnValueChanged(EDMUpdateType::Value, /* bInAllowPropagate */ false);
+	OnValueChanged(EDMUpdateType::Value);
 }
 
 void UDMMaterialValueColorAtlas::SetMIDParameter(UMaterialInstanceDynamic* InMID) const
@@ -77,7 +77,7 @@ void UDMMaterialValueColorAtlas::SetAtlas(UCurveLinearColorAtlas* InAtlas)
 
 	Atlas = InAtlas;
 
-	OnValueChanged(EDMUpdateType::Structure, /* bInAllowPropagate */ true);
+	OnValueChanged(EDMUpdateType::Structure | EDMUpdateType::AllowParentUpdate);
 }
 
 void UDMMaterialValueColorAtlas::SetCurve(UCurveLinearColor* InCurve)
@@ -94,7 +94,7 @@ void UDMMaterialValueColorAtlas::SetCurve(UCurveLinearColor* InCurve)
 
 	Curve = InCurve;
 
-	OnValueChanged(EDMUpdateType::Structure, /* bInAllowPropagate */ true);
+	OnValueChanged(EDMUpdateType::Structure | EDMUpdateType::AllowParentUpdate);
 }
 
 void UDMMaterialValueColorAtlas::GenerateExpression(const TSharedRef<IDMMaterialBuildStateInterface>& InBuildState) const
@@ -198,7 +198,7 @@ bool UDMMaterialValueColorAtlas::JsonDeserialize(const TSharedPtr<FJsonValue>& I
 
 	if (bSuccess)
 	{
-		OnValueChanged(UpdateType, /* bInAllowPropagate */ true);
+		OnValueChanged(UpdateType | EDMUpdateType::AllowParentUpdate);
 	}
 
 	return bSuccess;
diff --git a/Engine/Plugins/Experimental/DynamicMaterial/Source/DynamicMaterial/Private/Components/MaterialValues/DMMaterialValueFloat1.cpp b/Engine/Plugins/Experimental/DynamicMaterial/Source/DynamicMaterial/Private/Components/MaterialValues/DMMaterialValueFloat1.cpp
--- a/Engine/Plugins/Experimental/DynamicMaterial/Source/DynamicMaterial/Private/Components/MaterialValues/DMMaterialValueFloat1.cpp
+++ b/Engine/Plugins/Experimental/DynamicMaterial/Source/DynamicMaterial/Private/Components/MaterialValues/DMMaterialValueFloat1.cpp
@@ -99,7 +99,7 @@ void UDMMaterialValueFloat1::SetValue(float InValue)
  
 	Value = InValue;
  
-	OnValueChanged(EDMUpdateType::Value, /* bInAllowPropagate */ false);
+	OnValueChanged(EDMUpdateType::Value);
 }
  
 void UDMMaterialValueFloat1::SetMIDParameter(UMaterialInstanceDynamic* InMID) const
diff --git a/Engine/Plugins/Experimental/DynamicMaterial/Source/DynamicMaterial/Private/Components/MaterialValues/DMMaterialValueFloat2.cpp b/Engine/Plugins/Experimental/DynamicMaterial/Source/DynamicMaterial/Private/Components/MaterialValues/DMMaterialValueFloat2.cpp
--- a/Engine/Plugins/Experimental/DynamicMaterial/Source/DynamicMaterial/Private/Components/MaterialValues/DMMaterialValueFloat2.cpp
+++ b/Engine/Plugins/Experimental/DynamicMaterial/Source/DynamicMaterial/Private/Components/MaterialValues/DMMaterialValueFloat2.cpp
@@ -112,7 +112,7 @@ void UDMMaterialValueFloat2::SetValue(const FVector2D& InValue)
  
 	Value = ValueClamped;
  
-	OnValueChanged(EDMUpdateType::Value, /* bInAllowPropagate */ false);
+	OnValueChanged(EDMUpdateType::Value);
 }
  
 #if WITH_EDITOR
diff --git a/Engine/Plugins/Experimental/DynamicMaterial/Source/DynamicMaterial/Private/Components/MaterialValues/DMMaterialValueFloat3RGB.cpp b/Engine/Plugins/Experimental/DynamicMaterial/Source/DynamicMaterial/Private/Components/MaterialValues/DMMaterialValueFloat3RGB.cpp
--- a/Engine/Plugins/Experimental/DynamicMaterial/Source/DynamicMaterial/Private/Components/MaterialValues/DMMaterialValueFloat3RGB.cpp
+++ b/Engine/Plugins/Experimental/DynamicMaterial/Source/DynamicMaterial/Private/Components/MaterialValues/DMMaterialValueFloat3RGB.cpp
@@ -112,7 +112,7 @@ void UDMMaterialValueFloat3RGB::SetValue(const FLinearColor& InValue)
 
 	Value = ValueClamped;
 
-	OnValueChanged(EDMUpdateType::Value, /* bInAllowPropagate */ false);
+	OnValueChanged(EDMUpdateType::Value);
 }
  
 #if WITH_EDITOR
diff --git a/Engine/Plugins/Experimental/DynamicMaterial/Source/DynamicMaterial/Private/Components/MaterialValues/DMMaterialValueFloat3RPY.cpp b/Engine/Plugins/Experimental/DynamicMaterial/Source/DynamicMaterial/Private/Components/MaterialValues/DMMaterialValueFloat3RPY.cpp
--- a/Engine/Plugins/Experimental/DynamicMaterial/Source/DynamicMaterial/Private/Components/MaterialValues/DMMaterialValueFloat3RPY.cpp
+++ b/Engine/Plugins/Experimental/DynamicMaterial/Source/DynamicMaterial/Private/Components/MaterialValues/DMMaterialValueFloat3RPY.cpp
@@ -114,7 +114,7 @@ void UDMMaterialValueFloat3RPY::SetValue(const FRotator& InValue)
  
 	Value = ValueClamped;
  
-	OnValueChanged(EDMUpdateType::Value, /* bInAllowPropagate */ false);
+	OnValueChanged(EDMUpdateType::Value);
 }
  
 #if WITH_EDITOR
diff --git a/Engine/Plugins/Experimental/DynamicMaterial/Source/DynamicMaterial/Private/Components/MaterialValues/DMMaterialValueFloat3XYZ.cpp b/Engine/Plugins/Experimental/DynamicMaterial/Source/DynamicMaterial/Private/Components/MaterialValues/DMMaterialValueFloat3XYZ.cpp
--- a/Engine/Plugins/Experimental/DynamicMaterial/Source/DynamicMaterial/Private/Components/MaterialValues/DMMaterialValueFloat3XYZ.cpp
+++ b/Engine/Plugins/Experimental/DynamicMaterial/Source/DynamicMaterial/Private/Components/MaterialValues/DMMaterialValueFloat3XYZ.cpp
@@ -111,7 +111,7 @@ void UDMMaterialValueFloat3XYZ::SetValue(const FVector& InValue)
  
 	Value = ValueClamped;
  
-	OnValueChanged(EDMUpdateType::Value, /* bInAllowPropagate */ false);
+	OnValueChanged(EDMUpdateType::Value);
 }
  
 #if WITH_EDITOR
diff --git a/Engine/Plugins/Experimental/DynamicMaterial/Source/DynamicMaterial/Private/Components/MaterialValues/DMMaterialValueFloat4.cpp b/Engine/Plugins/Experimental/DynamicMaterial/Source/DynamicMaterial/Private/Components/MaterialValues/DMMaterialValueFloat4.cpp
--- a/Engine/Plugins/Experimental/DynamicMaterial/Source/DynamicMaterial/Private/Components/MaterialValues/DMMaterialValueFloat4.cpp
+++ b/Engine/Plugins/Experimental/DynamicMaterial/Source/DynamicMaterial/Private/Components/MaterialValues/DMMaterialValueFloat4.cpp
@@ -114,7 +114,7 @@ void UDMMaterialValueFloat4::SetValue(const FLinearColor& InValue)
  
 	Value = ValueClamped;
  
-	OnValueChanged(EDMUpdateType::Value, /* bInAllowPropagate */ false);
+	OnValueChanged(EDMUpdateType::Value);
 }
  
 #if WITH_EDITOR
diff --git a/Engine/Plugins/Experimental/DynamicMaterial/Source/DynamicMaterial/Private/Components/MaterialValues/DMMaterialValueRenderTarget.cpp b/Engine/Plugins/Experimental/DynamicMaterial/Source/DynamicMaterial/Private/Components/MaterialValues/DMMaterialValueRenderTarget.cpp
--- a/Engine/Plugins/Experimental/DynamicMaterial/Source/DynamicMaterial/Private/Components/MaterialValues/DMMaterialValueRenderTarget.cpp
+++ b/Engine/Plugins/Experimental/DynamicMaterial/Source/DynamicMaterial/Private/Components/MaterialValues/DMMaterialValueRenderTarget.cpp
@@ -224,7 +224,7 @@ bool UDMMaterialValueRenderTarget::JsonDeserialize(const TSharedPtr<FJsonValue>&
 
 	if (bSuccess)
 	{
-		OnValueChanged(EDMUpdateType::Structure, /* bInAllowPropagate */ true);
+		OnValueChanged(EDMUpdateType::Structure | EDMUpdateType::AllowParentUpdate);
 	}
 
 	return bSuccess;
diff --git a/Engine/Plugins/Experimental/DynamicMaterial/Source/DynamicMaterial/Private/Components/MaterialValues/DMMaterialValueTexture.cpp b/Engine/Plugins/Experimental/DynamicMaterial/Source/DynamicMaterial/Private/Components/MaterialValues/DMMaterialValueTexture.cpp
--- a/Engine/Plugins/Experimental/DynamicMaterial/Source/DynamicMaterial/Private/Components/MaterialValues/DMMaterialValueTexture.cpp
+++ b/Engine/Plugins/Experimental/DynamicMaterial/Source/DynamicMaterial/Private/Components/MaterialValues/DMMaterialValueTexture.cpp
@@ -150,11 +150,11 @@ void UDMMaterialValueTexture::PostEditChangeProperty(FPropertyChangedEvent& InPr
 					? EDMUpdateType::Structure
 					: EDMUpdateType::Value;
 
-				OnValueChanged(UpdateType, /* bInAllowPropagate */ true);
+				OnValueChanged(EDMUpdateType::Value | EDMUpdateType::AllowParentUpdate);
 				return;
 			}
 
-			OnValueChanged(EDMUpdateType::Structure, /* bInAllowPropagate */ true);
+			OnValueChanged(EDMUpdateType::Structure | EDMUpdateType::AllowParentUpdate);
 			return;
 		}
 	}
@@ -225,18 +225,9 @@ void UDMMaterialValueTexture::SetValue(UTexture* InValue)
 		return;
 	}
 
-	/*
-	const TextureCompressionSettings CurrentCompression = Value ? Value->CompressionSettings.GetValue() : TextureCompressionSettings::TC_MAX;
-	const TextureCompressionSettings NewCompression = InValue ? InValue->CompressionSettings.GetValue() : TextureCompressionSettings::TC_MAX;
-
-	const EDMUpdateType UpdateType = CurrentCompression == NewCompression
-		? EDMUpdateType::Structure
-		: EDMUpdateType::Value;
-	*/
-
 	Value = InValue;
 
-	OnValueChanged(EDMUpdateType::Value, /* bInAllowPropagate */ false);
+	OnValueChanged(EDMUpdateType::Value);
 }
 
 void UDMMaterialValueTexture::SetMIDParameter(UMaterialInstanceDynamic* InMID) const
diff --git a/Engine/Plugins/Experimental/DynamicMaterial/Source/DynamicMaterial/Private/Model/DynamicMaterialModel.cpp b/Engine/Plugins/Experimental/DynamicMaterial/Source/DynamicMaterial/Private/Model/DynamicMaterialModel.cpp
--- a/Engine/Plugins/Experimental/DynamicMaterial/Source/DynamicMaterial/Private/Model/DynamicMaterialModel.cpp
+++ b/Engine/Plugins/Experimental/DynamicMaterial/Source/DynamicMaterial/Private/Model/DynamicMaterialModel.cpp
@@ -427,7 +427,7 @@ void UDynamicMaterialModel::OnValueUpdated(UDMMaterialValue* InValue, EDMUpdateT
 		return;
 	}
 
-	if (InUpdateType == EDMUpdateType::Value && IsValid(DynamicMaterialInstance))
+	if (!EnumHasAnyFlags(InUpdateType, EDMUpdateType::Structure) && IsValid(DynamicMaterialInstance))
 	{
 		InValue->SetMIDParameter(DynamicMaterialInstance);
 	}
diff --git a/Engine/Plugins/Experimental/DynamicMaterial/Source/DynamicMaterial/Public/Components/DMMaterialValue.h b/Engine/Plugins/Experimental/DynamicMaterial/Source/DynamicMaterial/Public/Components/DMMaterialValue.h
--- a/Engine/Plugins/Experimental/DynamicMaterial/Source/DynamicMaterial/Public/Components/DMMaterialValue.h
+++ b/Engine/Plugins/Experimental/DynamicMaterial/Source/DynamicMaterial/Public/Components/DMMaterialValue.h
@@ -128,7 +128,7 @@ public:
 	//~ End IDMJsonSerializable
  
 	//~ Begin UDMMaterialComponent
-	virtual void Update(EDMUpdateType InUpdateType = EDMUpdateType::Value) override;
+	virtual void Update(EDMUpdateType InUpdateType) override;
 #if WITH_EDITOR
 	virtual void DoClean() override;
 	virtual void PostEditorDuplicate(UDynamicMaterialModel* InMaterialModel, UDMMaterialComponent* InParent) override;
@@ -174,7 +174,7 @@ protected:
  
 	UDMMaterialValue(EDMValueType InType);
  
-	virtual void OnValueChanged(EDMUpdateType InUpdateType, bool bInUpdateParent);
+	virtual void OnValueChanged(EDMUpdateType InUpdateType);
  
 	//~ Begin UDMMaterialComponent
 	virtual UDMMaterialComponent* GetSubComponentByPath(FDMComponentPath& InPath, const FDMComponentPathSegment& InPathSegment) const override;
diff --git a/Engine/Plugins/Experimental/DynamicMaterial/Source/DynamicMaterial/Public/Components/DMRenderTargetRenderer.h b/Engine/Plugins/Experimental/DynamicMaterial/Source/DynamicMaterial/Public/Components/DMRenderTargetRenderer.h
--- a/Engine/Plugins/Experimental/DynamicMaterial/Source/DynamicMaterial/Public/Components/DMRenderTargetRenderer.h
+++ b/Engine/Plugins/Experimental/DynamicMaterial/Source/DynamicMaterial/Public/Components/DMRenderTargetRenderer.h
@@ -49,7 +49,7 @@ public:
 	//~ End IDMJsonSerializable
 
 	//~ Begin UDMMaterialComponent
-	virtual void Update(EDMUpdateType InUpdateType = EDMUpdateType::Value) override;
+	virtual void Update(EDMUpdateType InUpdateType) override;
 	//~ End UDMMaterialComponent
 
 	//~ Begin UObject
diff --git a/Engine/Plugins/Experimental/DynamicMaterial/Source/DynamicMaterial/Public/Components/DMTextureUV.h b/Engine/Plugins/Experimental/DynamicMaterial/Source/DynamicMaterial/Public/Components/DMTextureUV.h
--- a/Engine/Plugins/Experimental/DynamicMaterial/Source/DynamicMaterial/Public/Components/DMTextureUV.h
+++ b/Engine/Plugins/Experimental/DynamicMaterial/Source/DynamicMaterial/Public/Components/DMTextureUV.h
@@ -219,7 +219,7 @@ protected:
 	void RemoveParameterNames();
 #endif
 
-	void OnTextureUVChanged(EDMUpdateType InUpdateType, bool bInUpdateParent);
+	void OnTextureUVChanged(EDMUpdateType InUpdateType);
 
 	//~ Begin UDMMaterialComponent
 	virtual void Update(EDMUpdateType InUpdateType) override;
diff --git a/Engine/Plugins/Experimental/DynamicMaterial/Source/DynamicMaterial/Public/Components/MaterialValues/DMMaterialValueRenderTarget.h b/Engine/Plugins/Experimental/DynamicMaterial/Source/DynamicMaterial/Public/Components/MaterialValues/DMMaterialValueRenderTarget.h
--- a/Engine/Plugins/Experimental/DynamicMaterial/Source/DynamicMaterial/Public/Components/MaterialValues/DMMaterialValueRenderTarget.h
+++ b/Engine/Plugins/Experimental/DynamicMaterial/Source/DynamicMaterial/Public/Components/MaterialValues/DMMaterialValueRenderTarget.h
@@ -72,7 +72,7 @@ public:
 #endif
 
 	//~ Begin UDMMaterialComponent
-	virtual void Update(EDMUpdateType InUpdateType = EDMUpdateType::Value) override;
+	virtual void Update(EDMUpdateType InUpdateType) override;
 #if WITH_EDITOR
 	virtual void PostEditorDuplicate(UDynamicMaterialModel* InMaterialModel, UDMMaterialComponent* InParent) override;
 #endif
diff --git a/Engine/Plugins/Experimental/DynamicMaterial/Source/DynamicMaterial/Public/DMDefs.h b/Engine/Plugins/Experimental/DynamicMaterial/Source/DynamicMaterial/Public/DMDefs.h
--- a/Engine/Plugins/Experimental/DynamicMaterial/Source/DynamicMaterial/Public/DMDefs.h
+++ b/Engine/Plugins/Experimental/DynamicMaterial/Source/DynamicMaterial/Public/DMDefs.h
@@ -63,9 +63,11 @@ enum class EDMValueType : uint8
 UENUM(BlueprintType)
 enum class EDMUpdateType : uint8
 {
-	Value,
-	Structure
+	Value = 0,
+	Structure = 1 << 0,
+	AllowParentUpdate = 1 << 1
 };
+ENUM_CLASS_FLAGS(EDMUpdateType)
 
 UENUM(BlueprintType)
 enum class EDMMaterialShadingModel : uint8
diff --git a/Engine/Plugins/Experimental/DynamicMaterial/Source/DynamicMaterialEditor/Private/Components/DMMaterialProperty.cpp b/Engine/Plugins/Experimental/DynamicMaterial/Source/DynamicMaterialEditor/Private/Components/DMMaterialProperty.cpp
--- a/Engine/Plugins/Experimental/DynamicMaterial/Source/DynamicMaterialEditor/Private/Components/DMMaterialProperty.cpp
+++ b/Engine/Plugins/Experimental/DynamicMaterial/Source/DynamicMaterialEditor/Private/Components/DMMaterialProperty.cpp
@@ -611,7 +611,7 @@ void UDMMaterialProperty::Update(EDMUpdateType InUpdateType)
 
 	Super::Update(InUpdateType);
 
-	if (InUpdateType == EDMUpdateType::Structure)
+	if (EnumHasAnyFlags(InUpdateType, EDMUpdateType::Structure))
 	{
 		UDynamicMaterialModelEditorOnlyData* ModelEditorOnlyData = GetMaterialModelEditorOnlyData();
 		check(ModelEditorOnlyData);
diff --git a/Engine/Plugins/Experimental/DynamicMaterial/Source/DynamicMaterialEditor/Private/Components/DMMaterialSlot.cpp b/Engine/Plugins/Experimental/DynamicMaterial/Source/DynamicMaterialEditor/Private/Components/DMMaterialSlot.cpp
--- a/Engine/Plugins/Experimental/DynamicMaterial/Source/DynamicMaterialEditor/Private/Components/DMMaterialSlot.cpp
+++ b/Engine/Plugins/Experimental/DynamicMaterial/Source/DynamicMaterialEditor/Private/Components/DMMaterialSlot.cpp
@@ -192,7 +192,7 @@ void UDMMaterialSlot::Update(EDMUpdateType InUpdateType)
 		return;
 	}
 
-	if (InUpdateType == EDMUpdateType::Structure)
+	if (EnumHasAnyFlags(InUpdateType, EDMUpdateType::Structure))
 	{
 		MarkComponentDirty();
 		UpdateMaterialProperties();
@@ -200,7 +200,7 @@ void UDMMaterialSlot::Update(EDMUpdateType InUpdateType)
 
 	Super::Update(InUpdateType);
 
-	if (InUpdateType == EDMUpdateType::Structure)
+	if (EnumHasAnyFlags(InUpdateType, EDMUpdateType::Structure))
 	{
 		UDynamicMaterialModelEditorOnlyData* ModelEditorOnlyData = GetMaterialModelEditorOnlyData();
 		check(ModelEditorOnlyData);
diff --git a/Engine/Plugins/Experimental/DynamicMaterial/Source/DynamicMaterialEditor/Private/Components/DMMaterialStage.cpp b/Engine/Plugins/Experimental/DynamicMaterial/Source/DynamicMaterialEditor/Private/Components/DMMaterialStage.cpp
--- a/Engine/Plugins/Experimental/DynamicMaterial/Source/DynamicMaterialEditor/Private/Components/DMMaterialStage.cpp
+++ b/Engine/Plugins/Experimental/DynamicMaterial/Source/DynamicMaterialEditor/Private/Components/DMMaterialStage.cpp
@@ -472,7 +472,7 @@ void UDMMaterialStage::Update(EDMUpdateType InUpdateType)
 		return;
 	}
 
-	if (InUpdateType == EDMUpdateType::Structure)
+	if (EnumHasAnyFlags(InUpdateType, EDMUpdateType::Structure))
 	{
 		MarkComponentDirty();
 		VerifyAllInputMaps();
diff --git a/Engine/Plugins/Experimental/DynamicMaterial/Source/DynamicMaterialEditor/Private/Components/DMMaterialStageInput.cpp b/Engine/Plugins/Experimental/DynamicMaterial/Source/DynamicMaterialEditor/Private/Components/DMMaterialStageInput.cpp
--- a/Engine/Plugins/Experimental/DynamicMaterial/Source/DynamicMaterialEditor/Private/Components/DMMaterialStageInput.cpp
+++ b/Engine/Plugins/Experimental/DynamicMaterial/Source/DynamicMaterialEditor/Private/Components/DMMaterialStageInput.cpp
@@ -24,7 +24,7 @@ void UDMMaterialStageInput::Update(EDMUpdateType InUpdateType)
 		return;
 	}
 
-	if (InUpdateType == EDMUpdateType::Structure)
+	if (EnumHasAnyFlags(InUpdateType, EDMUpdateType::Structure))
 	{
 		MarkComponentDirty();
 	}
diff --git a/Engine/Plugins/Experimental/DynamicMaterial/Source/DynamicMaterialEditor/Private/Components/DMMaterialStageSource.cpp b/Engine/Plugins/Experimental/DynamicMaterial/Source/DynamicMaterialEditor/Private/Components/DMMaterialStageSource.cpp
--- a/Engine/Plugins/Experimental/DynamicMaterial/Source/DynamicMaterialEditor/Private/Components/DMMaterialStageSource.cpp
+++ b/Engine/Plugins/Experimental/DynamicMaterial/Source/DynamicMaterialEditor/Private/Components/DMMaterialStageSource.cpp
@@ -43,7 +43,7 @@ void UDMMaterialStageSource::Update(EDMUpdateType InUpdateType)
 		return;
 	}
 
-	if (InUpdateType == EDMUpdateType::Structure)
+	if (EnumHasAnyFlags(InUpdateType, EDMUpdateType::Structure))
 	{
 		MarkComponentDirty();
 	}
diff --git a/Engine/Plugins/Experimental/DynamicMaterial/Source/DynamicMaterialEditor/Private/Components/DMMaterialStageThroughputLayerBlend.cpp b/Engine/Plugins/Experimental/DynamicMaterial/Source/DynamicMaterialEditor/Private/Components/DMMaterialStageThroughputLayerBlend.cpp
--- a/Engine/Plugins/Experimental/DynamicMaterial/Source/DynamicMaterialEditor/Private/Components/DMMaterialStageThroughputLayerBlend.cpp
+++ b/Engine/Plugins/Experimental/DynamicMaterial/Source/DynamicMaterialEditor/Private/Components/DMMaterialStageThroughputLayerBlend.cpp
@@ -189,7 +189,7 @@ void UDMMaterialStageThroughputLayerBlend::Update(EDMUpdateType InUpdateType)
 		return;
 	}
  
-	if (InUpdateType == EDMUpdateType::Structure)
+	if (EnumHasAnyFlags(InUpdateType, EDMUpdateType::Structure))
 	{
 		UDMMaterialStage* Stage = GetStage();
 		check(Stage);
@@ -1003,7 +1003,7 @@ void UDMMaterialStageThroughputLayerBlend::OnStageUpdated(UDMMaterialComponent*
 		return;
 	}
 
-	if (InUpdateType == EDMUpdateType::Structure)
+	if (EnumHasAnyFlags(InUpdateType, EDMUpdateType::Structure))
 	{
 		const bool bIsAlphaOnlyBlendNow = bIsAlphaOnlyBlend;
 		UpdateAlphaOnlyMaskStatus();
diff --git a/Engine/Plugins/Experimental/DynamicMaterial/Source/DynamicMaterialEditor/Private/Components/MaterialStageExpressions/DMMSEMathBase.cpp b/Engine/Plugins/Experimental/DynamicMaterial/Source/DynamicMaterialEditor/Private/Components/MaterialStageExpressions/DMMSEMathBase.cpp
--- a/Engine/Plugins/Experimental/DynamicMaterial/Source/DynamicMaterialEditor/Private/Components/MaterialStageExpressions/DMMSEMathBase.cpp
+++ b/Engine/Plugins/Experimental/DynamicMaterial/Source/DynamicMaterialEditor/Private/Components/MaterialStageExpressions/DMMSEMathBase.cpp
@@ -213,7 +213,7 @@ void UDMMaterialStageExpressionMathBase::Update(EDMUpdateType InUpdateType)
 		return;
 	}
 
-	if (InUpdateType == EDMUpdateType::Structure)
+	if (EnumHasAnyFlags(InUpdateType, EDMUpdateType::Structure))
 	{
 		check(InputConnectors.IsEmpty() == false);
 		check(OutputConnectors.Num() == 1);
diff --git a/Engine/Plugins/Experimental/DynamicMaterial/Source/DynamicMaterialEditor/Private/Components/MaterialStageInputs/DMMSITextureUV.cpp b/Engine/Plugins/Experimental/DynamicMaterial/Source/DynamicMaterialEditor/Private/Components/MaterialStageInputs/DMMSITextureUV.cpp
--- a/Engine/Plugins/Experimental/DynamicMaterial/Source/DynamicMaterialEditor/Private/Components/MaterialStageInputs/DMMSITextureUV.cpp
+++ b/Engine/Plugins/Experimental/DynamicMaterial/Source/DynamicMaterialEditor/Private/Components/MaterialStageInputs/DMMSITextureUV.cpp
@@ -298,7 +298,7 @@ void UDMMaterialStageInputTextureUV::OnTextureUVUpdated(UDMMaterialComponent* In
 		return;
 	}
 
-	if (InComponent == TextureUV && InUpdateType == EDMUpdateType::Structure)
+	if (InComponent == TextureUV && EnumHasAnyFlags(InUpdateType, EDMUpdateType::Structure))
 	{
 		Update(InUpdateType);
 	}
diff --git a/Engine/Plugins/Experimental/DynamicMaterial/Source/DynamicMaterialEditor/Private/Model/DynamicMaterialModelEditorOnlyData.cpp b/Engine/Plugins/Experimental/DynamicMaterial/Source/DynamicMaterialEditor/Private/Model/DynamicMaterialModelEditorOnlyData.cpp
--- a/Engine/Plugins/Experimental/DynamicMaterial/Source/DynamicMaterialEditor/Private/Model/DynamicMaterialModelEditorOnlyData.cpp
+++ b/Engine/Plugins/Experimental/DynamicMaterial/Source/DynamicMaterialEditor/Private/Model/DynamicMaterialModelEditorOnlyData.cpp
@@ -1547,7 +1547,7 @@ void UDynamicMaterialModelEditorOnlyData::OnValueUpdated(UDMMaterialValue* InVal
 	// Exported materials need to be rebuilt to update the main material.
 	const bool bMaterialInDifferentPackage = MaterialModel->DynamicMaterial ? MaterialModel->DynamicMaterial->GetPackage() != GetPackage() : true;
 
-	if (InUpdateType == EDMUpdateType::Structure || bMaterialInDifferentPackage)
+	if (EnumHasAnyFlags(InUpdateType, EDMUpdateType::Structure) || bMaterialInDifferentPackage)
 	{
 		RequestMaterialBuild();
 	}
diff --git a/Engine/Plugins/Experimental/DynamicMaterial/Source/DynamicMaterialEditor/Private/Slate/SDMEditor.cpp b/Engine/Plugins/Experimental/DynamicMaterial/Source/DynamicMaterialEditor/Private/Slate/SDMEditor.cpp
--- a/Engine/Plugins/Experimental/DynamicMaterial/Source/DynamicMaterialEditor/Private/Slate/SDMEditor.cpp
+++ b/Engine/Plugins/Experimental/DynamicMaterial/Source/DynamicMaterialEditor/Private/Slate/SDMEditor.cpp
@@ -1697,7 +1697,7 @@ void SDMEditor::OnSplitterResized() const
 
 void SDMEditor::OnComponentUpdated(UDMMaterialComponent* InComponent, EDMUpdateType InUpdateType)
 {
-	if (InUpdateType == EDMUpdateType::Structure)
+	if (EnumHasAnyFlags(InUpdateType, EDMUpdateType::Structure))
 	{
 		InvalidateComponentEditWidget();
 	}
