Material Designer
- Changed how default layer settings are applied. All textures and solid colors will use the default settings as describe in the default slot value override settings.
- You can now specify both a default texture and default colour for slots. The "default type" selector will choose which is the default layer added to a new slot.

diff --git a/Engine/Plugins/Experimental/DynamicMaterial/Source/DynamicMaterialEditor/Private/Components/DMMaterialProperty.cpp b/Engine/Plugins/Experimental/DynamicMaterial/Source/DynamicMaterialEditor/Private/Components/DMMaterialProperty.cpp
--- a/Engine/Plugins/Experimental/DynamicMaterial/Source/DynamicMaterialEditor/Private/Components/DMMaterialProperty.cpp
+++ b/Engine/Plugins/Experimental/DynamicMaterial/Source/DynamicMaterialEditor/Private/Components/DMMaterialProperty.cpp
@@ -157,7 +157,7 @@ void UDMMaterialProperty::AddDefaultBaseStage(UDMMaterialLayerObject* InLayer)
 
 	const FDMDefaultMaterialPropertySlotValue& DefaultValue = UDynamicMaterialEditorSettings::Get()->GetDefaultSlotValue(MaterialProperty);
 
-	switch (DefaultValue.Type)
+	switch (DefaultValue.DefaultType)
 	{
 		case EDMDefaultMaterialPropertySlotValueType::Texture:
 		{
@@ -170,27 +170,6 @@ void UDMMaterialProperty::AddDefaultBaseStage(UDMMaterialLayerObject* InLayer)
 				FDMMaterialStageConnectorChannel::WHOLE_CHANNEL
 			);
 
-			if (UDMMaterialStageExpressionTextureSample* BaseInputTextureSample = Cast<UDMMaterialStageExpressionTextureSample>(BaseInputExpression->GetMaterialStageExpression()))
-			{
-				if (UDMMaterialStage* BaseTextureInputStage = BaseInputExpression->GetSubStage())
-				{
-					const TArray<UDMMaterialStageInput*> BaseTextureStageInputs = BaseTextureInputStage->GetInputs();
-
-					for (UDMMaterialStageInput* BaseTextureStageInput : BaseTextureStageInputs)
-					{
-						if (UDMMaterialStageInputValue* BaseTextureInputValue = Cast<UDMMaterialStageInputValue>(BaseTextureStageInput))
-						{
-							if (UDMMaterialValueTexture* BaseTextureValue = Cast<UDMMaterialValueTexture>(BaseTextureInputValue->GetValue()))
-							{
-								BaseTextureValue->SetDefaultValue(DefaultValue.Texture.LoadSynchronous());
-								BaseTextureValue->ApplyDefaultValue();
-								break;
-							}
-						}
-					}
-				}
-			}
-
 			break;
 		}
 
@@ -204,12 +183,6 @@ void UDMMaterialProperty::AddDefaultBaseStage(UDMMaterialLayerObject* InLayer)
 				0
 			);
 
-			if (UDMMaterialValueFloat3RGB* RGBValue = Cast<UDMMaterialValueFloat3RGB>(InputValue->GetValue()))
-			{
-				RGBValue->SetDefaultValue(DefaultValue.Color);
-				RGBValue->ApplyDefaultValue();
-			}
-
 			break;
 		}
 	}
diff --git a/Engine/Plugins/Experimental/DynamicMaterial/Source/DynamicMaterialEditor/Private/Components/MaterialStageInputs/DMMSIValue.cpp b/Engine/Plugins/Experimental/DynamicMaterial/Source/DynamicMaterialEditor/Private/Components/MaterialStageInputs/DMMSIValue.cpp
--- a/Engine/Plugins/Experimental/DynamicMaterial/Source/DynamicMaterialEditor/Private/Components/MaterialStageInputs/DMMSIValue.cpp
+++ b/Engine/Plugins/Experimental/DynamicMaterial/Source/DynamicMaterialEditor/Private/Components/MaterialStageInputs/DMMSIValue.cpp
@@ -8,7 +8,11 @@
 #include "Components/DMMaterialStageThroughputLayerBlend.h"
 #include "Components/DMMaterialSubStage.h"
 #include "Components/DMMaterialValue.h"
+#include "Components/MaterialValues/DMMaterialValueFloat3RGB.h"
+#include "Components/MaterialValues/DMMaterialValueFloat4.h"
+#include "Components/MaterialValues/DMMaterialValueTexture.h"
 #include "DMComponentPath.h"
+#include "DynamicMaterialEditorSettings.h"
 #include "Model/DMMaterialBuildState.h"
 #include "Model/DynamicMaterialModel.h"
 #include "Model/DynamicMaterialModelEditorOnlyData.h"
@@ -205,6 +209,8 @@ UDMMaterialStageInputValue* UDMMaterialStageInputValue::ChangeStageInput_NewLoca
 		}
 	);
 
+	NewInputValue->ApplyDefaultLayerSettings();
+
 	NewInputValue->ApplyWholeLayerValue();
 
 	return NewInputValue;
@@ -490,6 +496,81 @@ int32 UDMMaterialStageInputValue::GetInnateMaskOutput(int32 OutputIndex, int32 O
 	return UDMMaterialStageSource::GetInnateMaskOutput(OutputIndex, OutputChannels);
 }
 
+void UDMMaterialStageInputValue::ApplyDefaultLayerSettings()
+{
+	if (!Value)
+	{
+		return;
+	}
+
+	const EDMValueType ValueType = Value->GetType();
+
+	switch (ValueType)
+	{
+		case EDMValueType::VT_Float3_RGB:
+		case EDMValueType::VT_Float4_RGBA:
+		case EDMValueType::VT_Texture:
+			break;
+
+		// Only the above 3 are supported.
+		default:
+			return;
+	}
+
+	const UDynamicMaterialEditorSettings* Settings = UDynamicMaterialEditorSettings::Get();
+
+	if (!Settings)
+	{
+		return;
+	}
+
+	UDMMaterialLayerObject* Layer = GetTypedParent<UDMMaterialLayerObject>(/* Allow Subclasses */ true);
+
+	if (!Layer)
+	{
+		return;
+	}
+
+	const EDMMaterialPropertyType Property = Layer->GetMaterialProperty();
+
+	if (Property == EDMMaterialPropertyType::None)
+	{
+		return;
+	}
+
+	const FDMDefaultMaterialPropertySlotValue& DefaultSlotValue = Settings->GetDefaultSlotValue(Property);
+
+	switch (ValueType)
+	{
+		case EDMValueType::VT_Float3_RGB:
+		{
+			UDMMaterialValueFloat3RGB* Float3Value = Cast<UDMMaterialValueFloat3RGB>(Value);
+			Float3Value->SetDefaultValue(DefaultSlotValue.Color);
+			Float3Value->ApplyDefaultValue();
+			break;
+		}
+
+		case EDMValueType::VT_Float4_RGBA:
+		{
+			UDMMaterialValueFloat4* Float4Value = Cast<UDMMaterialValueFloat4>(Value);
+			Float4Value->SetDefaultValue(DefaultSlotValue.Color);
+			Float4Value->ApplyDefaultValue();
+			break;
+		}
+
+		case EDMValueType::VT_Texture:
+		{
+			if (UTexture* Texture = DefaultSlotValue.Texture.LoadSynchronous())
+			{
+				UDMMaterialValueTexture* TextureValue = Cast<UDMMaterialValueTexture>(Value);
+				TextureValue->SetDefaultValue(Texture);
+				TextureValue->ApplyDefaultValue();
+			}
+			break;
+		}
+	}
+}
+
 bool UDMMaterialStageInputValue::Modify(bool bInAlwaysMarkDirty)
 {
 	const bool bSaved = Super::Modify(bInAlwaysMarkDirty);
diff --git a/Engine/Plugins/Experimental/DynamicMaterial/Source/DynamicMaterialEditor/Private/DynamicMaterialEditorSettings.cpp b/Engine/Plugins/Experimental/DynamicMaterial/Source/DynamicMaterialEditor/Private/DynamicMaterialEditorSettings.cpp
--- a/Engine/Plugins/Experimental/DynamicMaterial/Source/DynamicMaterialEditor/Private/DynamicMaterialEditorSettings.cpp
+++ b/Engine/Plugins/Experimental/DynamicMaterial/Source/DynamicMaterialEditor/Private/DynamicMaterialEditorSettings.cpp
@@ -27,26 +27,38 @@ namespace UE::DynamicMaterialEditor::Private
 		{EDMMaterialPropertyType::Metallic,            TSoftObjectPtr<UTexture>(FSoftObjectPath(TEXT("/Script/Engine.Texture2D'/DynamicMaterial/Textures/SlotDefaults/T_MD_Metallic.T_MD_Metallic'")))},
 		{EDMMaterialPropertyType::Specular,            FLinearColor(0.5f, 0.5f, 0.5f, 1.f)},
 		{EDMMaterialPropertyType::Roughness,           TSoftObjectPtr<UTexture>(FSoftObjectPath(TEXT("/Script/Engine.Texture2D'/DynamicMaterial/Textures/SlotDefaults/T_MD_Roughness.T_MD_Roughness'")))},
-		{EDMMaterialPropertyType::Normal,              TSoftObjectPtr<UTexture>(FSoftObjectPath(TEXT("/Script/Engine.Texture2D'/DynamicMaterial/Textures/SlotDefaults/T_MD_Normal.T_MD_Normal'")))},
+		{EDMMaterialPropertyType::Normal,              {
+				EDMDefaultMaterialPropertySlotValueType::Texture,
+				TSoftObjectPtr<UTexture>(FSoftObjectPath(TEXT("/Script/Engine.Texture2D'/DynamicMaterial/Textures/SlotDefaults/T_MD_Normal.T_MD_Normal'"))),
+				FLinearColor(0.f, 0.f, 1.f, 1.f)
+			}},
 		{EDMMaterialPropertyType::AmbientOcclusion,    TSoftObjectPtr<UTexture>(FSoftObjectPath(TEXT("/Script/Engine.Texture2D'/DynamicMaterial/Textures/SlotDefaults/T_MD_AmbientOcclusion.T_MD_AmbientOcclusion'")))},
 	};
 }
 
 FDMDefaultMaterialPropertySlotValue::FDMDefaultMaterialPropertySlotValue()
-	: Type(EDMDefaultMaterialPropertySlotValueType::Texture)
+	: DefaultType(EDMDefaultMaterialPropertySlotValueType::Texture)
 	, Color(FLinearColor::White)
 {
 }
 
 FDMDefaultMaterialPropertySlotValue::FDMDefaultMaterialPropertySlotValue(const TSoftObjectPtr<UTexture>& InTexture)
-	: Type(EDMDefaultMaterialPropertySlotValueType::Texture)
+	: DefaultType(EDMDefaultMaterialPropertySlotValueType::Texture)
 	, Texture(InTexture)
 	, Color(FLinearColor::Black)
 {
 }
 
 FDMDefaultMaterialPropertySlotValue::FDMDefaultMaterialPropertySlotValue(const FLinearColor& InColor)
-	: Type(EDMDefaultMaterialPropertySlotValueType::Color)
+	: DefaultType(EDMDefaultMaterialPropertySlotValueType::Color)
+	, Color(InColor)
+{
+}
+
+FDMDefaultMaterialPropertySlotValue::FDMDefaultMaterialPropertySlotValue(EDMDefaultMaterialPropertySlotValueType InDefaultType, 
+	const TSoftObjectPtr<UTexture>& InTexture, const FLinearColor& InColor)
+	: DefaultType(InDefaultType)
+	, Texture(InTexture)
 	, Color(InColor)
 {
 }
@@ -426,7 +438,7 @@ const FDMDefaultMaterialPropertySlotValue& UDynamicMaterialEditorSettings::GetDe
 {
 	if (const FDMDefaultMaterialPropertySlotValue* OverridePtr = DefaultSlotValueOverrides.Find(InProperty))
 	{
-		switch (OverridePtr->Type)
+		switch (OverridePtr->DefaultType)
 		{
 			case EDMDefaultMaterialPropertySlotValueType::Texture:
 				if (!OverridePtr->Texture.IsNull() && OverridePtr->Texture.LoadSynchronous())
@@ -442,7 +454,7 @@ const FDMDefaultMaterialPropertySlotValue& UDynamicMaterialEditorSettings::GetDe
 	
 	if (const FDMDefaultMaterialPropertySlotValue* DefaultPtr = UE::DynamicMaterialEditor::Private::DefaultSlotValues.Find(InProperty))
 	{
-		switch (DefaultPtr->Type)
+		switch (DefaultPtr->DefaultType)
 		{
 			case EDMDefaultMaterialPropertySlotValueType::Texture:
 				if (!DefaultPtr->Texture.IsNull() && DefaultPtr->Texture.LoadSynchronous())
diff --git a/Engine/Plugins/Experimental/DynamicMaterial/Source/DynamicMaterialEditor/Private/DynamicMaterialEditorSettings.h b/Engine/Plugins/Experimental/DynamicMaterial/Source/DynamicMaterialEditor/Private/DynamicMaterialEditorSettings.h
--- a/Engine/Plugins/Experimental/DynamicMaterial/Source/DynamicMaterialEditor/Private/DynamicMaterialEditorSettings.h
+++ b/Engine/Plugins/Experimental/DynamicMaterial/Source/DynamicMaterialEditor/Private/DynamicMaterialEditorSettings.h
@@ -44,16 +44,16 @@ struct FDMDefaultMaterialPropertySlotValue
 	FDMDefaultMaterialPropertySlotValue();
 	FDMDefaultMaterialPropertySlotValue(const TSoftObjectPtr<UTexture>& InTexture);
 	FDMDefaultMaterialPropertySlotValue(const FLinearColor& InColor);
+	FDMDefaultMaterialPropertySlotValue(EDMDefaultMaterialPropertySlotValueType InDefaultType, const TSoftObjectPtr<UTexture>& InTexture,
+		const FLinearColor& InColor);
 
 	UPROPERTY(EditAnywhere, BlueprintReadWrite, Category = "Material Designer")
-	EDMDefaultMaterialPropertySlotValueType Type;
+	EDMDefaultMaterialPropertySlotValueType DefaultType;
 
-	UPROPERTY(EditAnywhere, BlueprintReadWrite, Category = "Material Designer", 
-		meta = (EditCondition = "Type == EDMDefaultMaterialPropertySlotValueType::Texture", EditConditionHides))
+	UPROPERTY(EditAnywhere, BlueprintReadWrite, Category = "Material Designer")
 	TSoftObjectPtr<UTexture> Texture;
 
-	UPROPERTY(EditAnywhere, BlueprintReadWrite, Category = "Material Designer",
-		meta = (EditCondition = "Type == EDMDefaultMaterialPropertySlotValueType::Color", EditConditionHides))
+	UPROPERTY(EditAnywhere, BlueprintReadWrite, Category = "Material Designer")
 	FLinearColor Color;
 };
 
diff --git a/Engine/Plugins/Experimental/DynamicMaterial/Source/DynamicMaterialEditor/Private/Slate/SDMSlot.cpp b/Engine/Plugins/Experimental/DynamicMaterial/Source/DynamicMaterialEditor/Private/Slate/SDMSlot.cpp
--- a/Engine/Plugins/Experimental/DynamicMaterial/Source/DynamicMaterialEditor/Private/Slate/SDMSlot.cpp
+++ b/Engine/Plugins/Experimental/DynamicMaterial/Source/DynamicMaterialEditor/Private/Slate/SDMSlot.cpp
@@ -1189,8 +1189,6 @@ void SDMSlot::AddNewLayer_NewLocalValue(TSubclassOf<UDMMaterialValue> InValueCla
 		InValueClass,
 		FDMMaterialStageConnectorChannel::WHOLE_CHANNEL
 	);
-
-	OnValueAdded(NewBase, InputValue->GetValue());
 }
 
 void SDMSlot::AddNewLayer_GlobalValue(UDMMaterialValue* InValue)
@@ -1253,8 +1251,6 @@ void SDMSlot::AddNewLayer_NewGlobalValue(TSubclassOf<UDMMaterialValue> InValueCl
 		InValueClass,
 		FDMMaterialStageConnectorChannel::WHOLE_CHANNEL
 	);
-
-	OnValueAdded(NewStage, InputValue->GetValue());
 }
 
 void SDMSlot::AddNewLayer_Slot(UDMMaterialSlot* InSlot, EDMMaterialPropertyType InMaterialProperty)
@@ -1336,20 +1332,6 @@ void SDMSlot::AddNewLayer_Expression(TSubclassOf<UDMMaterialStageExpression> InE
 			// No nothing
 			break;
 	}
-
-	if (InExpressionClass->IsChildOf(UDMMaterialStageExpressionTextureSampleBase::StaticClass()))
-	{
-		if (UDMMaterialSubStage* SubStage = NewExpression->GetSubStage())
-		{
-			for (UDMMaterialStageInput* Input : SubStage->GetInputs())
-			{
-				if (UDMMaterialStageInputValue* InputValue = Cast<UDMMaterialStageInputValue>(Input))
-				{
-					OnValueAdded(NewStage, InputValue->GetValue());
-				}
-			}
-		}
-	}
 }
 
 void SDMSlot::AddNewLayer_Blend(TSubclassOf<UDMMaterialStageBlend> InBlendClass)
@@ -1512,58 +1494,4 @@ void SDMSlot::OnLayerStageSelected(const bool bInSelected, const TSharedRef<SDMS
 	}
 }
 
-void SDMSlot::OnValueAdded(UDMMaterialStage* InStage, UDMMaterialValue* InValue)
-{
-	if (!InStage || !InValue)
-	{
-		return;
-	}
-
-	UDMMaterialLayerObject* Layer = InStage->GetLayer();
-
-	if (!Layer)
-	{
-		return;
-	}
-
-	UClass* ValueClass = InValue->GetClass();
-
-	// Do not include subclasses
-	if (ValueClass != UDMMaterialValueTexture::StaticClass()
-		&& ValueClass != UDMMaterialValueFloat3RGB::StaticClass()
-		&& ValueClass != UDMMaterialValueFloat4::StaticClass())
-	{
-		return;
-	}
-
-	const FDMDefaultMaterialPropertySlotValue& DefaultValue = GetDefault<UDynamicMaterialEditorSettings>()->GetDefaultSlotValue(Layer->GetMaterialProperty());
-
-	switch (DefaultValue.Type)
-	{
-		case EDMDefaultMaterialPropertySlotValueType::Texture:
-			if (UDMMaterialValueTexture* TextureValue = Cast<UDMMaterialValueTexture>(InValue))
-			{
-				if (UTexture* Texture = DefaultValue.Texture.LoadSynchronous())
-				{
-					TextureValue->SetDefaultValue(Texture);
-					TextureValue->ApplyDefaultValue();
-				}
-			}
-			break;
-
-		case EDMDefaultMaterialPropertySlotValueType::Color:
-			if (UDMMaterialValueFloat3RGB* RGBValue = Cast<UDMMaterialValueFloat3RGB>(InValue))
-			{
-				RGBValue->SetDefaultValue(DefaultValue.Color);
-				RGBValue->ApplyDefaultValue();
-			}
-			else if (UDMMaterialValueFloat4* RGBAValue = Cast<UDMMaterialValueFloat4>(InValue))
-			{
-				RGBAValue->SetDefaultValue(DefaultValue.Color);
-				RGBAValue->ApplyDefaultValue();
-			}
-			break;
-	}
-}
-
 #undef LOCTEXT_NAMESPACE
diff --git a/Engine/Plugins/Experimental/DynamicMaterial/Source/DynamicMaterialEditor/Private/Slate/SDMSlot.h b/Engine/Plugins/Experimental/DynamicMaterial/Source/DynamicMaterialEditor/Private/Slate/SDMSlot.h
--- a/Engine/Plugins/Experimental/DynamicMaterial/Source/DynamicMaterialEditor/Private/Slate/SDMSlot.h
+++ b/Engine/Plugins/Experimental/DynamicMaterial/Source/DynamicMaterialEditor/Private/Slate/SDMSlot.h
@@ -134,8 +134,6 @@ protected:
 	void OnLayerSelected(TSharedPtr<FDMMaterialLayerReference> InLayerItem, const int32 InLayerIndex);
 	void OnLayerStageSelected(const bool bInSelected, const TSharedRef<SDMStage>& InStageWidget);
 
-	void OnValueAdded(UDMMaterialStage* InStage, UDMMaterialValue* InValue);
-
 	void RefreshMainWidget();
 	void RefreshHeaderPropertyListWidget();
 	void RefreshSlotSettingsRowWidget();
diff --git a/Engine/Plugins/Experimental/DynamicMaterial/Source/DynamicMaterialEditor/Public/Components/MaterialStageInputs/DMMSIValue.h b/Engine/Plugins/Experimental/DynamicMaterial/Source/DynamicMaterialEditor/Public/Components/MaterialStageInputs/DMMSIValue.h
--- a/Engine/Plugins/Experimental/DynamicMaterial/Source/DynamicMaterialEditor/Public/Components/MaterialStageInputs/DMMSIValue.h
+++ b/Engine/Plugins/Experimental/DynamicMaterial/Source/DynamicMaterialEditor/Public/Components/MaterialStageInputs/DMMSIValue.h
@@ -66,6 +66,8 @@ public:
 
 	virtual int32 GetInnateMaskOutput(int32 OutputIndex, int32 OutputChannels) const override;
 
+	void ApplyDefaultLayerSettings();
+
 	//~ Begin UObject
 	virtual bool Modify(bool bInAlwaysMarkDirty = true) override;
 	virtual void PostLoad() override;
