Material Designer
- Changed normal material property to use the normals strength effect node.

diff --git a/Engine/Plugins/Experimental/DynamicMaterial/Source/DynamicMaterialEditor/Private/Components/MaterialProperties/DMMPNormal.cpp b/Engine/Plugins/Experimental/DynamicMaterial/Source/DynamicMaterialEditor/Private/Components/MaterialProperties/DMMPNormal.cpp
--- a/Engine/Plugins/Experimental/DynamicMaterial/Source/DynamicMaterialEditor/Private/Components/MaterialProperties/DMMPNormal.cpp
+++ b/Engine/Plugins/Experimental/DynamicMaterial/Source/DynamicMaterialEditor/Private/Components/MaterialProperties/DMMPNormal.cpp
@@ -21,14 +21,14 @@ namespace UE::DynamicMaterialEditor::Private
 		return NormalizeBlend;
 	}
 
-	UMaterialFunctionInterface* GetNormalMagnitude()
+	UMaterialFunctionInterface* GetNormalsStrength()
 	{
-		static UMaterialFunctionInterface* NormalMagnitude = FDMMaterialFunctionLibrary::Get().GetFunction(
-			"MF_DM_Normal_Magnitude",
-			TEXT("/Script/Engine.MaterialFunction'/DynamicMaterial/MaterialFunctions/MF_DM_Normal_Magnitude.MF_DM_Normal_Magnitude'")
+		static UMaterialFunctionInterface* NormalsStrength = FDMMaterialFunctionLibrary::Get().GetFunction(
+			"MF_DM_Normals_Strength",
+			TEXT("/Script/Engine.MaterialFunction'/DynamicMaterial/MaterialFunctions/Effects/Normals/MF_DM_Normals_Strength.MF_DM_Normals_Strength'")
 		);
 
-		return NormalMagnitude;
+		return NormalsStrength;
 	}
 }
 
@@ -119,7 +119,7 @@ void UDMMaterialPropertyNormal::AddAlphaMultiplier(const TSharedRef<FDMMaterialB
 
 	UMaterialExpressionMaterialFunctionCall* MaterialFunctionCall = FDMMaterialFunctionLibrary::Get().MakeExpression(
 		InBuildState->GetDynamicMaterial(),
-		UE::DynamicMaterialEditor::Private::GetNormalMagnitude(),
+		UE::DynamicMaterialEditor::Private::GetNormalsStrength(),
 		UE_DM_NodeComment_Default
 	);
 
