Material Designer
- Duplicated the default metallic texture for the emissive channel.
- Added to new texture as the default for the emissive channel.

#virtualized

diff --git a/Engine/Plugins/Experimental/DynamicMaterial/Source/DynamicMaterialEditor/Private/DynamicMaterialEditorSettings.cpp b/Engine/Plugins/Experimental/DynamicMaterial/Source/DynamicMaterialEditor/Private/DynamicMaterialEditorSettings.cpp
--- a/Engine/Plugins/Experimental/DynamicMaterial/Source/DynamicMaterialEditor/Private/DynamicMaterialEditorSettings.cpp
+++ b/Engine/Plugins/Experimental/DynamicMaterial/Source/DynamicMaterialEditor/Private/DynamicMaterialEditorSettings.cpp
@@ -21,7 +21,7 @@ namespace UE::DynamicMaterialEditor::Private
 {
 	const TMap<EDMMaterialPropertyType, FDMDefaultMaterialPropertySlotValue> DefaultSlotValues = {
 		{EDMMaterialPropertyType::BaseColor,           TSoftObjectPtr<UTexture>(FSoftObjectPath(TEXT("Texture2D'/DynamicMaterial/T_Default_Texture.T_Default_Texture'")))},
-		{EDMMaterialPropertyType::EmissiveColor,       TSoftObjectPtr<UTexture>(FSoftObjectPath(TEXT("Texture2D'/DynamicMaterial/T_Default_Texture.T_Default_Texture'")))},
+		{EDMMaterialPropertyType::EmissiveColor,       TSoftObjectPtr<UTexture>(FSoftObjectPath(TEXT("/Script/Engine.Texture2D'/DynamicMaterial/Textures/SlotDefaults/T_MD_Emissive.T_MD_Emissive'")))},
 		{EDMMaterialPropertyType::Opacity,             TSoftObjectPtr<UTexture>(FSoftObjectPath(TEXT("/Script/Engine.Texture2D'/DynamicMaterial/Textures/T_DM_HorizontalGradient.T_DM_HorizontalGradient'")))},
 		{EDMMaterialPropertyType::OpacityMask,         TSoftObjectPtr<UTexture>(FSoftObjectPath(TEXT("/Script/Engine.Texture2D'/DynamicMaterial/Textures/T_DM_HorizontalGradient.T_DM_HorizontalGradient'")))},
 		{EDMMaterialPropertyType::Metallic,            TSoftObjectPtr<UTexture>(FSoftObjectPath(TEXT("/Script/Engine.Texture2D'/DynamicMaterial/Textures/SlotDefaults/T_MD_Metallic.T_MD_Metallic'")))},
