Material Designer
- WITH_EDITOR issue.

diff --git a/Engine/Plugins/Experimental/DynamicMaterial/Source/DynamicMaterial/Private/Components/DMTextureUV.cpp b/Engine/Plugins/Experimental/DynamicMaterial/Source/DynamicMaterial/Private/Components/DMTextureUV.cpp
--- a/Engine/Plugins/Experimental/DynamicMaterial/Source/DynamicMaterial/Private/Components/DMTextureUV.cpp
+++ b/Engine/Plugins/Experimental/DynamicMaterial/Source/DynamicMaterial/Private/Components/DMTextureUV.cpp
@@ -513,25 +513,22 @@ void UDMTextureUV::RemoveParameterNames()
 	UDynamicMaterialModel* MaterialModel = GetMaterialModel();
 	check(MaterialModel);
 
-#if WITH_EDITOR
 	if (GUndo)
 	{
 		MaterialModel->Modify();
 	}
-#endif
 
 	for (const TPair<int32, TObjectPtr<UDMMaterialParameter>>& Pair : MaterialParameters)
 	{
-#if WITH_EDITOR
 		if (GUndo)
 		{
 			Pair.Value->Modify();
 		}
-#endif
 
 		MaterialModel->FreeParameter(Pair.Value);
 	}
 }
+#endif
 
 void UDMTextureUV::OnTextureUVChanged(EDMUpdateType InUpdateType)
 {
@@ -553,6 +550,7 @@ void UDMTextureUV::OnTextureUVChanged(EDMUpdateType InUpdateType)
 	}
 }
 
+#if WITH_EDITOR
 void UDMTextureUV::OnComponentAdded()
 {
 	if (!IsComponentValid())
