Material Designer
- Changed the update blocker to work within the update method.
- Copying and pasting layers between slots now works.

diff --git a/Engine/Plugins/Experimental/DynamicMaterial/Source/DynamicMaterial/Private/Components/DMMaterialComponent.cpp b/Engine/Plugins/Experimental/DynamicMaterial/Source/DynamicMaterial/Private/Components/DMMaterialComponent.cpp
--- a/Engine/Plugins/Experimental/DynamicMaterial/Source/DynamicMaterial/Private/Components/DMMaterialComponent.cpp
+++ b/Engine/Plugins/Experimental/DynamicMaterial/Source/DynamicMaterial/Private/Components/DMMaterialComponent.cpp
@@ -75,6 +75,11 @@ UDMMaterialComponent* UDMMaterialComponent::GetSubComponentByPath(FDMComponentPa
 
 void UDMMaterialComponent::Update(EDMUpdateType InUpdateType)
 {
+	if (!FDMUpdateGuard::CanUpdate())
+	{
+		return;
+	}
+
 #if WITH_EDITOR
 	OnUpdate.Broadcast(this, InUpdateType);
 #endif
diff --git a/Engine/Plugins/Experimental/DynamicMaterial/Source/DynamicMaterial/Private/Components/DMMaterialValue.cpp b/Engine/Plugins/Experimental/DynamicMaterial/Source/DynamicMaterial/Private/Components/DMMaterialValue.cpp
--- a/Engine/Plugins/Experimental/DynamicMaterial/Source/DynamicMaterial/Private/Components/DMMaterialValue.cpp
+++ b/Engine/Plugins/Experimental/DynamicMaterial/Source/DynamicMaterial/Private/Components/DMMaterialValue.cpp
@@ -440,21 +440,23 @@ void UDMMaterialValue::OnValueChanged(EDMUpdateType InUpdateType)
 		return;
 	}
 
-	if (FDMUpdateGuard::CanUpdate())
-	{
-		Update(InUpdateType);
+	Update(InUpdateType);
 
 #if WITH_EDITOR
-		if (EnumHasAnyFlags(InUpdateType, EDMUpdateType::AllowParentUpdate) && ParentComponent)
-		{
-			ParentComponent->Update(InUpdateType);
-		}
-#endif
+	if (EnumHasAnyFlags(InUpdateType, EDMUpdateType::AllowParentUpdate) && ParentComponent)
+	{
+		ParentComponent->Update(InUpdateType);
 	}
+#endif
 }
  
 void UDMMaterialValue::Update(EDMUpdateType InUpdateType)
 {
+	if (!FDMUpdateGuard::CanUpdate())
+	{
+		return;
+	}
+
 	if (!IsComponentValid())
 	{
 		return;
diff --git a/Engine/Plugins/Experimental/DynamicMaterial/Source/DynamicMaterial/Private/Components/DMRenderTargetRenderer.cpp b/Engine/Plugins/Experimental/DynamicMaterial/Source/DynamicMaterial/Private/Components/DMRenderTargetRenderer.cpp
--- a/Engine/Plugins/Experimental/DynamicMaterial/Source/DynamicMaterial/Private/Components/DMRenderTargetRenderer.cpp
+++ b/Engine/Plugins/Experimental/DynamicMaterial/Source/DynamicMaterial/Private/Components/DMRenderTargetRenderer.cpp
@@ -78,6 +78,11 @@ bool UDMRenderTargetRenderer::JsonDeserialize(const TSharedPtr<FJsonValue>& InJs
 
 void UDMRenderTargetRenderer::Update(EDMUpdateType InUpdateType)
 {
+	if (!FDMUpdateGuard::CanUpdate())
+	{
+		return;
+	}
+
 	if (!IsComponentValid())
 	{
 		return;
diff --git a/Engine/Plugins/Experimental/DynamicMaterial/Source/DynamicMaterial/Private/Components/DMTextureUV.cpp b/Engine/Plugins/Experimental/DynamicMaterial/Source/DynamicMaterial/Private/Components/DMTextureUV.cpp
--- a/Engine/Plugins/Experimental/DynamicMaterial/Source/DynamicMaterial/Private/Components/DMTextureUV.cpp
+++ b/Engine/Plugins/Experimental/DynamicMaterial/Source/DynamicMaterial/Private/Components/DMTextureUV.cpp
@@ -371,6 +371,11 @@ bool UDMTextureUV::Modify(bool bInAlwaysMarkDirty)
 
 void UDMTextureUV::Update(EDMUpdateType InUpdateType)
 {
+	if (!FDMUpdateGuard::CanUpdate())
+	{
+		return;
+	}
+
 	if (!IsComponentValid())
 	{
 		return;
@@ -537,17 +542,14 @@ void UDMTextureUV::OnTextureUVChanged(EDMUpdateType InUpdateType)
 		return;
 	}
 
-	if (FDMUpdateGuard::CanUpdate())
-	{
-		Update(InUpdateType);
+	Update(InUpdateType);
 
 #if WITH_EDITOR
-		if (EnumHasAnyFlags(InUpdateType, EDMUpdateType::AllowParentUpdate) && ParentComponent)
-		{
-			ParentComponent->Update(InUpdateType);
-		}
-#endif
+	if (EnumHasAnyFlags(InUpdateType, EDMUpdateType::AllowParentUpdate) && ParentComponent)
+	{
+		ParentComponent->Update(InUpdateType);
 	}
+#endif
 }
 
 #if WITH_EDITOR
diff --git a/Engine/Plugins/Experimental/DynamicMaterial/Source/DynamicMaterial/Private/Components/MaterialValues/DMMaterialValueRenderTarget.cpp b/Engine/Plugins/Experimental/DynamicMaterial/Source/DynamicMaterial/Private/Components/MaterialValues/DMMaterialValueRenderTarget.cpp
--- a/Engine/Plugins/Experimental/DynamicMaterial/Source/DynamicMaterial/Private/Components/MaterialValues/DMMaterialValueRenderTarget.cpp
+++ b/Engine/Plugins/Experimental/DynamicMaterial/Source/DynamicMaterial/Private/Components/MaterialValues/DMMaterialValueRenderTarget.cpp
@@ -140,6 +140,11 @@ void UDMMaterialValueRenderTarget::FlushCreateRenderTarget()
 
 void UDMMaterialValueRenderTarget::Update(EDMUpdateType InUpdateType)
 {
+	if (!FDMUpdateGuard::CanUpdate())
+	{
+		return;
+	}
+
 	Super::Update(InUpdateType);
 
 	if (!IsValid(GetRenderTarget()))
diff --git a/Engine/Plugins/Experimental/DynamicMaterial/Source/DynamicMaterialEditor/Private/Components/DMMaterialEffect.cpp b/Engine/Plugins/Experimental/DynamicMaterial/Source/DynamicMaterialEditor/Private/Components/DMMaterialEffect.cpp
--- a/Engine/Plugins/Experimental/DynamicMaterial/Source/DynamicMaterialEditor/Private/Components/DMMaterialEffect.cpp
+++ b/Engine/Plugins/Experimental/DynamicMaterial/Source/DynamicMaterialEditor/Private/Components/DMMaterialEffect.cpp
@@ -110,6 +110,11 @@ FText UDMMaterialEffect::GetComponentDescription() const
 
 void UDMMaterialEffect::Update(EDMUpdateType InUpdateType)
 {
+	if (!FDMUpdateGuard::CanUpdate())
+	{
+		return;
+	}
+
 	if (!IsComponentValid())
 	{
 		return;
diff --git a/Engine/Plugins/Experimental/DynamicMaterial/Source/DynamicMaterialEditor/Private/Components/DMMaterialEffectStack.cpp b/Engine/Plugins/Experimental/DynamicMaterial/Source/DynamicMaterialEditor/Private/Components/DMMaterialEffectStack.cpp
--- a/Engine/Plugins/Experimental/DynamicMaterial/Source/DynamicMaterialEditor/Private/Components/DMMaterialEffectStack.cpp
+++ b/Engine/Plugins/Experimental/DynamicMaterial/Source/DynamicMaterialEditor/Private/Components/DMMaterialEffectStack.cpp
@@ -389,6 +389,11 @@ FText UDMMaterialEffectStack::GetComponentDescription() const
 
 void UDMMaterialEffectStack::Update(EDMUpdateType InUpdateType)
 {
+	if (!FDMUpdateGuard::CanUpdate())
+	{
+		return;
+	}
+
 	if (!IsComponentValid())
 	{
 		return;
diff --git a/Engine/Plugins/Experimental/DynamicMaterial/Source/DynamicMaterialEditor/Private/Components/DMMaterialLayer.cpp b/Engine/Plugins/Experimental/DynamicMaterial/Source/DynamicMaterialEditor/Private/Components/DMMaterialLayer.cpp
--- a/Engine/Plugins/Experimental/DynamicMaterial/Source/DynamicMaterialEditor/Private/Components/DMMaterialLayer.cpp
+++ b/Engine/Plugins/Experimental/DynamicMaterial/Source/DynamicMaterialEditor/Private/Components/DMMaterialLayer.cpp
@@ -789,6 +789,11 @@ FText UDMMaterialLayerObject::GetComponentDescription() const
 
 void UDMMaterialLayerObject::Update(EDMUpdateType InUpdateType)
 {
+	if (!FDMUpdateGuard::CanUpdate())
+	{
+		return;
+	}
+
 	if (!IsComponentValid())
 	{
 		return;
diff --git a/Engine/Plugins/Experimental/DynamicMaterial/Source/DynamicMaterialEditor/Private/Components/DMMaterialProperty.cpp b/Engine/Plugins/Experimental/DynamicMaterial/Source/DynamicMaterialEditor/Private/Components/DMMaterialProperty.cpp
--- a/Engine/Plugins/Experimental/DynamicMaterial/Source/DynamicMaterialEditor/Private/Components/DMMaterialProperty.cpp
+++ b/Engine/Plugins/Experimental/DynamicMaterial/Source/DynamicMaterialEditor/Private/Components/DMMaterialProperty.cpp
@@ -572,6 +572,11 @@ void UDMMaterialProperty::AddOutputProcessor(const TSharedRef<FDMMaterialBuildSt
 
 void UDMMaterialProperty::Update(EDMUpdateType InUpdateType)
 {
+	if (!FDMUpdateGuard::CanUpdate())
+	{
+		return;
+	}
+
 	if (!IsComponentValid())
 	{
 		return;
diff --git a/Engine/Plugins/Experimental/DynamicMaterial/Source/DynamicMaterialEditor/Private/Components/DMMaterialSlot.cpp b/Engine/Plugins/Experimental/DynamicMaterial/Source/DynamicMaterialEditor/Private/Components/DMMaterialSlot.cpp
--- a/Engine/Plugins/Experimental/DynamicMaterial/Source/DynamicMaterialEditor/Private/Components/DMMaterialSlot.cpp
+++ b/Engine/Plugins/Experimental/DynamicMaterial/Source/DynamicMaterialEditor/Private/Components/DMMaterialSlot.cpp
@@ -182,6 +182,11 @@ UDMMaterialLayerObject* UDMMaterialSlot::GetLastLayerForMaterialProperty(EDMMate
 
 void UDMMaterialSlot::Update(EDMUpdateType InUpdateType)
 {
+	if (!FDMUpdateGuard::CanUpdate())
+	{
+		return;
+	}
+
 	if (!IsComponentValid())
 	{
 		return;
@@ -555,7 +560,45 @@ bool UDMMaterialSlot::PasteLayer(UDMMaterialLayerObject* InLayer)
 		return false;
 	}
 
-	UDynamicMaterialModel* MaterialModel = ModelEditorOnlyData->GetMaterialModel();
+	EDMMaterialPropertyType NewPropertyType = EDMMaterialPropertyType::None;
+
+	if (!LayerObjects.IsEmpty())
+	{
+		for (UDMMaterialLayerObject* CurrentLayer : UE::Core::Private::TReverseIterationAdapter(LayerObjects))
+		{
+			if (CurrentLayer->IsEnabled())
+			{
+				NewPropertyType = CurrentLayer->GetMaterialProperty();
+				break;
+			}
+		}
+
+		if (NewPropertyType == EDMMaterialPropertyType::None)
+		{
+			for (UDMMaterialLayerObject* CurrentLayer : UE::Core::Private::TReverseIterationAdapter(LayerObjects))
+			{
+				NewPropertyType = CurrentLayer->GetMaterialProperty();
+				break;
+			}
+		}
+	}
+
+	if (NewPropertyType == EDMMaterialPropertyType::None)
+	{
+		TArray<EDMMaterialPropertyType> SlotProperties = ModelEditorOnlyData->GetMaterialPropertiesForSlot(this);
+
+		if (SlotProperties.IsEmpty())
+		{
+			return false;
+		}
+
+		NewPropertyType = SlotProperties[0];
+	}
+
+	{
+		const FDMUpdateGuard Guard;
+		InLayer->SetMaterialProperty(NewPropertyType);
+	}
 
 	LayerObjects.Add(InLayer);
 
diff --git a/Engine/Plugins/Experimental/DynamicMaterial/Source/DynamicMaterialEditor/Private/Components/DMMaterialStage.cpp b/Engine/Plugins/Experimental/DynamicMaterial/Source/DynamicMaterialEditor/Private/Components/DMMaterialStage.cpp
--- a/Engine/Plugins/Experimental/DynamicMaterial/Source/DynamicMaterialEditor/Private/Components/DMMaterialStage.cpp
+++ b/Engine/Plugins/Experimental/DynamicMaterial/Source/DynamicMaterialEditor/Private/Components/DMMaterialStage.cpp
@@ -379,22 +379,19 @@ void UDMMaterialStage::SetSource(UDMMaterialStageSource* InSource)
 
 	Source = InSource;
 
-	if (FDMUpdateGuard::CanUpdate())
-	{
-		ResetInputConnectionMap();
+	ResetInputConnectionMap();
 
-		if (IsComponentAdded())
+	if (IsComponentAdded())
+	{
+		if (GUndo)
 		{
-			if (GUndo)
-			{
-				Source->Modify();
-			}
-
-			Source->SetComponentState(EDMComponentLifetimeState::Added);
+			Source->Modify();
 		}
 
-		Update(EDMUpdateType::Structure);
+		Source->SetComponentState(EDMComponentLifetimeState::Added);
 	}
+
+	Update(EDMUpdateType::Structure);
 }
 
 FText UDMMaterialStage::GetComponentDescription() const
@@ -462,6 +459,11 @@ bool UDMMaterialStage::IsInputMapped(int32 InputIndex) const
 
 void UDMMaterialStage::Update(EDMUpdateType InUpdateType)
 {
+	if (!FDMUpdateGuard::CanUpdate())
+	{
+		return;
+	}
+
 	if (!IsComponentValid())
 	{
 		return;
@@ -951,10 +953,7 @@ void UDMMaterialStage::UpdateInputMap(int32 InInputIdx, int32 InSourceIndex, int
 
 	RemoveUnusedInputs();
 
-	if (FDMUpdateGuard::CanUpdate())
-	{
-		Source->Update(EDMUpdateType::Structure);
-	}
+	Source->Update(EDMUpdateType::Structure);
 }
 
 int32 UDMMaterialStage::FindIndex() const
diff --git a/Engine/Plugins/Experimental/DynamicMaterial/Source/DynamicMaterialEditor/Private/Components/DMMaterialStageFunction.cpp b/Engine/Plugins/Experimental/DynamicMaterial/Source/DynamicMaterialEditor/Private/Components/DMMaterialStageFunction.cpp
--- a/Engine/Plugins/Experimental/DynamicMaterial/Source/DynamicMaterialEditor/Private/Components/DMMaterialStageFunction.cpp
+++ b/Engine/Plugins/Experimental/DynamicMaterial/Source/DynamicMaterialEditor/Private/Components/DMMaterialStageFunction.cpp
@@ -288,10 +288,7 @@ void UDMMaterialStageFunction::OnMaterialFunctionChanged()
 	DeinitFunction();
 	InitFunction();
 
-	if (FDMUpdateGuard::CanUpdate())
-	{
-		Update(EDMUpdateType::Structure);
-	}
+	Update(EDMUpdateType::Structure);
 }
 
 bool UDMMaterialStageFunction::NeedsFunctionInit() const
diff --git a/Engine/Plugins/Experimental/DynamicMaterial/Source/DynamicMaterialEditor/Private/Components/DMMaterialStageInput.cpp b/Engine/Plugins/Experimental/DynamicMaterial/Source/DynamicMaterialEditor/Private/Components/DMMaterialStageInput.cpp
--- a/Engine/Plugins/Experimental/DynamicMaterial/Source/DynamicMaterialEditor/Private/Components/DMMaterialStageInput.cpp
+++ b/Engine/Plugins/Experimental/DynamicMaterial/Source/DynamicMaterialEditor/Private/Components/DMMaterialStageInput.cpp
@@ -14,6 +14,11 @@ const FString UDMMaterialStageInput::StageInputPrefixStr = FString(TEXT("DMMater
 
 void UDMMaterialStageInput::Update(EDMUpdateType InUpdateType)
 {
+	if (!FDMUpdateGuard::CanUpdate())
+	{
+		return;
+	}
+
 	if (!IsComponentValid())
 	{
 		return;
diff --git a/Engine/Plugins/Experimental/DynamicMaterial/Source/DynamicMaterialEditor/Private/Components/DMMaterialStageSource.cpp b/Engine/Plugins/Experimental/DynamicMaterial/Source/DynamicMaterialEditor/Private/Components/DMMaterialStageSource.cpp
--- a/Engine/Plugins/Experimental/DynamicMaterial/Source/DynamicMaterialEditor/Private/Components/DMMaterialStageSource.cpp
+++ b/Engine/Plugins/Experimental/DynamicMaterial/Source/DynamicMaterialEditor/Private/Components/DMMaterialStageSource.cpp
@@ -33,6 +33,11 @@ UDMMaterialStage* UDMMaterialStageSource::GetStage() const
 
 void UDMMaterialStageSource::Update(EDMUpdateType InUpdateType)
 {
+	if (!FDMUpdateGuard::CanUpdate())
+	{
+		return;
+	}
+
 	if (!IsComponentValid())
 	{
 		return;
@@ -65,10 +70,7 @@ void UDMMaterialStageSource::OnComponentAdded()
 
 	Super::OnComponentAdded();
 
-	if (FDMUpdateGuard::CanUpdate())
-	{
-		Update(EDMUpdateType::Structure);
-	}
+	Update(EDMUpdateType::Structure);
 }
 
 void UDMMaterialStageSource::GetMaskAlphaBlendNode(const TSharedRef<FDMMaterialBuildState>& InBuildState, UMaterialExpression*& OutExpression, int32& OutOutputIndex, int32& OutOutputChannel) const
diff --git a/Engine/Plugins/Experimental/DynamicMaterial/Source/DynamicMaterialEditor/Private/Components/DMMaterialStageThroughputLayerBlend.cpp b/Engine/Plugins/Experimental/DynamicMaterial/Source/DynamicMaterialEditor/Private/Components/DMMaterialStageThroughputLayerBlend.cpp
--- a/Engine/Plugins/Experimental/DynamicMaterial/Source/DynamicMaterialEditor/Private/Components/DMMaterialStageThroughputLayerBlend.cpp
+++ b/Engine/Plugins/Experimental/DynamicMaterial/Source/DynamicMaterialEditor/Private/Components/DMMaterialStageThroughputLayerBlend.cpp
@@ -179,6 +179,11 @@ void UDMMaterialStageThroughputLayerBlend::Update(EDMUpdateType InUpdateType)
 		return;
 	}
 
+	if (!FDMUpdateGuard::CanUpdate())
+	{
+		return;
+	}
+
 	if (!IsComponentValid())
 	{
 		return;
diff --git a/Engine/Plugins/Experimental/DynamicMaterial/Source/DynamicMaterialEditor/Private/Components/MaterialStageExpressions/DMMSEMathBase.cpp b/Engine/Plugins/Experimental/DynamicMaterial/Source/DynamicMaterialEditor/Private/Components/MaterialStageExpressions/DMMSEMathBase.cpp
--- a/Engine/Plugins/Experimental/DynamicMaterial/Source/DynamicMaterialEditor/Private/Components/MaterialStageExpressions/DMMSEMathBase.cpp
+++ b/Engine/Plugins/Experimental/DynamicMaterial/Source/DynamicMaterialEditor/Private/Components/MaterialStageExpressions/DMMSEMathBase.cpp
@@ -208,6 +208,11 @@ bool UDMMaterialStageExpressionMathBase::CanInputAcceptType(int32 InInputIndex,
 
 void UDMMaterialStageExpressionMathBase::Update(EDMUpdateType InUpdateType)
 {
+	if (!FDMUpdateGuard::CanUpdate())
+	{
+		return;
+	}
+
 	if (HasComponentBeenRemoved())
 	{
 		return;
diff --git a/Engine/Plugins/Experimental/DynamicMaterial/Source/DynamicMaterialEditor/Private/Components/MaterialStageInputs/DMMSISlot.cpp b/Engine/Plugins/Experimental/DynamicMaterial/Source/DynamicMaterialEditor/Private/Components/MaterialStageInputs/DMMSISlot.cpp
--- a/Engine/Plugins/Experimental/DynamicMaterial/Source/DynamicMaterialEditor/Private/Components/MaterialStageInputs/DMMSISlot.cpp
+++ b/Engine/Plugins/Experimental/DynamicMaterial/Source/DynamicMaterialEditor/Private/Components/MaterialStageInputs/DMMSISlot.cpp
@@ -185,10 +185,7 @@ void UDMMaterialStageInputSlot::SetSlot(UDMMaterialSlot* InSlot)
 
 	UpdateOutputConnectors();
 
-	if (FDMUpdateGuard::CanUpdate())
-	{
-		Update(EDMUpdateType::Structure);
-	}
+	Update(EDMUpdateType::Structure);
 }
 
 void UDMMaterialStageInputSlot::SetMaterialProperty(EDMMaterialPropertyType InMaterialProperty)
@@ -207,10 +204,7 @@ void UDMMaterialStageInputSlot::SetMaterialProperty(EDMMaterialPropertyType InMa
 
 	UpdateOutputConnectors();
 
-	if (FDMUpdateGuard::CanUpdate())
-	{
-		Update(EDMUpdateType::Structure);
-	}
+	Update(EDMUpdateType::Structure);
 }
 
 void UDMMaterialStageInputSlot::OnComponentRemoved()
diff --git a/Engine/Plugins/Experimental/DynamicMaterial/Source/DynamicMaterialEditor/Private/Components/MaterialStageInputs/DMMSIThroughput.cpp b/Engine/Plugins/Experimental/DynamicMaterial/Source/DynamicMaterialEditor/Private/Components/MaterialStageInputs/DMMSIThroughput.cpp
--- a/Engine/Plugins/Experimental/DynamicMaterial/Source/DynamicMaterialEditor/Private/Components/MaterialStageInputs/DMMSIThroughput.cpp
+++ b/Engine/Plugins/Experimental/DynamicMaterial/Source/DynamicMaterialEditor/Private/Components/MaterialStageInputs/DMMSIThroughput.cpp
@@ -59,10 +59,7 @@ void UDMMaterialStageInputThroughput::SetMaterialStageThroughputClass(TSubclassO
 		}
 	}
 
-	if (FDMUpdateGuard::CanUpdate())
-	{
-		Update(EDMUpdateType::Structure);
-	}
+	Update(EDMUpdateType::Structure);
 }
 
 UDMMaterialStageThroughput* UDMMaterialStageInputThroughput::GetMaterialStageThroughput() const
@@ -294,7 +291,7 @@ void UDMMaterialStageInputThroughput::OnSubStageUpdated(UDMMaterialComponent* In
 		return;
 	}
 
-	if (InComponent == SubStage && FDMUpdateGuard::CanUpdate())
+	if (InComponent == SubStage)
 	{
 		Update(InUpdateType);
 	}
diff --git a/Engine/Plugins/Experimental/DynamicMaterial/Source/DynamicMaterialEditor/Private/Components/MaterialStageInputs/DMMSIValue.cpp b/Engine/Plugins/Experimental/DynamicMaterial/Source/DynamicMaterialEditor/Private/Components/MaterialStageInputs/DMMSIValue.cpp
--- a/Engine/Plugins/Experimental/DynamicMaterial/Source/DynamicMaterialEditor/Private/Components/MaterialStageInputs/DMMSIValue.cpp
+++ b/Engine/Plugins/Experimental/DynamicMaterial/Source/DynamicMaterialEditor/Private/Components/MaterialStageInputs/DMMSIValue.cpp
@@ -462,10 +462,7 @@ void UDMMaterialStageInputValue::SetValue(UDMMaterialValue* InValue)
 
 	UpdateOutputConnectors();
 
-	if (FDMUpdateGuard::CanUpdate())
-	{
-		Update(EDMUpdateType::Structure);
-	}
+	Update(EDMUpdateType::Structure);
 }
 
 void UDMMaterialStageInputValue::GenerateExpressions(const TSharedRef<FDMMaterialBuildState>& InBuildState) const
