Material Designer
- Changed wrong Editor Only Data version registeration variable name.

diff --git a/Engine/Plugins/Experimental/DynamicMaterial/Source/DynamicMaterialEditor/Private/Model/DynamicMaterialModelEditorOnlyDataVersion.cpp b/Engine/Plugins/Experimental/DynamicMaterial/Source/DynamicMaterialEditor/Private/Model/DynamicMaterialModelEditorOnlyDataVersion.cpp
--- a/Engine/Plugins/Experimental/DynamicMaterial/Source/DynamicMaterialEditor/Private/Model/DynamicMaterialModelEditorOnlyDataVersion.cpp
+++ b/Engine/Plugins/Experimental/DynamicMaterial/Source/DynamicMaterialEditor/Private/Model/DynamicMaterialModelEditorOnlyDataVersion.cpp
@@ -4,4 +4,4 @@
 #include "Serialization/CustomVersion.h"
 
 const FGuid FDynamicMaterialModelEditorOnlyDataVersion::GUID(0xFCF57AFC, 0x50764285, 0xB9A9E660, 0xFFA02D34);
-FCustomVersionRegistration GRegisterDMTextureUVVersion(FDynamicMaterialModelEditorOnlyDataVersion::GUID, static_cast<int32>(FDynamicMaterialModelEditorOnlyDataVersion::Type::LatestVersion), TEXT("DynamicMaterialModelEditorOnlyData"));
+FCustomVersionRegistration GRegisterDynamicMaterialModelEditorOnlyDataVersion(FDynamicMaterialModelEditorOnlyDataVersion::GUID, static_cast<int32>(FDynamicMaterialModelEditorOnlyDataVersion::Type::LatestVersion), TEXT("DynamicMaterialModelEditorOnlyData"));
