Material Designer
- Material Designer Dynamic.
- Updating documentation and using method specific API exports.
- Open parent material button
- Convert to non-dynamic button
- Parameter groups (00-Global, 01-BaseColor, etc.)
- Individual parameter expose option to either add to the above group or in a "hidden" group.
- More descriptive parameter names
- Changed Texture UVs to not require parameter objects.
- Remove extensions for disabled widgets.
- Fixed dirty on load

diff --git a/Engine/Plugins/Experimental/Avalanche/Source/AvalancheEditor/Private/DetailView/Widgets/SAvaDynamicMaterialWidget.cpp b/Engine/Plugins/Experimental/Avalanche/Source/AvalancheEditor/Private/DetailView/Widgets/SAvaDynamicMaterialWidget.cpp
--- a/Engine/Plugins/Experimental/Avalanche/Source/AvalancheEditor/Private/DetailView/Widgets/SAvaDynamicMaterialWidget.cpp
+++ b/Engine/Plugins/Experimental/Avalanche/Source/AvalancheEditor/Private/DetailView/Widgets/SAvaDynamicMaterialWidget.cpp
@@ -213,7 +213,7 @@ FReply SAvaDynamicMaterialWidget::OpenDynamicMaterialInstanceTab()
 
 	if (Outers.Num() == 0)
 	{
-		MaterialDesignerModule.OpenMaterialModel(Instance->GetMaterialModel(), nullptr, bInvokeTab);
+		MaterialDesignerModule.OpenMaterialModel(Instance->GetMaterialModelBase(), nullptr, bInvokeTab);
 		return FReply::Handled();
 	}
 
@@ -229,7 +229,7 @@ FReply SAvaDynamicMaterialWidget::OpenDynamicMaterialInstanceTab()
 
 		if (!PrimitiveComponent)
 		{
-			MaterialDesignerModule.OpenMaterialModel(Instance->GetMaterialModel(), nullptr, bInvokeTab);
+			MaterialDesignerModule.OpenMaterialModel(Instance->GetMaterialModelBase(), nullptr, bInvokeTab);
 			return FReply::Handled();
 		}
 	}
@@ -245,7 +245,7 @@ FReply SAvaDynamicMaterialWidget::OpenDynamicMaterialInstanceTab()
 		}
 	}
 
-	MaterialDesignerModule.OpenMaterialModel(Instance->GetMaterialModel(), OuterWorld, bInvokeTab);
+	MaterialDesignerModule.OpenMaterialModel(Instance->GetMaterialModelBase(), OuterWorld, bInvokeTab);
 	return FReply::Handled();
 }
 
diff --git a/Engine/Plugins/Experimental/Avalanche/Source/AvalancheEditor/Private/MaterialDesigner/AvaMaterialDesignerExtension.cpp b/Engine/Plugins/Experimental/Avalanche/Source/AvalancheEditor/Private/MaterialDesigner/AvaMaterialDesignerExtension.cpp
--- a/Engine/Plugins/Experimental/Avalanche/Source/AvalancheEditor/Private/MaterialDesigner/AvaMaterialDesignerExtension.cpp
+++ b/Engine/Plugins/Experimental/Avalanche/Source/AvalancheEditor/Private/MaterialDesigner/AvaMaterialDesignerExtension.cpp
@@ -40,7 +40,7 @@ void FAvaMaterialDesignerExtension::ExtendToolbarMenu(UToolMenu& InMenu)
 	Entry.StyleNameOverride = "CalloutToolbar";
 }
 
-bool FAvaMaterialDesignerExtension::IsDynamicMaterialModelValid(UDynamicMaterialModel* InMaterialModel)
+bool FAvaMaterialDesignerExtension::IsDynamicMaterialModelValid(UDynamicMaterialModelBase* InMaterialModel)
 {
 	if (!IsValid(InMaterialModel))
 	{
diff --git a/Engine/Plugins/Experimental/Avalanche/Source/AvalancheEditor/Private/MaterialDesigner/AvaMaterialDesignerExtension.h b/Engine/Plugins/Experimental/Avalanche/Source/AvalancheEditor/Private/MaterialDesigner/AvaMaterialDesignerExtension.h
--- a/Engine/Plugins/Experimental/Avalanche/Source/AvalancheEditor/Private/MaterialDesigner/AvaMaterialDesignerExtension.h
+++ b/Engine/Plugins/Experimental/Avalanche/Source/AvalancheEditor/Private/MaterialDesigner/AvaMaterialDesignerExtension.h
@@ -5,7 +5,7 @@
 #include "IAvaEditorExtension.h"
 
 class UDynamicMaterialInstance;
-class UDynamicMaterialModel;
+class UDynamicMaterialModelBase;
 struct FDMObjectMaterialProperty;
 
 class FAvaMaterialDesignerExtension : public FAvaEditorExtension
@@ -21,7 +21,7 @@ public:
 	//~ End IAvaEditorExtension
 
 private:
-	bool IsDynamicMaterialModelValid(UDynamicMaterialModel* InMaterialModel);
+	bool IsDynamicMaterialModelValid(UDynamicMaterialModelBase* InMaterialModel);
 
 	bool SetDynamicMaterialValue(const FDMObjectMaterialProperty& InObjectMaterialProperty, UDynamicMaterialInstance* InMaterial);
 
diff --git a/Engine/Plugins/Experimental/Avalanche/Source/AvalancheMask/Private/Handling/AvaMaskDesignedMaterialHandle.cpp b/Engine/Plugins/Experimental/Avalanche/Source/AvalancheMask/Private/Handling/AvaMaskDesignedMaterialHandle.cpp
--- a/Engine/Plugins/Experimental/Avalanche/Source/AvalancheMask/Private/Handling/AvaMaskDesignedMaterialHandle.cpp
+++ b/Engine/Plugins/Experimental/Avalanche/Source/AvalancheMask/Private/Handling/AvaMaskDesignedMaterialHandle.cpp
@@ -1,4 +1,4 @@
-﻿// Copyright Epic Games, Inc. All Rights Reserved.
+// Copyright Epic Games, Inc. All Rights Reserved.
 
 #include "Handling/AvaMaskDesignedMaterialHandle.h"
 
@@ -265,7 +265,7 @@ UMaterialInstanceDynamic* FAvaMaskDesignedMaterialHandle::GetMaterialInstance()
 }
 
 #if WITH_EDITOR
-void FAvaMaskDesignedMaterialHandle::OnMaterialBuilt(UDynamicMaterialModel* InMaterialModel)
+void FAvaMaskDesignedMaterialHandle::OnMaterialBuilt(UDynamicMaterialModelBase* InMaterialModel)
 {
 	check(InMaterialModel);
 	
diff --git a/Engine/Plugins/Experimental/Avalanche/Source/AvalancheMask/Private/Handling/AvaMaskDesignedMaterialHandle.h b/Engine/Plugins/Experimental/Avalanche/Source/AvalancheMask/Private/Handling/AvaMaskDesignedMaterialHandle.h
--- a/Engine/Plugins/Experimental/Avalanche/Source/AvalancheMask/Private/Handling/AvaMaskDesignedMaterialHandle.h
+++ b/Engine/Plugins/Experimental/Avalanche/Source/AvalancheMask/Private/Handling/AvaMaskDesignedMaterialHandle.h
@@ -1,4 +1,4 @@
-﻿// Copyright Epic Games, Inc. All Rights Reserved.
+// Copyright Epic Games, Inc. All Rights Reserved.
 
 #pragma once
 
@@ -9,7 +9,7 @@
 #include "AvaMaskDesignedMaterialHandle.generated.h"
 
 #if WITH_EDITOR
-class UDynamicMaterialModel;
+class UDynamicMaterialModelBase;
 #endif
 
 class UDynamicMaterialInstance;
@@ -70,7 +70,7 @@ protected:
 	virtual UMaterialInstanceDynamic* GetMaterialInstance() override;
 
 #if WITH_EDITOR
-	void OnMaterialBuilt(UDynamicMaterialModel* InMaterialModel);
+	void OnMaterialBuilt(UDynamicMaterialModelBase* InMaterialModel);
 	
 	UMaterialFunctionInterface* GetOutputProcessor();
 	void SetOutputProcessor(UMaterialFunctionInterface* InMaterialFunction);
diff --git a/Engine/Plugins/Experimental/DynamicMaterial/Config/BaseDynamicMaterial.ini b/Engine/Plugins/Experimental/DynamicMaterial/Config/BaseDynamicMaterial.ini
--- a/Engine/Plugins/Experimental/DynamicMaterial/Config/BaseDynamicMaterial.ini
+++ b/Engine/Plugins/Experimental/DynamicMaterial/Config/BaseDynamicMaterial.ini
@@ -30,3 +30,7 @@
 +PackageRedirects=(OldName="/Script/DynamicMaterial.DynamicMaterialModel.GlobalPixelDepthOffset",NewName="/Script/DynamicMaterial.DynamicMaterialModel.GlobalPixelDepthOffsetParameter")
 +PackageRedirects=(OldName="/Script/DynamicMaterial.DynamicMaterialModel.GlobalOffset",NewName="/Script/DynamicMaterial.DynamicMaterialModel.GlobalOffsetParameter")
 +PackageRedirects=(OldName="/Script/DynamicMaterial.DynamicMaterialModel.GlobalRotation",NewName="/Script/DynamicMaterial.DynamicMaterialModel.GlobalRotationParameter")
+
++PropertyRedirects=(OldName="/Script/DynamicMaterial.DMDefaultMaterialPropertySlotValue.Type",NewName="/Script/DynamicMaterial.DMDefaultMaterialPropertySlotValue.DefaultType")
++PropertyRedirects=(OldName="/Script/DynamicMaterial.DynamicMaterialInstance.MaterialModel",NewName="/Script/DynamicMaterial.DynamicMaterialInstance.MaterialModelBase")
++FunctionRedirects=(OldName="/Script/DynamicMaterial.DynamicMaterialModel.RemoveValueByName",NewName="/Script/DynamicMaterial.DynamicMaterialModel.RemoveValueByParameterName")
\ No newline at end of file
diff --git a/Engine/Plugins/Experimental/DynamicMaterial/Source/DynamicMaterial/Private/Components/DMMaterialComponent.cpp b/Engine/Plugins/Experimental/DynamicMaterial/Source/DynamicMaterial/Private/Components/DMMaterialComponent.cpp
--- a/Engine/Plugins/Experimental/DynamicMaterial/Source/DynamicMaterial/Private/Components/DMMaterialComponent.cpp
+++ b/Engine/Plugins/Experimental/DynamicMaterial/Source/DynamicMaterial/Private/Components/DMMaterialComponent.cpp
@@ -160,9 +160,9 @@ bool UDMMaterialComponent::CanClean()
 	return (FPlatformTime::Seconds() >= MinCleanTime);
 }
 
-void UDMMaterialComponent::PreventClean(double DelayFor)
+void UDMMaterialComponent::PreventClean(double InDelayFor)
 {
-	MinCleanTime = FMath::Max(MinCleanTime, FPlatformTime::Seconds() + DelayFor);
+	MinCleanTime = FMath::Max(MinCleanTime, FPlatformTime::Seconds() + InDelayFor);
 }
 
 bool UDMMaterialComponent::NeedsClean()
@@ -190,9 +190,9 @@ void UDMMaterialComponent::DoClean()
 	PreventClean(VeryShortTime);
 }
 
-void UDMMaterialComponent::SetComponentState(EDMComponentLifetimeState NewState)
+void UDMMaterialComponent::SetComponentState(EDMComponentLifetimeState InNewState)
 {
-	if (ComponentState == NewState)
+	if (ComponentState == InNewState)
 	{
 		return;
 	}
@@ -202,8 +202,8 @@ void UDMMaterialComponent::SetComponentState(EDMComponentLifetimeState NewState)
 		return;
 	}
 
-	ComponentState = NewState;
-	OnComponentStateChange(NewState);
+	ComponentState = InNewState;
+	OnComponentStateChange(InNewState);
 }
 
 void UDMMaterialComponent::PostLoad()
@@ -232,14 +232,14 @@ bool UDMMaterialComponent::Modify(bool bInAlwaysMarkDirty /*= true*/)
 	return bSaved;
 }
 
-void UDMMaterialComponent::OnComponentStateChange(EDMComponentLifetimeState NewState)
+void UDMMaterialComponent::OnComponentStateChange(EDMComponentLifetimeState InNewState)
 {
 	if (!IsComponentValid())
 	{
 		return;
 	}
 
-	switch (NewState)
+	switch (InNewState)
 	{
 		case EDMComponentLifetimeState::Added:
 			OnComponentAdded();
diff --git a/Engine/Plugins/Experimental/DynamicMaterial/Source/DynamicMaterial/Private/Components/DMMaterialValue.cpp b/Engine/Plugins/Experimental/DynamicMaterial/Source/DynamicMaterial/Private/Components/DMMaterialValue.cpp
--- a/Engine/Plugins/Experimental/DynamicMaterial/Source/DynamicMaterial/Private/Components/DMMaterialValue.cpp
+++ b/Engine/Plugins/Experimental/DynamicMaterial/Source/DynamicMaterial/Private/Components/DMMaterialValue.cpp
@@ -5,16 +5,11 @@
 #include "DMComponentPath.h"
 #include "Model/DynamicMaterialModel.h"
 #include "Materials/Material.h"
+#include "Model/IDMMaterialBuildUtilsInterface.h"
 #include "UObject/Package.h"
 
 #if WITH_EDITOR
-#include "AssetRegistry/AssetRegistryModule.h"
 #include "DMValueDefinition.h"
-#include "DynamicMaterialModule.h"
-#include "Factories/MaterialFactoryNew.h"
-#include "Model/IDMMaterialBuildStateInterface.h"
-#include "Model/IDMMaterialBuildUtilsInterface.h"
-#include "Model/IDynamicMaterialModelEditorOnlyDataInterface.h"
 #include "Modules/ModuleManager.h"
 #include "PropertyEditorModule.h"
 #endif // WITH_EDITOR
@@ -28,17 +23,6 @@ const FString UDMMaterialValue::ParameterPathToken = FString(TEXT("Parameter"));
 #if WITH_EDITOR
 const FName UDMMaterialValue::ValueName = "Value";
 
-UDMMaterialValue* UDMMaterialValue::CreateMaterialValue(UDynamicMaterialModel* InMaterialModel, const FString& InName, 
-	EDMValueType InValueType, bool bInLocal)
-{
-	return CreateMaterialValue(
-		InMaterialModel,
-		InName,
-		UDMValueDefinitionLibrary::GetValueDefinition(InValueType).GetValueClass(),
-		bInLocal
-	);
-}
-
 UDMMaterialValue* UDMMaterialValue::CreateMaterialValue(UDynamicMaterialModel* InMaterialModel, const FString& InName,
 	TSubclassOf<UDMMaterialValue> InValueClass, bool bInLocal)
 {
@@ -47,6 +31,8 @@ UDMMaterialValue* UDMMaterialValue::CreateMaterialValue(UDynamicMaterialModel* I
 	UDMMaterialValue* NewValue = NewObject<UDMMaterialValue>(InMaterialModel, InValueClass, NAME_None, RF_Transactional);
 	check(NewValue);
 	NewValue->bLocal = bInLocal;
+	NewValue->ResetDefaultValue();
+	NewValue->ApplyDefaultValue();
  
 	if (InName.IsEmpty() == false)
 	{
@@ -57,6 +43,7 @@ UDMMaterialValue* UDMMaterialValue::CreateMaterialValue(UDynamicMaterialModel* I
 
 		NewValue->Parameter = InMaterialModel->CreateUniqueParameter(*InName);
 		NewValue->Parameter->SetParentComponent(NewValue);
+		NewValue->CachedParameterName = *InName;
 	}
  
 	return NewValue;
@@ -66,8 +53,17 @@ UDMMaterialValue* UDMMaterialValue::CreateMaterialValue(UDynamicMaterialModel* I
 UDMMaterialValue::UDMMaterialValue()
 	: UDMMaterialValue(EDMValueType::VT_None)
 {
-	Parameter = nullptr;
+}
 
+UDMMaterialValue::UDMMaterialValue(EDMValueType InType)
+	: Type(InType)
+	, bLocal(false)
+	, Parameter(nullptr)
+	, CachedParameterName(NAME_None)
+#if WITH_EDITORONLY_DATA
+	, bExposeParameter(false)
+#endif
+{
 #if WITH_EDITORONLY_DATA
 	EditableProperties.Add(ValueName);
 #endif
@@ -86,11 +82,9 @@ FText UDMMaterialValue::GetTypeName() const
  
 FText UDMMaterialValue::GetDescription() const
 {
-	static const FText Template = LOCTEXT("ValueDescriptionTemplate", "{0} ({1})");
- 
 	return FText::Format(
-		Template,
-		Parameter ? FText::FromName(Parameter->GetParameterName()) : FText::AsNumber(FindIndex() + 1),
+		LOCTEXT("ValueDescriptionTemplate", "{0} ({1})"),
+		FText::FromName(GetMaterialParameterName()),
 		GetTypeName()
 	);
 }
@@ -142,17 +136,21 @@ void UDMMaterialValue::PostEditImport()
 		Parameter->SetParentComponent(this);
 	}
 }
-
-void UDMMaterialValue::BeginDestroy()
-{
-	Super::BeginDestroy();
-	PreviewMaterial = nullptr;
-}
 #endif // WITH_EDITOR
 
 FName UDMMaterialValue::GetMaterialParameterName() const
 {
-	return Parameter ? Parameter->GetParameterName() : GetFName();
+	if (Parameter)
+	{
+		return Parameter->GetParameterName();
+	}
+
+	if (!CachedParameterName.IsNone())
+	{
+		return CachedParameterName;
+	}
+
+	return GetFName();
 }
 
 UDMMaterialComponent* UDMMaterialValue::GetSubComponentByPath(FDMComponentPath& InPath, const FDMComponentPathSegment& InPathSegment) const
@@ -195,11 +193,8 @@ bool UDMMaterialValue::SetParameterName(FName InBaseName)
 			MaterialModel->FreeParameter(Parameter);
 			Parameter = nullptr;
 		}
- 
-		return true;
 	}
- 
-	if (Parameter)
+	else if (Parameter)
 	{
 		Parameter->RenameParameter(InBaseName);
 	}
@@ -208,10 +203,39 @@ bool UDMMaterialValue::SetParameterName(FName InBaseName)
 		Parameter = MaterialModel->CreateUniqueParameter(InBaseName);
 		Parameter->SetParentComponent(this);
 	}
+
+	UpdateCachedParameterName();
  
 	return true;
 }
 
+EDMMaterialParameterGroup UDMMaterialValue::GetParameterGroup() const
+{
+	if (GetMaterialParameterName().ToString().StartsWith("Global"))
+	{
+		return EDMMaterialParameterGroup::Global;
+	}
+
+	if (bExposeParameter)
+	{
+		return EDMMaterialParameterGroup::Property;
+	}
+
+	return EDMMaterialParameterGroup::NotExposed;
+}
+
+void UDMMaterialValue::SetShouldExposeParameter(bool bInExpose)
+{
+	if (bExposeParameter == bInExpose)
+	{
+		return;
+	}
+
+	bExposeParameter = bInExpose;
+
+	Update(EDMUpdateType::Structure);
+}
+
 void UDMMaterialValue::OnComponentAdded()
 {
 	Super::OnComponentAdded();
@@ -222,6 +246,8 @@ void UDMMaterialValue::OnComponentAdded()
 		{
 			MaterialModel->AddRuntimeComponentReference(this);
 		}
+
+		UpdateCachedParameterName();
 	}
 }
 
@@ -243,68 +269,13 @@ void UDMMaterialValue::OnComponentRemoved()
 		{
 			MaterialModel->RemoveRuntimeComponentReference(this);
 		}
-	}
- 
-	Super::OnComponentRemoved();
-}
 
-int32 UDMMaterialValue::FindIndex() const
-{
-	UDynamicMaterialModel* MaterialModel = GetMaterialModel();
-	check(MaterialModel);
- 
-	const TArray<UDMMaterialValue*>& Values = MaterialModel->GetValues();
- 
-	for (int32 ValueIdx = 0; ValueIdx < Values.Num(); ++ValueIdx)
-	{
-		if (Values[ValueIdx] == this)
-		{
-			return ValueIdx;
-		}
+		CachedParameterName = NAME_None;
 	}
  
-	check(false);
-	return INDEX_NONE;
-}
- 
-int32 UDMMaterialValue::FindIndexSafe() const
-{
-	UDynamicMaterialModel* MaterialModel = GetMaterialModel();
-	check(MaterialModel);
- 
-	const TArray<UDMMaterialValue*>& Values = MaterialModel->GetValues();
- 
-	for (int32 ValueIdx = 0; ValueIdx < Values.Num(); ++ValueIdx)
-	{
-		if (Values[ValueIdx] == this)
-		{
-			return ValueIdx;
-		}
-	}
- 
-	return INDEX_NONE;
-}
- 
-UMaterial* UDMMaterialValue::GetPreviewMaterial()
-{
-	if (!PreviewMaterial)
-	{
-		CreatePreviewMaterial();
- 
-		if (PreviewMaterial)
-		{
-			MarkComponentDirty();
-		}
-	}
- 
-	return PreviewMaterial;
+	Super::OnComponentRemoved();
 }
 #endif // WITH_EDITOR
- 
-void UDMMaterialValue::SetMIDParameter(UMaterialInstanceDynamic* InMID) const
-{
-	checkNoEntry();
-}
 
 TSharedPtr<FJsonValue> UDMMaterialValue::JsonSerialize() const
 {
@@ -317,51 +288,38 @@ bool UDMMaterialValue::JsonDeserialize(const TSharedPtr<FJsonValue>& InJsonValue
 }
 
 #if WITH_EDITOR
-void UDMMaterialValue::DoClean()
+void UDMMaterialValue::PostEditorDuplicate(UDynamicMaterialModel* InMaterialModel, UDMMaterialComponent* InParent)
 {
-	if (!IsComponentValid())
+	if (GetOuter() == InMaterialModel)
 	{
+		Super::PostEditorDuplicate(InMaterialModel, InParent);
 		return;
 	}
 
-	UpdatePreviewMaterial();
-
-	Super::DoClean();
-}
+	FName OldParameterName = NAME_None;
 
-void UDMMaterialValue::PostEditorDuplicate(UDynamicMaterialModel* InMaterialModel, UDMMaterialComponent* InParent)
-{
-	if (GetOuter() != InMaterialModel)
+	if (UDMMaterialParameter* Param = GetParameter())
 	{
-		FName OldParameterName = NAME_None;
-
-		if (UDMMaterialParameter* Param = GetParameter())
+		// Reset this to null as it holds a copy of the parameter from the copied-from object.
+		// This will not be in the model's parameter list and will share the same name as the old parameter.
+		// Just null the reference and create a new parameter.
+		if (InMaterialModel->ConditionalFreeParameter(Param))
 		{
-			// Reset this to null as it holds a copy of the parameter from the copied-from object.
-			// This will not be in the model's parameter list and will share the same name as the old parameter.
-			// Just null the reference and create a new parameter.
-			if (InMaterialModel->ConditionalFreeParameter(Param))
-			{
-				OldParameterName = Param->GetParameterName();
-				Parameter = nullptr;
-			}
+			OldParameterName = Param->GetParameterName();
+			Parameter = nullptr;
 		}
+	}
 
-		Super::PostEditorDuplicate(InMaterialModel, InParent);
+	Super::PostEditorDuplicate(InMaterialModel, InParent);
 
-		Rename(nullptr, InMaterialModel, UE::DynamicMaterial::RenameFlags);
+	Rename(nullptr, InMaterialModel, UE::DynamicMaterial::RenameFlags);
 
-		if (OldParameterName != NAME_None)
-		{
-			SetParameterName(OldParameterName);
-		}
-	}
-	else
+	if (!OldParameterName.IsNone())
 	{
-		Super::PostEditorDuplicate(InMaterialModel, InParent);
+		SetParameterName(OldParameterName);
 	}
 
-	PreviewMaterial = nullptr;
+	UpdateCachedParameterName();
 }
 
 bool UDMMaterialValue::Modify(bool bInAlwaysMarkDirty)
@@ -390,18 +348,18 @@ void UDMMaterialValue::PostEditUndo()
 	OnValueChanged(EDMUpdateType::Structure | EDMUpdateType::AllowParentUpdate); // Just in case - Undos are not meant to be quick and easy.
 }
  
-void UDMMaterialValue::PostEditChangeProperty(FPropertyChangedEvent& PropertyChangedEvent)
+void UDMMaterialValue::PostEditChangeProperty(FPropertyChangedEvent& InPropertyChangedEvent)
 {
-	Super::PostEditChangeProperty(PropertyChangedEvent);
+	Super::PostEditChangeProperty(InPropertyChangedEvent);
 
 	if (!IsComponentValid())
 	{
 		return;
 	}
 
-	FName MemberPropertyName = PropertyChangedEvent.GetMemberPropertyName();
+	const FName MemberPropertyName = InPropertyChangedEvent.GetMemberPropertyName();
 
-	if (MemberPropertyName == NAME_None)
+	if (MemberPropertyName.IsNone())
 	{
 		return;
 	}
@@ -419,19 +377,6 @@ void UDMMaterialValue::PostEditChangeProperty(FPropertyChangedEvent& PropertyCha
 	}
 }
 #endif // WITH_EDITOR
-
-UDMMaterialValue::UDMMaterialValue(EDMValueType InType)
-	: Type(InType)
-	, bLocal(false)
-	, Parameter(nullptr)
-#if WITH_EDITORONLY_DATA
-	, PreviewMaterial(nullptr)
-#endif
-{
-#if WITH_EDITORONLY_DATA
-	EditableProperties.Add(ValueName);
-#endif
-}
  
 void UDMMaterialValue::OnValueChanged(EDMUpdateType InUpdateType)
 {
@@ -449,123 +394,61 @@ void UDMMaterialValue::OnValueChanged(EDMUpdateType InUpdateType)
 	}
 #endif
 }
- 
-void UDMMaterialValue::Update(EDMUpdateType InUpdateType)
-{
-	if (!FDMUpdateGuard::CanUpdate())
-	{
-		return;
-	}
-
-	if (!IsComponentValid())
-	{
-		return;
-	}
 
 #if WITH_EDITOR
-	if (HasComponentBeenRemoved())
-	{
-		return;
-	}
-
-	MarkComponentDirty();
-#endif
-
-	Super::Update(InUpdateType);
-
-	if (UDynamicMaterialModel* MaterialModel = GetMaterialModel())
-	{
-		MaterialModel->OnValueUpdated(this, InUpdateType);
-	}
+FName UDMMaterialValue::GenerateAutomaticParameterName() const
+{
+	return *GetComponentPath();
 }
 
-#if WITH_EDITOR
-void UDMMaterialValue::CreatePreviewMaterial()
+void UDMMaterialValue::UpdateCachedParameterName()
 {
-	if (!IsComponentValid())
-	{
-		return;
-	}
-
-	if (FDynamicMaterialModule::IsMaterialExportEnabled() == false)
+	if (Parameter)
 	{
-		UMaterialFactoryNew* MaterialFactory = NewObject<UMaterialFactoryNew>();
-		check(MaterialFactory);
- 
-		PreviewMaterial = (UMaterial*)MaterialFactory->FactoryCreateNew(
-			UMaterial::StaticClass(),
-			GetTransientPackage(),
-			NAME_None,
-			RF_Transient,
-			nullptr,
-			GWarn
-		);
- 
-		PreviewMaterial->bIsPreviewMaterial = true;
+		CachedParameterName = Parameter->GetParameterName();
 	}
 	else
 	{
-		FString MaterialBaseName = GetName() + "-" + FGuid::NewGuid().ToString();
-		const FString FullName = "/Game/DynamicMaterials/" + MaterialBaseName;
-		UPackage* Package = CreatePackage(*FullName);
- 
-		UMaterialFactoryNew* MaterialFactory = NewObject<UMaterialFactoryNew>();
-		check(MaterialFactory);
- 
-		PreviewMaterial = Cast<UMaterial>(MaterialFactory->FactoryCreateNew(
-			UMaterial::StaticClass(),
-			Package,
-			*MaterialBaseName,
-			RF_Standalone | RF_Public,
-			nullptr,
-			GWarn
-		));
- 
-		FAssetRegistryModule::AssetCreated(PreviewMaterial);
+		CachedParameterName = GenerateAutomaticParameterName();
 	}
 }
- 
-void UDMMaterialValue::UpdatePreviewMaterial()
+#endif
+
+void UDMMaterialValue::Update(EDMUpdateType InUpdateType)
 {
-	if (!IsComponentValid())
+	if (!FDMUpdateGuard::CanUpdate())
 	{
 		return;
 	}
 
-	if (!PreviewMaterial)
+	if (!IsComponentValid())
 	{
-		CreatePreviewMaterial();
- 
-		if (!PreviewMaterial)
-		{
-			return;
-		}
+		return;
 	}
- 
-	UDynamicMaterialModel* MaterialModel = GetMaterialModel();
 
-	if (!MaterialModel)
+#if WITH_EDITOR
+	if (HasComponentBeenRemoved())
 	{
 		return;
 	}
 
-	IDynamicMaterialModelEditorOnlyDataInterface* ModelEditorOnlyData = MaterialModel->GetEditorOnlyData();
+	MarkComponentDirty();
 
-	if (!ModelEditorOnlyData)
+	if (InUpdateType == EDMUpdateType::Structure)
 	{
-		return;
+		UpdateCachedParameterName();
 	}
+#endif
 
-	TSharedRef<IDMMaterialBuildStateInterface> BuildState = ModelEditorOnlyData->CreateBuildStateInterface(PreviewMaterial);
-
-	UE_LOG(LogDynamicMaterial, Display, TEXT("Building Material Designer Value Preview (%s)..."), *GetName());
+	Super::Update(InUpdateType);
 
-	GenerateExpression(BuildState);
-	UMaterialExpression* ValueExpression = BuildState->GetLastValueExpression(this);
- 
-	BuildState->GetBuildUtils().UpdatePreviewMaterial(ValueExpression, 0, FDMMaterialStageConnectorChannel::WHOLE_CHANNEL, 32);
+	if (UDynamicMaterialModel* MaterialModel = GetMaterialModel())
+	{
+		MaterialModel->OnValueUpdated(this, InUpdateType);
+	}
 }
 
+#if WITH_EDITOR
 int32 UDMMaterialValue::GetInnateMaskOutput(int32 OutputChannels) const
 {
 	return INDEX_NONE;
diff --git a/Engine/Plugins/Experimental/DynamicMaterial/Source/DynamicMaterial/Private/Components/DMTextureUV.cpp b/Engine/Plugins/Experimental/DynamicMaterial/Source/DynamicMaterial/Private/Components/DMTextureUV.cpp
--- a/Engine/Plugins/Experimental/DynamicMaterial/Source/DynamicMaterial/Private/Components/DMTextureUV.cpp
+++ b/Engine/Plugins/Experimental/DynamicMaterial/Source/DynamicMaterial/Private/Components/DMTextureUV.cpp
@@ -2,21 +2,26 @@
 
 #include "Components/DMTextureUV.h"
 #include "Components/DMMaterialParameter.h"
+#include "Components/DMTextureUVDynamic.h"
 #include "DMComponentPath.h"
 #include "Materials/MaterialInstanceDynamic.h"
 #include "Model/DynamicMaterialModel.h"
 #include "Serialization/CustomVersion.h"
 
+#if WITH_EDITOR
+#include "Model/IDMMaterialBuildUtilsInterface.h"
+#endif
+
 namespace UE::DynamicMaterial::Private
 {
 	TMap<int32, FName> BaseParameterNames = {
-		{ParamID::PivotX,   FName(TEXT("TextureUV_PivotX"))},
-		{ParamID::PivotY,   FName(TEXT("TextureUV_PivotY"))},
-		{ParamID::TilingX,   FName(TEXT("TextureUV_TilingX"))},
-		{ParamID::TilingY,   FName(TEXT("TextureUV_TilingY"))},
-		{ParamID::Rotation, FName(TEXT("TextureUV_Rotation"))},
-		{ParamID::OffsetX,  FName(TEXT("TextureUV_OffsetX"))},
-		{ParamID::OffsetY,  FName(TEXT("TextureUV_OffsetY"))},
+		{ParamID::PivotX,   FName(TEXT("Pivot.X"))},
+		{ParamID::PivotY,   FName(TEXT("Pivot.Y"))},
+		{ParamID::TilingX,  FName(TEXT("Tiling.X"))},
+		{ParamID::TilingY,  FName(TEXT("Tiling.Y"))},
+		{ParamID::Rotation, FName(TEXT("Rotation"))},
+		{ParamID::OffsetX,  FName(TEXT("Offset.X"))},
+		{ParamID::OffsetY,  FName(TEXT("Offset.Y"))},
 	};
 }
 
@@ -42,10 +47,10 @@ const FString UDMTextureUV::TilingYPathToken   = FString(TEXT("TilingY"));
 const FName UDMTextureUV::NAME_Offset     = GET_MEMBER_NAME_CHECKED(UDMTextureUV, Offset);
 const FName UDMTextureUV::NAME_Pivot      = GET_MEMBER_NAME_CHECKED(UDMTextureUV, Pivot);
 const FName UDMTextureUV::NAME_Rotation   = GET_MEMBER_NAME_CHECKED(UDMTextureUV, Rotation);
-const FName UDMTextureUV::NAME_Tiling      = GET_MEMBER_NAME_CHECKED(UDMTextureUV, Tiling);
+const FName UDMTextureUV::NAME_Tiling     = GET_MEMBER_NAME_CHECKED(UDMTextureUV, Tiling);
 
 #if WITH_EDITOR
-const FName UDMTextureUV::NAME_UVSource = GET_MEMBER_NAME_CHECKED(UDMTextureUV, UVSource);
+const FName UDMTextureUV::NAME_UVSource   = GET_MEMBER_NAME_CHECKED(UDMTextureUV, UVSource);
 const FName UDMTextureUV::NAME_bMirrorOnX = GET_MEMBER_NAME_CHECKED(UDMTextureUV, bMirrorOnX);
 const FName UDMTextureUV::NAME_bMirrorOnY = GET_MEMBER_NAME_CHECKED(UDMTextureUV, bMirrorOnY);
 
@@ -54,7 +59,7 @@ const TMap<FName, bool> UDMTextureUV::TextureProperties = {
 	{NAME_Offset,     true},
 	{NAME_Pivot,      true},
 	{NAME_Rotation,   true},
-	{NAME_Tiling,      true},
+	{NAME_Tiling,     true},
 	{NAME_bMirrorOnX, false},
 	{NAME_bMirrorOnY, false}
 };
@@ -72,7 +77,7 @@ UDMTextureUV::UDMTextureUV()
 #endif
 }
 
-#if WITH_EDITOR
+#if WITH_EDITORONLY_DATA
 void UDMTextureUV::SetUVSource(EDMUVSource InUVSource)
 {
 	if (!IsComponentValid())
@@ -98,8 +103,7 @@ void UDMTextureUV::SetOffset(const FVector2D& InOffset)
 		return;
 	}
 
-	if (FMath::IsNearlyEqual(Offset.X, InOffset.X)
-		&& FMath::IsNearlyEqual(Offset.Y, InOffset.Y))
+	if (Offset.Equals(InOffset))
 	{
 		return;
 	}
@@ -116,8 +120,7 @@ void UDMTextureUV::SetPivot(const FVector2D& InPivot)
 		return;
 	}
 
-	if (FMath::IsNearlyEqual(Pivot.X, InPivot.X)
-		&& FMath::IsNearlyEqual(Pivot.Y, InPivot.Y))
+	if (Pivot.Equals(InPivot))
 	{
 		return;
 	}
@@ -151,8 +154,7 @@ void UDMTextureUV::SetTiling(const FVector2D& InTiling)
 		return;
 	}
 
-	if (FMath::IsNearlyEqual(Tiling.X, InTiling.X)
-		&& FMath::IsNearlyEqual(Tiling.Y, InTiling.Y))
+	if (Tiling.Equals(InTiling))
 	{
 		return;
 	}
@@ -162,7 +164,7 @@ void UDMTextureUV::SetTiling(const FVector2D& InTiling)
 	OnTextureUVChanged(EDMUpdateType::Value);
 }
 
-#if WITH_EDITOR
+#if WITH_EDITORONLY_DATA
 void UDMTextureUV::SetMirrorOnX(bool bInMirrorOnX)
 {
 	if (!IsComponentValid())
@@ -213,116 +215,189 @@ TArray<UDMMaterialParameter*> UDMTextureUV::GetParameters() const
 
 UDMMaterialParameter* UDMTextureUV::GetMaterialParameter(FName InPropertyName, int32 InComponent) const
 {
-	using namespace UE::DynamicMaterial;
+	if (const TObjectPtr<UDMMaterialParameter>* ParameterPtr = MaterialParameters.Find(PropertyComponentToParamId(InPropertyName, InComponent)))
+	{
+		return ParameterPtr->Get();
+	}
 
-	const TObjectPtr<UDMMaterialParameter>* ParameterPtr = nullptr;
+	return nullptr;
+}
 
-	if (InPropertyName == NAME_Offset)
+FName UDMTextureUV::GetMaterialParameterName(FName InPropertyName, int32 InComponent) const
+{
+	if (UDMMaterialParameter* Parameter = GetMaterialParameter(InPropertyName, InComponent))
 	{
-		switch (InComponent)
-		{
-			case 0:
-				ParameterPtr = MaterialParameters.Find(ParamID::OffsetX);
-				break;
+		return Parameter->GetParameterName();
+	}
 
-			case 1:
-				ParameterPtr = MaterialParameters.Find(ParamID::OffsetY);
-				break;
-		}
+	if (const FName* CachedNamePtr = CachedParameterNames.Find(PropertyComponentToParamId(InPropertyName, InComponent)))
+	{
+		return *CachedNamePtr;
 	}
-	else if (InPropertyName == NAME_Pivot)
+
+	using namespace UE::DynamicMaterial::Private;
+
+	const int32 ParamId = PropertyComponentToParamId(InPropertyName, InComponent);
+
+	if (const FName* NamePtr = BaseParameterNames.Find(ParamId))
 	{
-		switch (InComponent)
-		{
-			case 0:
-				ParameterPtr = MaterialParameters.Find(ParamID::PivotX);
-				break;
+		return *NamePtr;
+	}
 
-			case 1:
-				ParameterPtr = MaterialParameters.Find(ParamID::PivotY);
-				break;
-		}
+	return TEXT("Error");
+}
+
+#if WITH_EDITOR
+bool UDMTextureUV::SetMaterialParameterName(FName InPropertyName, int32 InComponent, FName InNewName)
+{
+	if (!IsComponentValid())
+	{
+		return false;
 	}
-	else if (InPropertyName == NAME_Rotation)
+
+	const int32 ParamId = PropertyComponentToParamId(InPropertyName, InComponent);
+
+	if (ParamId == UE::DynamicMaterial::ParamID::Invalid)
 	{
-		switch (InComponent)
-		{
-			case 0:
-				ParameterPtr = MaterialParameters.Find(ParamID::Rotation);
-				break;
-		}
+		return false;
 	}
-	else if (InPropertyName == NAME_Tiling)
+
+	UDMMaterialParameter* Parameter = GetMaterialParameter(InPropertyName, InComponent);
+
+	if (Parameter && Parameter->GetParameterName() == InNewName)
 	{
-		switch (InComponent)
-		{
-			case 0:
-				ParameterPtr = MaterialParameters.Find(ParamID::TilingX);
-				break;
+		return false;
+	}
 
-			case 1:
-				ParameterPtr = MaterialParameters.Find(ParamID::TilingY);
-				break;
-		}
+	UDynamicMaterialModel* MaterialModel = GetMaterialModel();
+	check(MaterialModel);
+
+	if (GUndo && IsValid(Parameter))
+	{
+		Parameter->Modify();
+		MaterialModel->Modify();
 	}
 
-	if (ParameterPtr)
+	if (InNewName.IsNone())
 	{
-		return (*ParameterPtr).Get();
+		if (Parameter)
+		{
+			Parameter->SetParentComponent(nullptr);
+			MaterialModel->FreeParameter(Parameter);
+			Parameter = nullptr;
+			MaterialParameters.Remove(ParamId);
+		}
+	}
+	else if (Parameter)
+	{
+		Parameter->RenameParameter(InNewName);
+	}
+	else
+	{
+		Parameter = MaterialModel->CreateUniqueParameter(InNewName);
+		Parameter->SetParentComponent(this);
+		MaterialParameters.Add(ParamId, Parameter);
 	}
 
-	return nullptr;
+	UpdateCachedParameterName(InPropertyName, InComponent);
+
+	return true;
 }
 
-FName UDMTextureUV::GetMaterialParameterName(FName InPropertyName, int32 InComponent) const
+EDMMaterialParameterGroup UDMTextureUV::GetParameterGroup(FName InPropertyName, int32 InComponent) const
 {
-	if (UDMMaterialParameter* Parameter = GetMaterialParameter(InPropertyName, InComponent))
-	{
-		return Parameter->GetParameterName();
-	}
+	return GetShouldExposeParameter(InPropertyName, InComponent)
+		? EDMMaterialParameterGroup::Property
+		: EDMMaterialParameterGroup::NotExposed;
+}
 
-	return NAME_None;
+bool UDMTextureUV::GetShouldExposeParameter(FName InPropertyName, int32 InComponent) const
+{
+	return ExposedParameters.Contains(PropertyComponentToParamId(InPropertyName, InComponent));
 }
 
-#if WITH_EDITOR
-bool UDMTextureUV::SetMaterialParameterName(FName InPropertyName, int32 InComponent, FName InNewName)
+void UDMTextureUV::SetShouldExposeParameter(FName InPropertyName, int32 InComponent, bool bInExpose)
 {
-	if (UDMMaterialParameter* Parameter = GetMaterialParameter(InPropertyName, InComponent))
+	const int32 ParamId = PropertyComponentToParamId(InPropertyName, InComponent);
+
+	if (bInExpose)
 	{
-		const FName CurrentName = Parameter->GetParameterName();
-		Parameter->RenameParameter(InNewName);
-		return Parameter->GetParameterName() != CurrentName;
+		ExposedParameters.Add(ParamId);
+	}
+	else
+	{
+		ExposedParameters.Remove(ParamId);
 	}
 
-	return false;
+	Update(EDMUpdateType::Structure);
 }
 
 void UDMTextureUV::PostEditorDuplicate(UDynamicMaterialModel* InMaterialModel, UDMMaterialComponent* InParent)
 {
-	Super::PostEditorDuplicate(InMaterialModel, InParent);
-
-	if (GetOuter() != InMaterialModel)
+	if (GetOuter() == InMaterialModel)
 	{
-		Rename(nullptr, InMaterialModel, UE::DynamicMaterial::RenameFlags);
+		Super::PostEditorDuplicate(InMaterialModel, InParent);
+		return;
 	}
 
-	// Reset this map as it holds copies of the parameters from the copied-from object.
-	// They will not be in this model's parameter list and will share the same name as the old parameters.
-	// Just empty the list and create new parameters.
-	for (TMap<int32, TObjectPtr<UDMMaterialParameter>>::TIterator It(MaterialParameters); It; ++It)
+	TMap<int32, FName> OldParameterNames;
+	OldParameterNames.Reserve(MaterialParameters.Num());
+
+	// Reset these to null as they holds a copy of the parameter from the copied-from object.
+	// This will not be in the model's parameter list and will share the same name as the old parameter.
+	// Just null the reference and create a new parameter.
+	for (const TPair<int32, TObjectPtr<UDMMaterialParameter>>& ParameterPair : MaterialParameters)
 	{
-		UDMMaterialParameter* Parameter = It->Value.Get();
+		OldParameterNames.Add(ParameterPair.Key, ParameterPair.Value->GetParameterName());
+	}
+
+	MaterialParameters.Empty();
+	CachedParameterNames.Empty();
+
+	Super::PostEditorDuplicate(InMaterialModel, InParent);
+
+	Rename(nullptr, InMaterialModel, UE::DynamicMaterial::RenameFlags);
 
-		if (!Parameter || InMaterialModel->ConditionalFreeParameter(Parameter))
+	using namespace UE::DynamicMaterial;
+
+	for (const TPair<int32, FName>& OldParameterPair : OldParameterNames)
+	{
+		if (!OldParameterPair.Value.IsNone())
 		{
-			It.RemoveCurrent();
+			switch (OldParameterPair.Key)
+			{
+				case ParamID::OffsetX:
+					SetMaterialParameterName(UDMTextureUV::NAME_Offset, 0, OldParameterPair.Value);
+					break;
+
+				case ParamID::OffsetY:
+					SetMaterialParameterName(UDMTextureUV::NAME_Offset, 1, OldParameterPair.Value);
+					break;
+
+				case ParamID::Rotation:
+					SetMaterialParameterName(UDMTextureUV::NAME_Rotation, 0, OldParameterPair.Value);
+					break;
+
+				case ParamID::PivotX:
+					SetMaterialParameterName(UDMTextureUV::NAME_Pivot, 0, OldParameterPair.Value);
+					break;
+
+				case ParamID::PivotY:
+					SetMaterialParameterName(UDMTextureUV::NAME_Pivot, 1, OldParameterPair.Value);
+					break;
+
+				case ParamID::TilingX:
+					SetMaterialParameterName(UDMTextureUV::NAME_Tiling, 0, OldParameterPair.Value);
+					break;
+
+				case ParamID::TilingY:
+					SetMaterialParameterName(UDMTextureUV::NAME_Tiling, 1, OldParameterPair.Value);
+					break;
+			}
 		}
 	}
 
-	MaterialParameters.Empty();
-
-	// Create new parameters.
-	CreateParameterNames();
+	UpdateCachedParameterNames();
 }
 #endif
 
@@ -334,9 +409,6 @@ void UDMTextureUV::SetMIDParameters(UMaterialInstanceDynamic* InMID)
 	}
 
 	check(InMID);
-	check(MaterialParameters.IsEmpty() == false);
-
-	using namespace UE::DynamicMaterial;
 
 	auto UpdateMID = [InMID](FName InParamName, float InValue)
 	{
@@ -346,16 +418,21 @@ void UDMTextureUV::SetMIDParameters(UMaterialInstanceDynamic* InMID)
 		}
 	};
 
-	UpdateMID(MaterialParameters[ParamID::PivotX]->GetParameterName(), GetPivot().X);
-	UpdateMID(MaterialParameters[ParamID::PivotY]->GetParameterName(), GetPivot().Y);
-	UpdateMID(MaterialParameters[ParamID::TilingX]->GetParameterName(), GetTiling().X);
-	UpdateMID(MaterialParameters[ParamID::TilingY]->GetParameterName(), GetTiling().Y);
-	UpdateMID(MaterialParameters[ParamID::Rotation]->GetParameterName(), GetRotation());
-	UpdateMID(MaterialParameters[ParamID::OffsetX]->GetParameterName(), GetOffset().X);
-	UpdateMID(MaterialParameters[ParamID::OffsetY]->GetParameterName(), GetOffset().Y);
+	UpdateMID(GetMaterialParameterName(UDMTextureUV::NAME_Offset,   0), GetOffset().X);
+	UpdateMID(GetMaterialParameterName(UDMTextureUV::NAME_Offset,   1), GetOffset().Y);
+	UpdateMID(GetMaterialParameterName(UDMTextureUV::NAME_Rotation, 0), GetRotation());
+	UpdateMID(GetMaterialParameterName(UDMTextureUV::NAME_Pivot,    0), GetPivot().X);
+	UpdateMID(GetMaterialParameterName(UDMTextureUV::NAME_Pivot,    1), GetPivot().Y);
+	UpdateMID(GetMaterialParameterName(UDMTextureUV::NAME_Tiling,   0), GetTiling().X);
+	UpdateMID(GetMaterialParameterName(UDMTextureUV::NAME_Tiling,   1), GetTiling().Y);
 }
 
 #if WITH_EDITOR
+UDMTextureUVDynamic* UDMTextureUV::ToDynamic(UDynamicMaterialModelDynamic* InMaterialModelDynamic)
+{
+	return UDMTextureUVDynamic::CreateTextureUVDynamic(InMaterialModelDynamic, this);
+}
+
 bool UDMTextureUV::Modify(bool bInAlwaysMarkDirty)
 {
 	const bool bSaved = Super::Modify(bInAlwaysMarkDirty);
@@ -388,6 +465,11 @@ void UDMTextureUV::Update(EDMUpdateType InUpdateType)
 	}
 
 	MarkComponentDirty();
+
+	if (InUpdateType == EDMUpdateType::Structure)
+	{
+		UpdateCachedParameterNames();
+	}
 #endif
 
 	Super::Update(InUpdateType);
@@ -398,6 +480,55 @@ void UDMTextureUV::Update(EDMUpdateType InUpdateType)
 	}
 }
 
+int32 UDMTextureUV::PropertyComponentToParamId(FName InPropertyName, int32 InComponent)
+{
+	using namespace UE::DynamicMaterial;
+
+	if (InPropertyName == NAME_Offset)
+	{
+		switch (InComponent)
+		{
+			case 0:
+				return ParamID::OffsetX;
+
+			case 1:
+				return ParamID::OffsetY;
+		}
+	}
+	else if (InPropertyName == NAME_Pivot)
+	{
+		switch (InComponent)
+		{
+			case 0:
+				return ParamID::PivotX;
+
+			case 1:
+				return ParamID::PivotY;
+		}
+	}
+	else if (InPropertyName == NAME_Rotation)
+	{
+		switch (InComponent)
+		{
+			case 0:
+				return ParamID::Rotation;
+		}
+	}
+	else if (InPropertyName == NAME_Tiling)
+	{
+		switch (InComponent)
+		{
+			case 0:
+				return ParamID::TilingX;
+
+			case 1:
+				return ParamID::TilingY;
+		}
+	}
+
+	return ParamID::Invalid;
+}
+
 UDynamicMaterialModel* UDMTextureUV::GetMaterialModel() const
 {
 	return Cast<UDynamicMaterialModel>(GetOuterSafe());
@@ -463,11 +594,13 @@ void UDMTextureUV::GetComponentPathInternal(TArray<FString>& OutChildComponentPa
 	// Replace parameter object names with the base parameter name
 	if (OutChildComponentPathComponents.IsEmpty() == false)
 	{
+		FString& LastPathComponent = OutChildComponentPathComponents.Last();
+
 		for (const TPair<int32, TObjectPtr<UDMMaterialParameter>>& MaterialParameter : MaterialParameters)
 		{
-			if (OutChildComponentPathComponents.Last() == MaterialParameter.Value->GetComponentPathComponent())
+			if (LastPathComponent == MaterialParameter.Value->GetComponentPathComponent())
 			{
-				OutChildComponentPathComponents.Last() = BaseParameterNames[MaterialParameter.Key].ToString();
+				LastPathComponent = BaseParameterNames[MaterialParameter.Key].ToString();
 				break;
 			}
 		}
@@ -476,38 +609,6 @@ void UDMTextureUV::GetComponentPathInternal(TArray<FString>& OutChildComponentPa
 	Super::GetComponentPathInternal(OutChildComponentPathComponents);
 }
 
-void UDMTextureUV::CreateParameterNames()
-{
-	if (!IsComponentValid())
-	{
-		return;
-	}
-
-	UDynamicMaterialModel* MaterialModel = GetMaterialModel();
-	check(MaterialModel);
-
-	using namespace UE::DynamicMaterial;
-
-	auto CreateParam = [this, MaterialModel](int32 ParamId)
-	{
-		if (!MaterialParameters.Contains(ParamId))
-		{
-			using namespace UE::DynamicMaterial::Private;
-
-			MaterialParameters.Add(ParamId, MaterialModel->CreateUniqueParameter(BaseParameterNames[ParamId]));
-			MaterialParameters[ParamId]->SetParentComponent(this);
-		}
-	};
-
-	CreateParam(ParamID::PivotX);
-	CreateParam(ParamID::PivotY);
-	CreateParam(ParamID::TilingX);
-	CreateParam(ParamID::TilingY);
-	CreateParam(ParamID::Rotation);
-	CreateParam(ParamID::OffsetX);
-	CreateParam(ParamID::OffsetY);
-}
-
 void UDMTextureUV::RemoveParameterNames()
 {
 	if (!IsComponentValid())
@@ -515,6 +616,8 @@ void UDMTextureUV::RemoveParameterNames()
 		return;
 	}
 
+	CachedParameterNames.Empty();
+
 	UDynamicMaterialModel* MaterialModel = GetMaterialModel();
 	check(MaterialModel);
 
@@ -553,6 +656,53 @@ void UDMTextureUV::OnTextureUVChanged(EDMUpdateType InUpdateType)
 }
 
 #if WITH_EDITOR
+FName UDMTextureUV::GenerateAutomaticPathComponent(FName InPropertyName, int32 InComponent) const
+{
+	using namespace UE::DynamicMaterial::Private;
+
+	if (const FName* NamePtr = BaseParameterNames.Find(PropertyComponentToParamId(InPropertyName, InComponent)))
+	{
+		return *NamePtr;
+	}
+
+	return TEXT("Error");
+}
+
+FName UDMTextureUV::GenerateAutomaticParameterName(FName InPropertyName, int32 InComponent) const
+{
+	return *(GetComponentPath() + TEXT(".") + GenerateAutomaticPathComponent(InPropertyName, InComponent).ToString());
+}
+
+void UDMTextureUV::UpdateCachedParameterName(FName InPropertyName, int32 InComponent)
+{
+	const int32 ParamId = PropertyComponentToParamId(InPropertyName, InComponent);
+
+	if (ParamId == UE::DynamicMaterial::ParamID::Invalid)
+	{
+		return;
+	}
+
+	if (const TObjectPtr<UDMMaterialParameter>* ParameterPtr = MaterialParameters.Find(ParamId))
+	{
+		CachedParameterNames.FindOrAdd(ParamId) = (*ParameterPtr)->GetParameterName();
+	}
+	else
+	{
+		CachedParameterNames.FindOrAdd(ParamId) = GenerateAutomaticParameterName(InPropertyName, InComponent);
+	}
+}
+
+void UDMTextureUV::UpdateCachedParameterNames()
+{
+	UpdateCachedParameterName(UDMTextureUV::NAME_Offset, 0);
+	UpdateCachedParameterName(UDMTextureUV::NAME_Offset, 1);
+	UpdateCachedParameterName(UDMTextureUV::NAME_Rotation, 0);
+	UpdateCachedParameterName(UDMTextureUV::NAME_Pivot, 0);
+	UpdateCachedParameterName(UDMTextureUV::NAME_Pivot, 1);
+	UpdateCachedParameterName(UDMTextureUV::NAME_Tiling, 0);
+	UpdateCachedParameterName(UDMTextureUV::NAME_Tiling, 1);
+}
+
 void UDMTextureUV::OnComponentAdded()
 {
 	if (!IsComponentValid())
@@ -560,7 +710,7 @@ void UDMTextureUV::OnComponentAdded()
 		return;
 	}
 
-	CreateParameterNames();
+	UpdateCachedParameterNames();
 
 	if (UDynamicMaterialModel* MaterialModel = GetMaterialModel())
 	{
@@ -582,30 +732,30 @@ void UDMTextureUV::OnComponentRemoved()
 	Super::OnComponentRemoved();
 }
 
-void UDMTextureUV::PostEditChangeProperty(struct FPropertyChangedEvent& PropertyChangedEvent)
+void UDMTextureUV::PostEditChangeProperty(FPropertyChangedEvent& InPropertyChangedEvent)
 {
-	Super::PostEditChangeProperty(PropertyChangedEvent);
+	Super::PostEditChangeProperty(InPropertyChangedEvent);
 
 	if (!IsComponentValid())
 	{
 		return;
 	}
 
-	if (!PropertyChangedEvent.MemberProperty)
+	if (!InPropertyChangedEvent.MemberProperty)
 	{
 		return;
 	}
 
-	if (PropertyChangedEvent.MemberProperty->GetFName() == NAME_Offset
-		|| PropertyChangedEvent.MemberProperty->GetFName() == NAME_Pivot
-		|| PropertyChangedEvent.MemberProperty->GetFName() == NAME_Rotation
-		|| PropertyChangedEvent.MemberProperty->GetFName() == NAME_Tiling)
+	if (InPropertyChangedEvent.MemberProperty->GetFName() == NAME_Offset
+		|| InPropertyChangedEvent.MemberProperty->GetFName() == NAME_Pivot
+		|| InPropertyChangedEvent.MemberProperty->GetFName() == NAME_Rotation
+		|| InPropertyChangedEvent.MemberProperty->GetFName() == NAME_Tiling)
 	{
 		OnTextureUVChanged(EDMUpdateType::Value | EDMUpdateType::AllowParentUpdate);
 	}
-	else if (PropertyChangedEvent.MemberProperty->GetFName() == NAME_UVSource
-		|| PropertyChangedEvent.MemberProperty->GetFName() == NAME_bMirrorOnX
-		|| PropertyChangedEvent.MemberProperty->GetFName() == NAME_bMirrorOnY)
+	else if (InPropertyChangedEvent.MemberProperty->GetFName() == NAME_UVSource
+		|| InPropertyChangedEvent.MemberProperty->GetFName() == NAME_bMirrorOnX
+		|| InPropertyChangedEvent.MemberProperty->GetFName() == NAME_bMirrorOnY)
 	{
 		OnTextureUVChanged(EDMUpdateType::Structure | EDMUpdateType::AllowParentUpdate);
 	}
@@ -650,10 +800,7 @@ void UDMTextureUV::PostLoad()
 		MaterialModel->AddRuntimeComponentReference(this);
 	}
 
-	if (MaterialParameters.IsEmpty())
-	{
-		CreateParameterNames();
-	}
+	UpdateCachedParameterNames();
 
 	for (const TPair<int32, TObjectPtr<UDMMaterialParameter>>& Pair : MaterialParameters)
 	{
@@ -688,10 +835,7 @@ void UDMTextureUV::PostEditImport()
 		return;
 	}
 
-	if (MaterialParameters.IsEmpty())
-	{
-		CreateParameterNames();
-	}
+	UpdateCachedParameterNames();
 
 	for (const TPair<int32, TObjectPtr<UDMMaterialParameter>>& Pair : MaterialParameters)
 	{
@@ -706,6 +850,11 @@ UDMTextureUV* UDMTextureUV::CreateTextureUV(UObject* InOuter)
 
 	return NewTextureUV;
 }
+
+FString UDMTextureUV::GetComponentPathComponent() const
+{
+	return TEXT("UV");
+}
 #endif
 
 void UDMTextureUV::Serialize(FArchive& Ar)
diff --git a/Engine/Plugins/Experimental/DynamicMaterial/Source/DynamicMaterial/Private/Components/MaterialValues/DMMaterialValueBool.cpp b/Engine/Plugins/Experimental/DynamicMaterial/Source/DynamicMaterial/Private/Components/MaterialValues/DMMaterialValueBool.cpp
--- a/Engine/Plugins/Experimental/DynamicMaterial/Source/DynamicMaterial/Private/Components/MaterialValues/DMMaterialValueBool.cpp
+++ b/Engine/Plugins/Experimental/DynamicMaterial/Source/DynamicMaterial/Private/Components/MaterialValues/DMMaterialValueBool.cpp
@@ -3,6 +3,7 @@
 #include "Components/MaterialValues/DMMaterialValueBool.h"
 
 #if WITH_EDITOR
+#include "Components/MaterialValuesDynamic/DMMaterialValueBoolDynamic.h"
 #include "DMDefs.h"
 #include "Materials/Material.h"
 #include "Materials/MaterialExpression.h"
@@ -36,7 +37,12 @@ void UDMMaterialValueBool::GenerateExpression(const TSharedRef<IDMMaterialBuildS
 		return;
 	}
  
-	UMaterialExpressionStaticBoolParameter* NewExpression = InBuildState->GetBuildUtils().CreateExpressionParameter<UMaterialExpressionStaticBoolParameter>(GetMaterialParameterName(), UE_DM_NodeComment_Default);
+	UMaterialExpressionStaticBoolParameter* NewExpression = InBuildState->GetBuildUtils().CreateExpressionParameter<UMaterialExpressionStaticBoolParameter>(
+		GetMaterialParameterName(),
+		GetParameterGroup(),
+		UE_DM_NodeComment_Default
+	);
+
 	check(NewExpression);
  
 	NewExpression->DefaultValue = Value;
@@ -54,7 +60,8 @@ void UDMMaterialValueBool::SetMIDParameter(UMaterialInstanceDynamic* InMID) cons
 
 	check(InMID);
  
-	// No idea how to implement this
+	// True dynamic branching is currently being worked on. When it is in, this will become relevant.
+	// There is no Jira yet.
 	checkNoEntry();
 }
 
@@ -74,6 +81,19 @@ void UDMMaterialValueBool::ResetDefaultValue()
 	bDefaultValue = false;
 }
 
+UDMMaterialValueDynamic* UDMMaterialValueBool::ToDynamic(UDynamicMaterialModelDynamic* InMaterialModelDynamic)
+{
+	UDMMaterialValueBoolDynamic* ValueDynamic = UDMMaterialValueDynamic::CreateValueDynamic<UDMMaterialValueBoolDynamic>(InMaterialModelDynamic, this);
+	ValueDynamic->SetValue(Value);
+
+	return ValueDynamic;
+}
+
+FString UDMMaterialValueBool::GetComponentPathComponent() const
+{
+	return TEXT("Bool");
+}
+
 TSharedPtr<FJsonValue> UDMMaterialValueBool::JsonSerialize() const
 {
 	return FDMJsonUtils::Serialize(Value);
diff --git a/Engine/Plugins/Experimental/DynamicMaterial/Source/DynamicMaterial/Private/Components/MaterialValues/DMMaterialValueColorAtlas.cpp b/Engine/Plugins/Experimental/DynamicMaterial/Source/DynamicMaterial/Private/Components/MaterialValues/DMMaterialValueColorAtlas.cpp
--- a/Engine/Plugins/Experimental/DynamicMaterial/Source/DynamicMaterial/Private/Components/MaterialValues/DMMaterialValueColorAtlas.cpp
+++ b/Engine/Plugins/Experimental/DynamicMaterial/Source/DynamicMaterial/Private/Components/MaterialValues/DMMaterialValueColorAtlas.cpp
@@ -5,6 +5,7 @@
 #include "Materials/MaterialInstanceDynamic.h"
 
 #if WITH_EDITOR
+#include "Components/MaterialValuesDynamic/DMMaterialValueColorAtlasDynamic.h"
 #include "Curves/CurveLinearColor.h"
 #include "Curves/CurveLinearColorAtlas.h"
 #include "Materials/Material.h"
@@ -109,7 +110,12 @@ void UDMMaterialValueColorAtlas::GenerateExpression(const TSharedRef<IDMMaterial
 		return;
 	}
 
-	UMaterialExpressionScalarParameter* AlphaParameter = InBuildState->GetBuildUtils().CreateExpressionParameter<UMaterialExpressionScalarParameter>(GetMaterialParameterName(), UE_DM_NodeComment_Default);
+	UMaterialExpressionScalarParameter* AlphaParameter = InBuildState->GetBuildUtils().CreateExpressionParameter<UMaterialExpressionScalarParameter>(
+		GetMaterialParameterName(), 
+		GetParameterGroup(), 
+		UE_DM_NodeComment_Default
+	);
+
 	check(AlphaParameter);
 
 	// This is a parameter, but we're treating it as a standard node.
@@ -137,6 +143,19 @@ void UDMMaterialValueColorAtlas::ResetDefaultValue()
 	DefaultValue = 0.f;
 }
 
+UDMMaterialValueDynamic* UDMMaterialValueColorAtlas::ToDynamic(UDynamicMaterialModelDynamic* InMaterialModelDynamic)
+{
+	UDMMaterialValueColorAtlasDynamic* ValueDynamic = UDMMaterialValueDynamic::CreateValueDynamic<UDMMaterialValueColorAtlasDynamic>(InMaterialModelDynamic, this);
+	ValueDynamic->SetValue(Value);
+
+	return ValueDynamic;
+}
+
+FString UDMMaterialValueColorAtlas::GetComponentPathComponent() const
+{
+	return TEXT("ColorAtlasAlpha");
+}
+
 TSharedPtr<FJsonValue> UDMMaterialValueColorAtlas::JsonSerialize() const
 {
 	return FDMJsonUtils::Serialize({
diff --git a/Engine/Plugins/Experimental/DynamicMaterial/Source/DynamicMaterial/Private/Components/MaterialValues/DMMaterialValueFloat1.cpp b/Engine/Plugins/Experimental/DynamicMaterial/Source/DynamicMaterial/Private/Components/MaterialValues/DMMaterialValueFloat1.cpp
--- a/Engine/Plugins/Experimental/DynamicMaterial/Source/DynamicMaterial/Private/Components/MaterialValues/DMMaterialValueFloat1.cpp
+++ b/Engine/Plugins/Experimental/DynamicMaterial/Source/DynamicMaterial/Private/Components/MaterialValues/DMMaterialValueFloat1.cpp
@@ -5,6 +5,7 @@
 #include "Materials/MaterialInstanceDynamic.h"
 
 #if WITH_EDITOR
+#include "Components/MaterialValuesDynamic/DMMaterialValueFloat1Dynamic.h"
 #include "DMDefs.h"
 #include "Materials/Material.h"
 #include "Materials/MaterialExpression.h"
@@ -38,7 +39,12 @@ void UDMMaterialValueFloat1::GenerateExpression(const TSharedRef<IDMMaterialBuil
 		return;
 	}
  
-	UMaterialExpressionScalarParameter* NewExpression = InBuildState->GetBuildUtils().CreateExpressionParameter<UMaterialExpressionScalarParameter>(GetMaterialParameterName(), UE_DM_NodeComment_Default);
+	UMaterialExpressionScalarParameter* NewExpression = InBuildState->GetBuildUtils().CreateExpressionParameter<UMaterialExpressionScalarParameter>(
+		GetMaterialParameterName(),
+		GetParameterGroup(), 
+		UE_DM_NodeComment_Default
+	);
+
 	check(NewExpression);
  
 	NewExpression->DefaultValue = Value;
@@ -56,6 +62,19 @@ void UDMMaterialValueFloat1::ResetDefaultValue()
 	DefaultValue = 0.f;
 }
 
+UDMMaterialValueDynamic* UDMMaterialValueFloat1::ToDynamic(UDynamicMaterialModelDynamic* InMaterialModelDynamic)
+{
+	UDMMaterialValueFloat1Dynamic* ValueDynamic = UDMMaterialValueDynamic::CreateValueDynamic<UDMMaterialValueFloat1Dynamic>(InMaterialModelDynamic, this);
+	ValueDynamic->SetValue(Value);
+
+	return ValueDynamic;
+}
+
+FString UDMMaterialValueFloat1::GetComponentPathComponent() const
+{
+	return TEXT("Scalar");
+}
+
 TSharedPtr<FJsonValue> UDMMaterialValueFloat1::JsonSerialize() const
 {
 	return FDMJsonUtils::Serialize(Value);
diff --git a/Engine/Plugins/Experimental/DynamicMaterial/Source/DynamicMaterial/Private/Components/MaterialValues/DMMaterialValueFloat2.cpp b/Engine/Plugins/Experimental/DynamicMaterial/Source/DynamicMaterial/Private/Components/MaterialValues/DMMaterialValueFloat2.cpp
--- a/Engine/Plugins/Experimental/DynamicMaterial/Source/DynamicMaterial/Private/Components/MaterialValues/DMMaterialValueFloat2.cpp
+++ b/Engine/Plugins/Experimental/DynamicMaterial/Source/DynamicMaterial/Private/Components/MaterialValues/DMMaterialValueFloat2.cpp
@@ -4,6 +4,7 @@
 #include "Materials/MaterialInstanceDynamic.h"
 
 #if WITH_EDITOR
+#include "Components/MaterialValuesDynamic/DMMaterialValueFloat2Dynamic.h"
 #include "DMDefs.h"
 #include "Materials/Material.h"
 #include "Materials/MaterialExpression.h"
@@ -38,7 +39,12 @@ void UDMMaterialValueFloat2::GenerateExpression(const TSharedRef<IDMMaterialBuil
 		return;
 	}
  
-	UMaterialExpressionVectorParameter* ValueExpression = InBuildState->GetBuildUtils().CreateExpressionParameter<UMaterialExpressionVectorParameter>(GetMaterialParameterName(), UE_DM_NodeComment_Default);
+	UMaterialExpressionVectorParameter* ValueExpression = InBuildState->GetBuildUtils().CreateExpressionParameter<UMaterialExpressionVectorParameter>(
+		GetMaterialParameterName(),
+		GetParameterGroup(),
+		UE_DM_NodeComment_Default
+	);
+
 	check(ValueExpression);
  
 	ValueExpression->DefaultValue = FLinearColor(Value.X, Value.Y, 0, 0);
@@ -65,6 +71,19 @@ void UDMMaterialValueFloat2::ResetDefaultValue()
 	DefaultValue = FVector2D::ZeroVector;
 }
 
+UDMMaterialValueDynamic* UDMMaterialValueFloat2::ToDynamic(UDynamicMaterialModelDynamic* InMaterialModelDynamic)
+{
+	UDMMaterialValueFloat2Dynamic* ValueDynamic = UDMMaterialValueDynamic::CreateValueDynamic<UDMMaterialValueFloat2Dynamic>(InMaterialModelDynamic, this);
+	ValueDynamic->SetValue(Value);
+
+	return ValueDynamic;
+}
+
+FString UDMMaterialValueFloat2::GetComponentPathComponent() const
+{
+	return TEXT("Vector2D");
+}
+
 TSharedPtr<FJsonValue> UDMMaterialValueFloat2::JsonSerialize() const
 {
 	return FDMJsonUtils::Serialize(Value);
@@ -83,7 +102,7 @@ bool UDMMaterialValueFloat2::JsonDeserialize(const TSharedPtr<FJsonValue>& InJso
 	return false;
 }
 
-void UDMMaterialValueFloat2::SetDefaultValue(FVector2D InDefaultValue)
+void UDMMaterialValueFloat2::SetDefaultValue(const FVector2D& InDefaultValue)
 {
 	DefaultValue = InDefaultValue;
 }
diff --git a/Engine/Plugins/Experimental/DynamicMaterial/Source/DynamicMaterial/Private/Components/MaterialValues/DMMaterialValueFloat3RGB.cpp b/Engine/Plugins/Experimental/DynamicMaterial/Source/DynamicMaterial/Private/Components/MaterialValues/DMMaterialValueFloat3RGB.cpp
--- a/Engine/Plugins/Experimental/DynamicMaterial/Source/DynamicMaterial/Private/Components/MaterialValues/DMMaterialValueFloat3RGB.cpp
+++ b/Engine/Plugins/Experimental/DynamicMaterial/Source/DynamicMaterial/Private/Components/MaterialValues/DMMaterialValueFloat3RGB.cpp
@@ -4,6 +4,7 @@
 #include "Materials/MaterialInstanceDynamic.h"
 
 #if WITH_EDITOR
+#include "Components/MaterialValuesDynamic/DMMaterialValueFloat3RGBDynamic.h"
 #include "DMDefs.h"
 #include "Materials/Material.h"
 #include "Materials/MaterialExpression.h"
@@ -37,7 +38,12 @@ void UDMMaterialValueFloat3RGB::GenerateExpression(const TSharedRef<IDMMaterialB
 		return;
 	}
  
-	UMaterialExpressionVectorParameter* NewExpression = InBuildState->GetBuildUtils().CreateExpressionParameter<UMaterialExpressionVectorParameter>(GetMaterialParameterName(), UE_DM_NodeComment_Default);
+	UMaterialExpressionVectorParameter* NewExpression = InBuildState->GetBuildUtils().CreateExpressionParameter<UMaterialExpressionVectorParameter>(
+		GetMaterialParameterName(),
+		GetParameterGroup(),
+		UE_DM_NodeComment_Default
+	);
+
 	check(NewExpression);
  
 	NewExpression->DefaultValue = FLinearColor(Value.R, Value.G, Value.B, 0);
@@ -62,6 +68,19 @@ void UDMMaterialValueFloat3RGB::ResetDefaultValue()
 	DefaultValue = FLinearColor(0, 0, 0, 1);
 }
 
+UDMMaterialValueDynamic* UDMMaterialValueFloat3RGB::ToDynamic(UDynamicMaterialModelDynamic* InMaterialModelDynamic)
+{
+	UDMMaterialValueFloat3RGBDynamic* ValueDynamic = UDMMaterialValueDynamic::CreateValueDynamic<UDMMaterialValueFloat3RGBDynamic>(InMaterialModelDynamic, this);
+	ValueDynamic->SetValue(Value);
+
+	return ValueDynamic;
+}
+
+FString UDMMaterialValueFloat3RGB::GetComponentPathComponent() const
+{
+	return TEXT("RGB");
+}
+
 TSharedPtr<FJsonValue> UDMMaterialValueFloat3RGB::JsonSerialize() const
 {
 	return FDMJsonUtils::Serialize(Value);
@@ -80,7 +99,7 @@ bool UDMMaterialValueFloat3RGB::JsonDeserialize(const TSharedPtr<FJsonValue>& In
 	return false;
 }
 
-void UDMMaterialValueFloat3RGB::SetDefaultValue(FLinearColor InDefaultValue)
+void UDMMaterialValueFloat3RGB::SetDefaultValue(const FLinearColor& InDefaultValue)
 {
 	DefaultValue = InDefaultValue;
 }
diff --git a/Engine/Plugins/Experimental/DynamicMaterial/Source/DynamicMaterial/Private/Components/MaterialValues/DMMaterialValueFloat3RPY.cpp b/Engine/Plugins/Experimental/DynamicMaterial/Source/DynamicMaterial/Private/Components/MaterialValues/DMMaterialValueFloat3RPY.cpp
--- a/Engine/Plugins/Experimental/DynamicMaterial/Source/DynamicMaterial/Private/Components/MaterialValues/DMMaterialValueFloat3RPY.cpp
+++ b/Engine/Plugins/Experimental/DynamicMaterial/Source/DynamicMaterial/Private/Components/MaterialValues/DMMaterialValueFloat3RPY.cpp
@@ -4,6 +4,7 @@
 #include "Materials/MaterialInstanceDynamic.h"
 
 #if WITH_EDITOR
+#include "Components/MaterialValuesDynamic/DMMaterialValueFloat3RPYDynamic.h"
 #include "DMDefs.h"
 #include "Materials/Material.h"
 #include "Materials/MaterialExpression.h"
@@ -37,7 +38,12 @@ void UDMMaterialValueFloat3RPY::GenerateExpression(const TSharedRef<IDMMaterialB
 		return;
 	}
  
-	UMaterialExpressionVectorParameter* NewExpression = InBuildState->GetBuildUtils().CreateExpressionParameter<UMaterialExpressionVectorParameter>(GetMaterialParameterName(), UE_DM_NodeComment_Default);
+	UMaterialExpressionVectorParameter* NewExpression = InBuildState->GetBuildUtils().CreateExpressionParameter<UMaterialExpressionVectorParameter>(
+		GetMaterialParameterName(),
+		GetParameterGroup(),
+		UE_DM_NodeComment_Default
+	);
+
 	check(NewExpression);
  
 	NewExpression->DefaultValue = FLinearColor(Value.Roll, Value.Yaw, Value.Pitch, 0);
@@ -62,6 +68,19 @@ void UDMMaterialValueFloat3RPY::ResetDefaultValue()
 	DefaultValue = FRotator::ZeroRotator;
 }
 
+UDMMaterialValueDynamic* UDMMaterialValueFloat3RPY::ToDynamic(UDynamicMaterialModelDynamic* InMaterialModelDynamic)
+{
+	UDMMaterialValueFloat3RPYDynamic* ValueDynamic = UDMMaterialValueDynamic::CreateValueDynamic<UDMMaterialValueFloat3RPYDynamic>(InMaterialModelDynamic, this);
+	ValueDynamic->SetValue(Value);
+
+	return ValueDynamic;
+}
+
+FString UDMMaterialValueFloat3RPY::GetComponentPathComponent() const
+{
+	return TEXT("Rotator");
+}
+
 TSharedPtr<FJsonValue> UDMMaterialValueFloat3RPY::JsonSerialize() const
 {
 	return FDMJsonUtils::Serialize(Value);
@@ -80,7 +99,7 @@ bool UDMMaterialValueFloat3RPY::JsonDeserialize(const TSharedPtr<FJsonValue>& In
 	return false;
 }
 
-void UDMMaterialValueFloat3RPY::SetDefaultValue(FRotator InDefaultValue)
+void UDMMaterialValueFloat3RPY::SetDefaultValue(const FRotator& InDefaultValue)
 {
 	DefaultValue = InDefaultValue;
 }
diff --git a/Engine/Plugins/Experimental/DynamicMaterial/Source/DynamicMaterial/Private/Components/MaterialValues/DMMaterialValueFloat3XYZ.cpp b/Engine/Plugins/Experimental/DynamicMaterial/Source/DynamicMaterial/Private/Components/MaterialValues/DMMaterialValueFloat3XYZ.cpp
--- a/Engine/Plugins/Experimental/DynamicMaterial/Source/DynamicMaterial/Private/Components/MaterialValues/DMMaterialValueFloat3XYZ.cpp
+++ b/Engine/Plugins/Experimental/DynamicMaterial/Source/DynamicMaterial/Private/Components/MaterialValues/DMMaterialValueFloat3XYZ.cpp
@@ -4,6 +4,7 @@
 #include "Materials/MaterialInstanceDynamic.h"
 
 #if WITH_EDITOR
+#include "Components/MaterialValuesDynamic/DMMaterialValueFloat3XYZDynamic.h"
 #include "DMDefs.h"
 #include "Materials/Material.h"
 #include "Materials/MaterialExpression.h"
@@ -37,7 +38,12 @@ void UDMMaterialValueFloat3XYZ::GenerateExpression(const TSharedRef<IDMMaterialB
 		return;
 	}
  
-	UMaterialExpressionVectorParameter* NewExpression = InBuildState->GetBuildUtils().CreateExpressionParameter<UMaterialExpressionVectorParameter>(GetMaterialParameterName(), UE_DM_NodeComment_Default);
+	UMaterialExpressionVectorParameter* NewExpression = InBuildState->GetBuildUtils().CreateExpressionParameter<UMaterialExpressionVectorParameter>(
+		GetMaterialParameterName(),
+		GetParameterGroup(),
+		UE_DM_NodeComment_Default
+	);
+
 	check(NewExpression);
  
 	NewExpression->DefaultValue = FLinearColor(Value.X, Value.Y, Value.Z, 0);
@@ -62,6 +68,19 @@ void UDMMaterialValueFloat3XYZ::ResetDefaultValue()
 	DefaultValue = FVector::ZeroVector;
 }
 
+UDMMaterialValueDynamic* UDMMaterialValueFloat3XYZ::ToDynamic(UDynamicMaterialModelDynamic* InMaterialModelDynamic)
+{
+	UDMMaterialValueFloat3XYZDynamic* ValueDynamic = UDMMaterialValueDynamic::CreateValueDynamic<UDMMaterialValueFloat3XYZDynamic>(InMaterialModelDynamic, this);
+	ValueDynamic->SetValue(Value);
+
+	return ValueDynamic;
+}
+
+FString UDMMaterialValueFloat3XYZ::GetComponentPathComponent() const
+{
+	return TEXT("Vector3D");
+}
+
 TSharedPtr<FJsonValue> UDMMaterialValueFloat3XYZ::JsonSerialize() const
 {
 	return FDMJsonUtils::Serialize(Value);
@@ -80,7 +99,7 @@ bool UDMMaterialValueFloat3XYZ::JsonDeserialize(const TSharedPtr<FJsonValue>& In
 	return false;
 }
 
-void UDMMaterialValueFloat3XYZ::SetDefaultValue(FVector InDefaultValue)
+void UDMMaterialValueFloat3XYZ::SetDefaultValue(const FVector& InDefaultValue)
 {
 	DefaultValue = InDefaultValue;
 }
diff --git a/Engine/Plugins/Experimental/DynamicMaterial/Source/DynamicMaterial/Private/Components/MaterialValues/DMMaterialValueFloat4.cpp b/Engine/Plugins/Experimental/DynamicMaterial/Source/DynamicMaterial/Private/Components/MaterialValues/DMMaterialValueFloat4.cpp
--- a/Engine/Plugins/Experimental/DynamicMaterial/Source/DynamicMaterial/Private/Components/MaterialValues/DMMaterialValueFloat4.cpp
+++ b/Engine/Plugins/Experimental/DynamicMaterial/Source/DynamicMaterial/Private/Components/MaterialValues/DMMaterialValueFloat4.cpp
@@ -4,6 +4,7 @@
 #include "Materials/MaterialInstanceDynamic.h"
 
 #if WITH_EDITOR
+#include "Components/MaterialValuesDynamic/DMMaterialValueFloat4Dynamic.h"
 #include "DMDefs.h"
 #include "Materials/Material.h"
 #include "Materials/MaterialExpression.h"
@@ -37,7 +38,12 @@ void UDMMaterialValueFloat4::GenerateExpression(const TSharedRef<IDMMaterialBuil
 		return;
 	}
  
-	UMaterialExpressionVectorParameter* NewExpression = InBuildState->GetBuildUtils().CreateExpressionParameter<UMaterialExpressionVectorParameter>(GetMaterialParameterName(), UE_DM_NodeComment_Default);
+	UMaterialExpressionVectorParameter* NewExpression = InBuildState->GetBuildUtils().CreateExpressionParameter<UMaterialExpressionVectorParameter>(
+		GetMaterialParameterName(),
+		GetParameterGroup(),
+		UE_DM_NodeComment_Default
+	);
+
 	check(NewExpression);
  
 	NewExpression->DefaultValue = Value;
@@ -63,6 +69,19 @@ void UDMMaterialValueFloat4::ResetDefaultValue()
 	DefaultValue = FLinearColor(0, 0, 0, 1);
 }
 
+UDMMaterialValueDynamic* UDMMaterialValueFloat4::ToDynamic(UDynamicMaterialModelDynamic* InMaterialModelDynamic)
+{
+	UDMMaterialValueFloat4Dynamic* ValueDynamic = UDMMaterialValueDynamic::CreateValueDynamic<UDMMaterialValueFloat4Dynamic>(InMaterialModelDynamic, this);
+	ValueDynamic->SetValue(Value);
+
+	return ValueDynamic;
+}
+
+FString UDMMaterialValueFloat4::GetComponentPathComponent() const
+{
+	return TEXT("RGBA");
+}
+
 TSharedPtr<FJsonValue> UDMMaterialValueFloat4::JsonSerialize() const
 {
 	return FDMJsonUtils::Serialize(Value);
@@ -81,7 +100,7 @@ bool UDMMaterialValueFloat4::JsonDeserialize(const TSharedPtr<FJsonValue>& InJso
 	return false;
 }
 
-void UDMMaterialValueFloat4::SetDefaultValue(FLinearColor InDefaultValue)
+void UDMMaterialValueFloat4::SetDefaultValue(const FLinearColor& InDefaultValue)
 {
 	DefaultValue = InDefaultValue;
 }
diff --git a/Engine/Plugins/Experimental/DynamicMaterial/Source/DynamicMaterial/Private/Components/MaterialValues/DMMaterialValueRenderTarget.cpp b/Engine/Plugins/Experimental/DynamicMaterial/Source/DynamicMaterial/Private/Components/MaterialValues/DMMaterialValueRenderTarget.cpp
--- a/Engine/Plugins/Experimental/DynamicMaterial/Source/DynamicMaterial/Private/Components/MaterialValues/DMMaterialValueRenderTarget.cpp
+++ b/Engine/Plugins/Experimental/DynamicMaterial/Source/DynamicMaterial/Private/Components/MaterialValues/DMMaterialValueRenderTarget.cpp
@@ -6,6 +6,10 @@
 #include "Engine/TextureRenderTarget2D.h"
 #include "Misc/CoreDelegates.h"
 
+#if WITH_EDITOR
+#include "Components/MaterialValuesDynamic/DMMaterialValueRenderTargetDynamic.h"
+#endif
+
 const FString UDMMaterialValueRenderTarget::RendererPathToken = "Renderer";
 
 UDMMaterialValueRenderTarget::UDMMaterialValueRenderTarget()
@@ -158,6 +162,16 @@ void UDMMaterialValueRenderTarget::Update(EDMUpdateType InUpdateType)
 }
 
 #if WITH_EDITOR
+UDMMaterialValueDynamic* UDMMaterialValueRenderTarget::ToDynamic(UDynamicMaterialModelDynamic* InMaterialModelDynamic)
+{
+	return UDMMaterialValueDynamic::CreateValueDynamic<UDMMaterialValueRenderTargetDynamic>(InMaterialModelDynamic, this);
+}
+
+FString UDMMaterialValueRenderTarget::GetComponentPathComponent() const
+{
+	return TEXT("RenderTarget");
+}
+
 TSharedPtr<FJsonValue> UDMMaterialValueRenderTarget::JsonSerialize() const
 {
 	return FDMJsonUtils::Serialize({
diff --git a/Engine/Plugins/Experimental/DynamicMaterial/Source/DynamicMaterial/Private/Components/MaterialValues/DMMaterialValueTexture.cpp b/Engine/Plugins/Experimental/DynamicMaterial/Source/DynamicMaterial/Private/Components/MaterialValues/DMMaterialValueTexture.cpp
--- a/Engine/Plugins/Experimental/DynamicMaterial/Source/DynamicMaterial/Private/Components/MaterialValues/DMMaterialValueTexture.cpp
+++ b/Engine/Plugins/Experimental/DynamicMaterial/Source/DynamicMaterial/Private/Components/MaterialValues/DMMaterialValueTexture.cpp
@@ -5,6 +5,7 @@
 #include "Materials/MaterialInstanceDynamic.h"
 
 #if WITH_EDITOR
+#include "Components/MaterialValuesDynamic/DMMaterialValueTextureDynamic.h"
 #include "DMDefs.h"
 #include "Engine/Texture2D.h"
 #include "Engine/TextureCube.h"
@@ -60,16 +61,12 @@ FDMGetDefaultRGBTexture UDMMaterialValueTexture::GetDefaultRGBTexture;
 
 UDMMaterialValueTexture::UDMMaterialValueTexture()
 	: UDMMaterialValue(EDMValueType::VT_Texture)
+	, Value(nullptr)
 #if WITH_EDITORONLY_DATA
+	, DefaultValue(nullptr)
 	, OldValue(nullptr)
 #endif
 {
-#if WITH_EDITOR
-	ResetDefaultValue();
-	Value = DefaultValue;
-#else
-	Value = nullptr;
-#endif
 }
 
 #if WITH_EDITOR
@@ -85,17 +82,23 @@ void UDMMaterialValueTexture::GenerateExpression(const TSharedRef<IDMMaterialBui
 		return;
 	}
 
-	UMaterialExpressionTextureObjectParameter* NewExpression = InBuildState->GetBuildUtils().CreateExpressionParameter<UMaterialExpressionTextureObjectParameter>(GetMaterialParameterName(), UE_DM_NodeComment_Default, Value);
+	UMaterialExpressionTextureObjectParameter* NewExpression = InBuildState->GetBuildUtils().CreateExpressionParameter<UMaterialExpressionTextureObjectParameter>(
+		GetMaterialParameterName(), 
+		GetParameterGroup(), 
+		UE_DM_NodeComment_Default, 
+		Value
+	);
+
 	check(NewExpression);
 
 	InBuildState->AddValueExpressions(this, {NewExpression});
 }
  
-UDMMaterialValueTexture* UDMMaterialValueTexture::CreateMaterialValueTexture(UObject* Outer, UTexture* InTexture)
+UDMMaterialValueTexture* UDMMaterialValueTexture::CreateMaterialValueTexture(UObject* InOuter, UTexture* InTexture)
 {
 	check(InTexture);
  
-	UDMMaterialValueTexture* TextureValue = NewObject<UDMMaterialValueTexture>(Outer, NAME_None, RF_Transactional);
+	UDMMaterialValueTexture* TextureValue = NewObject<UDMMaterialValueTexture>(InOuter, NAME_None, RF_Transactional);
 	TextureValue->SetValue(InTexture);
 	return TextureValue;
 }
@@ -130,9 +133,9 @@ void UDMMaterialValueTexture::PostEditChangeProperty(FPropertyChangedEvent& InPr
 		return;
 	}
 
-	FName MemberPropertyName = InPropertyChangedEvent.GetMemberPropertyName();
+	const FName MemberPropertyName = InPropertyChangedEvent.GetMemberPropertyName();
 
-	if (MemberPropertyName == NAME_None)
+	if (MemberPropertyName.IsNone())
 	{
 		return;
 	}
@@ -184,6 +187,19 @@ void UDMMaterialValueTexture::ResetDefaultValue()
 	}
 }
 
+UDMMaterialValueDynamic* UDMMaterialValueTexture::ToDynamic(UDynamicMaterialModelDynamic* InMaterialModelDynamic)
+{
+	UDMMaterialValueTextureDynamic* ValueDynamic = UDMMaterialValueDynamic::CreateValueDynamic<UDMMaterialValueTextureDynamic>(InMaterialModelDynamic, this);
+	ValueDynamic->SetValue(Value);
+
+	return ValueDynamic;
+}
+
+FString UDMMaterialValueTexture::GetComponentPathComponent() const
+{
+	return TEXT("Texture");
+}
+
 TSharedPtr<FJsonValue> UDMMaterialValueTexture::JsonSerialize() const
 {
 	return FDMJsonUtils::Serialize(Value);
diff --git a/Engine/Plugins/Experimental/DynamicMaterial/Source/DynamicMaterial/Private/DMDefs.cpp b/Engine/Plugins/Experimental/DynamicMaterial/Source/DynamicMaterial/Private/DMDefs.cpp
--- a/Engine/Plugins/Experimental/DynamicMaterial/Source/DynamicMaterial/Private/DMDefs.cpp
+++ b/Engine/Plugins/Experimental/DynamicMaterial/Source/DynamicMaterial/Private/DMDefs.cpp
@@ -4,22 +4,3 @@
 
 int32 FDMUpdateGuard::GuardCount = 0;
 uint32 FDMInitializationGuard::GuardCount = 0;
-
-bool FDMInitializationGuard::IsInitializing()
-{
-	return GuardCount > 0;
-}
-
-FDMInitializationGuard::FDMInitializationGuard()
-{
-	// Used the struct name to make it clear it's a static variable.
-	++FDMInitializationGuard::GuardCount;
-}
-
-FDMInitializationGuard::~FDMInitializationGuard()
-{
-	if (FDMInitializationGuard::GuardCount > 0)
-	{
-		--FDMInitializationGuard::GuardCount;
-	}
-}
diff --git a/Engine/Plugins/Experimental/DynamicMaterial/Source/DynamicMaterial/Private/DMValueDefinition.cpp b/Engine/Plugins/Experimental/DynamicMaterial/Source/DynamicMaterial/Private/DMValueDefinition.cpp
--- a/Engine/Plugins/Experimental/DynamicMaterial/Source/DynamicMaterial/Private/DMValueDefinition.cpp
+++ b/Engine/Plugins/Experimental/DynamicMaterial/Source/DynamicMaterial/Private/DMValueDefinition.cpp
@@ -152,21 +152,21 @@ const FDMValueDefinition& UDMValueDefinitionLibrary::GetTypeForFloatCount(int32
 	return GetTypeForFloatCount(static_cast<uint8>(Enum));
 }
 
-bool UDMValueDefinitionLibrary::AreTypesCompatible(EDMValueType A, EDMValueType B, int32 AChannel, int32 BChannel)
+bool UDMValueDefinitionLibrary::AreTypesCompatible(EDMValueType InA, EDMValueType InB, int32 InAChannel, int32 BChannel)
 {
 	using namespace UE::MaterialDesigner::Private;
 
 	// While all floats are compatible with all over floats, this may change in the future.
-	const FDMValueDefinition* TypeA = &TypeDefinitions[A];
-	const FDMValueDefinition* TypeB = &TypeDefinitions[B];
+	const FDMValueDefinition* TypeA = &TypeDefinitions[InA];
+	const FDMValueDefinition* TypeB = &TypeDefinitions[InB];
 
-	if (AChannel != FDMMaterialStageConnectorChannel::WHOLE_CHANNEL)
+	if (InAChannel != FDMMaterialStageConnectorChannel::WHOLE_CHANNEL)
 	{
 		const int32 Count =
-			!!(AChannel & FDMMaterialStageConnectorChannel::FIRST_CHANNEL)
-			+ !!(AChannel & FDMMaterialStageConnectorChannel::SECOND_CHANNEL)
-			+ !!(AChannel & FDMMaterialStageConnectorChannel::THIRD_CHANNEL)
-			+ !!(AChannel & FDMMaterialStageConnectorChannel::FOURTH_CHANNEL);
+			!!(InAChannel & FDMMaterialStageConnectorChannel::FIRST_CHANNEL)
+			+ !!(InAChannel & FDMMaterialStageConnectorChannel::SECOND_CHANNEL)
+			+ !!(InAChannel & FDMMaterialStageConnectorChannel::THIRD_CHANNEL)
+			+ !!(InAChannel & FDMMaterialStageConnectorChannel::FOURTH_CHANNEL);
 
 		switch (Count)
 		{
@@ -259,11 +259,6 @@ const FText& FDMValueDefinition::GetChannelName(int32 InChannel) const
 	return Error;
 }
 
-TSubclassOf<UDMMaterialValue> FDMValueDefinition::GetValueClass() const
-{
-	return ValueClass;
-}
-
 const TArray<EDMValueType>& UDMValueDefinitionLibrary::GetValueTypes()
 {
 	using namespace UE::MaterialDesigner::Private;
diff --git a/Engine/Plugins/Experimental/DynamicMaterial/Source/DynamicMaterial/Private/Material/DynamicMaterialInstance.cpp b/Engine/Plugins/Experimental/DynamicMaterial/Source/DynamicMaterial/Private/Material/DynamicMaterialInstance.cpp
--- a/Engine/Plugins/Experimental/DynamicMaterial/Source/DynamicMaterial/Private/Material/DynamicMaterialInstance.cpp
+++ b/Engine/Plugins/Experimental/DynamicMaterial/Source/DynamicMaterial/Private/Material/DynamicMaterialInstance.cpp
@@ -1,9 +1,10 @@
 // Copyright Epic Games, Inc. All Rights Reserved.
 
 #include "Material/DynamicMaterialInstance.h"
+
+#include "Components/MaterialValues/DMMaterialValueTexture.h"
 #include "Materials/Material.h"
 #include "Model/DynamicMaterialModel.h"
-#include "UObject/ObjectSaveContext.h"
 
 #if WITH_EDITOR
 #include "Model/IDynamicMaterialModelEditorOnlyDataInterface.h"
@@ -11,19 +12,45 @@
 
 UDynamicMaterialInstance::UDynamicMaterialInstance()
 {
-	BaseMaterial = nullptr;
-	MaterialModel = nullptr;
+	MaterialModelBase = nullptr;
 
 	bOutputTranslucentVelocity = true;
 }
 
+UDynamicMaterialModelBase* UDynamicMaterialInstance::GetMaterialModelBase()
+{
+	return MaterialModelBase;
+}
+
+UDynamicMaterialModel* UDynamicMaterialInstance::GetMaterialModel()
+{
+	if (IsValid(MaterialModelBase))
+	{
+		return MaterialModelBase->ResolveMaterialModel();
+	}
+
+	return nullptr;
+}
+
 #if WITH_EDITOR
+void UDynamicMaterialInstance::SetMaterialModel(UDynamicMaterialModelBase* InMaterialModel)
+{
+	MaterialModelBase = InMaterialModel;
+
+	if (InMaterialModel)
+	{
+		InMaterialModel->Rename(nullptr, this, UE::DynamicMaterial::RenameFlags);
+	}
+}
+
 void UDynamicMaterialInstance::InitializeMIDPublic()
 {
+	check(MaterialModelBase);
+
+	UDynamicMaterialModel* MaterialModel = MaterialModelBase->ResolveMaterialModel();
 	check(MaterialModel);
 
-	BaseMaterial = MaterialModel->GetGeneratedMaterial();
-	SetParentInternal(BaseMaterial, false);
+	SetParentInternal(MaterialModel->GetGeneratedMaterial(), false);
 	ClearParameterValues();
 	UpdateCachedData();
 }
@@ -32,13 +59,16 @@ void UDynamicMaterialInstance::PostDuplicate(bool bDuplicateForPIE)
 {
 	Super::PostDuplicate(bDuplicateForPIE);
 
-	if (MaterialModel)
+	if (MaterialModelBase)
 	{
-		MaterialModel->SetDynamicMaterialInstance(this);
+		MaterialModelBase->SetDynamicMaterialInstance(this);
 
-		if (IDynamicMaterialModelEditorOnlyDataInterface* ModelEditorOnlyData = MaterialModel->GetEditorOnlyData())
+		if (UDynamicMaterialModel* MaterialModel = MaterialModelBase->ResolveMaterialModel())
 		{
-			ModelEditorOnlyData->RequestMaterialBuild();
+			if (IDynamicMaterialModelEditorOnlyDataInterface* ModelEditorOnlyData = MaterialModel->GetEditorOnlyData())
+			{
+				ModelEditorOnlyData->RequestMaterialBuild();
+			}
 		}
 	}
 }
@@ -47,20 +77,23 @@ void UDynamicMaterialInstance::PostEditImport()
 {
 	Super::PostEditImport();
 
-	if (MaterialModel)
+	if (MaterialModelBase)
 	{
-		MaterialModel->SetDynamicMaterialInstance(this);
+		MaterialModelBase->SetDynamicMaterialInstance(this);
 
-		if (IDynamicMaterialModelEditorOnlyDataInterface* ModelEditorOnlyData = MaterialModel->GetEditorOnlyData())
+		if (UDynamicMaterialModel* MaterialModel = MaterialModelBase->ResolveMaterialModel())
 		{
-			ModelEditorOnlyData->RequestMaterialBuild();
+			if (IDynamicMaterialModelEditorOnlyDataInterface* ModelEditorOnlyData = MaterialModel->GetEditorOnlyData())
+			{
+				ModelEditorOnlyData->RequestMaterialBuild();
+			}
 		}
 	}
 }
 
-void UDynamicMaterialInstance::OnMaterialBuilt(UDynamicMaterialModel* InMaterialModel)
+void UDynamicMaterialInstance::OnMaterialBuilt(UDynamicMaterialModelBase* InMaterialModel)
 {
-	if (MaterialModel != InMaterialModel)
+	if (MaterialModelBase != InMaterialModel)
 	{
 		return;
 	}
diff --git a/Engine/Plugins/Experimental/DynamicMaterial/Source/DynamicMaterial/Private/Model/DynamicMaterialModel.cpp b/Engine/Plugins/Experimental/DynamicMaterial/Source/DynamicMaterial/Private/Model/DynamicMaterialModel.cpp
--- a/Engine/Plugins/Experimental/DynamicMaterial/Source/DynamicMaterial/Private/Model/DynamicMaterialModel.cpp
+++ b/Engine/Plugins/Experimental/DynamicMaterial/Source/DynamicMaterial/Private/Model/DynamicMaterialModel.cpp
@@ -57,6 +57,7 @@ UDynamicMaterialModel::UDynamicMaterialModel()
 	auto AddFloatParameter = [this](FName InPropertyName, FName InParameterName, float InDefaultValue = 1.f)
 		{
 			UDMMaterialValueFloat1* Property = CreateDefaultSubobject<UDMMaterialValueFloat1>(InPropertyName);
+			Property->CachedParameterName = InParameterName;
 			GlobalParameterValues.Add(InPropertyName, Property);
 
 #if WITH_EDITOR
@@ -74,6 +75,7 @@ UDynamicMaterialModel::UDynamicMaterialModel()
 	auto AddVector2Parameter = [this](FName InPropertyName, FName InParameterName, const FVector2D& InDefaultValue)
 		{
 			UDMMaterialValueFloat2* Property = CreateDefaultSubobject<UDMMaterialValueFloat2>(InPropertyName);
+			Property->CachedParameterName = InParameterName;
 			GlobalParameterValues.Add(InPropertyName, Property);
 
 #if WITH_EDITOR
@@ -107,8 +109,18 @@ UDynamicMaterialModel::UDynamicMaterialModel()
 
 void UDynamicMaterialModel::SetDynamicMaterialInstance(UDynamicMaterialInstance* InDynamicMaterialInstance)
 {
+	if (DynamicMaterialInstance == InDynamicMaterialInstance)
+	{
+		return;
+	}
+
 	DynamicMaterialInstance = InDynamicMaterialInstance;
 
+	if (InDynamicMaterialInstance)
+	{
+		ApplyComponents(InDynamicMaterialInstance);
+	}
+
 #if WITH_EDITOR
 	if (IDynamicMaterialModelEditorOnlyDataInterface* ModelEditorOnlyData = GetEditorOnlyData())
 	{
@@ -122,6 +134,11 @@ UDMMaterialValueFloat1* UDynamicMaterialModel::GetGlobalOpacityValue() const
 	return GetTypedGlobalParameterValue<UDMMaterialValueFloat1>(GlobalOpacityValueName);
 }
 
+const TMap<FName, TObjectPtr<UDMMaterialValue>>& UDynamicMaterialModel::GetGlobalParameterValues() const
+{
+	return GlobalParameterValues;
+}
+
 UDMMaterialValue* UDynamicMaterialModel::GetGlobalParameterValue(FName InName) const
 {
 	if (const TObjectPtr<UDMMaterialValue>* ValuePtr = GlobalParameterValues.Find(InName))
@@ -194,17 +211,6 @@ UDMMaterialComponent* UDynamicMaterialModel::GetComponentByPath(FDMComponentPath
 	return nullptr;
 }
 
-#if WITH_EDITOR
-IDynamicMaterialModelEditorOnlyDataInterface* UDynamicMaterialModel::GetEditorOnlyData() const
-{
-	if (IsValid(EditorOnlyDataSI.GetObject()))
-	{
-		return EditorOnlyDataSI.GetInterface();
-	}
-
-	return nullptr;
-}
-
 UDMMaterialValue* UDynamicMaterialModel::GetValueByName(FName InName) const
 {
 	for (UDMMaterialValue* Value : Values)
@@ -218,21 +224,17 @@ UDMMaterialValue* UDynamicMaterialModel::GetValueByName(FName InName) const
 	return nullptr;
 }
 
-UDMMaterialValue* UDynamicMaterialModel::GetValueByIndex(int32 Index) const
+#if WITH_EDITOR
+IDynamicMaterialModelEditorOnlyDataInterface* UDynamicMaterialModel::GetEditorOnlyData() const
 {
-	if (Values.IsValidIndex(Index))
+	if (IsValid(EditorOnlyDataSI.GetObject()))
 	{
-		return Values[Index];
+		return EditorOnlyDataSI.GetInterface();
 	}
 
 	return nullptr;
 }
 
-UDMMaterialValue* UDynamicMaterialModel::AddValue(EDMValueType InValueType)
-{
-	return AddValue(UDMValueDefinitionLibrary::GetValueDefinition(InValueType).GetValueClass());
-}
-
 UDMMaterialValue* UDynamicMaterialModel::AddValue(TSubclassOf<UDMMaterialValue> InValueClass)
 {
 	UDMMaterialValue* NewValue = UDMMaterialValue::CreateMaterialValue(this, TEXT(""), InValueClass, false);
@@ -256,7 +258,7 @@ void UDynamicMaterialModel::RemoveRuntimeComponentReference(UDMMaterialComponent
 	RuntimeComponents.Remove(InValue);
 }
 
-void UDynamicMaterialModel::RemoveValueByName(FName InName)
+void UDynamicMaterialModel::RemoveValueByParameterName(FName InName)
 {
 	int32 FoundIndex = Values.IndexOfByPredicate([InName](UDMMaterialValue* Value)
 		{
@@ -277,20 +279,6 @@ void UDynamicMaterialModel::RemoveValueByName(FName InName)
 	}
 }
 
-void UDynamicMaterialModel::RemoveValueByIndex(int32 Index)
-{
-	if (Values.IsValidIndex(Index))
-	{
-		Values.RemoveAt(Index);
-
-		if (IDynamicMaterialModelEditorOnlyDataInterface* ModelEditorOnlyData = GetEditorOnlyData())
-		{
-			ModelEditorOnlyData->RequestMaterialBuild();
-			ModelEditorOnlyData->OnValueListUpdate();
-		}
-	}
-}
-
 bool UDynamicMaterialModel::HasParameterName(FName InParameterName) const
 {
 	if (Values.ContainsByPredicate([InParameterName](const UDMMaterialValue* InValue)
@@ -366,11 +354,13 @@ void UDynamicMaterialModel::FreeParameter(UDMMaterialParameter* InParameter)
 		return;
 	}
 
+	const FName ParameterName = InParameter->GetParameterName();
+
 	for (TMap<FName, TWeakObjectPtr<UDMMaterialParameter>>::TIterator It(ParameterMap); It; ++It)
 	{
-		const TWeakObjectPtr<UDMMaterialParameter>& ParameterWeak = It->Value;
+		UDMMaterialParameter* Parameter = It->Value.Get();
 
-		if (ParameterWeak.IsValid() == false || ParameterWeak->GetParameterName() == InParameter->GetParameterName())
+		if (!Parameter || Parameter->GetParameterName() == ParameterName)
 		{
 			It.RemoveCurrent();
 		}
@@ -406,16 +396,6 @@ bool UDynamicMaterialModel::ConditionalFreeParameter(UDMMaterialParameter* InPar
 	// We're in the map at the given name.
 	return false;
 }
-
-void UDynamicMaterialModel::ResetData()
-{
-	Values.Empty();
-
-	if (IDynamicMaterialModelEditorOnlyDataInterface* ModelEditorOnlyData = GetEditorOnlyData())
-	{
-		ModelEditorOnlyData->ResetData();
-	}
-}
 #endif
 
 void UDynamicMaterialModel::OnValueUpdated(UDMMaterialValue* InValue, EDMUpdateType InUpdateType)
@@ -466,6 +446,45 @@ void UDynamicMaterialModel::OnTextureUVUpdated(UDMTextureUV* InTextureUV)
 #endif
 }
 
+void UDynamicMaterialModel::ApplyComponents(UMaterialInstanceDynamic* InMID)
+{
+	for (const TPair<FName, TObjectPtr<UDMMaterialValue>>& GlobalParameterPair : GlobalParameterValues)
+	{
+#if WITH_EDITOR
+		if (GlobalParameterPair.Value->IsComponentCreated())
+		{
+			GlobalParameterPair.Value->SetComponentState(EDMComponentLifetimeState::Added);
+		}
+#endif
+
+		GlobalParameterPair.Value->SetMIDParameter(InMID);
+	}
+
+	for (UDMMaterialValue* Value : Values)
+	{
+#if WITH_EDITOR
+		if (Value->IsComponentCreated())
+		{
+			Value->SetComponentState(EDMComponentLifetimeState::Added);
+		}
+#endif
+
+		Value->SetMIDParameter(InMID);
+	}
+
+	for (const TObjectPtr<UDMMaterialComponent>& RuntimeComponent : RuntimeComponents)
+	{
+		if (UDMMaterialValue* Value = Cast<UDMMaterialValue>(RuntimeComponent))
+		{
+			Value->SetMIDParameter(InMID);
+		}
+		else if (UDMTextureUV* TextureUV = Cast<UDMTextureUV>(RuntimeComponent))
+		{
+			TextureUV->SetMIDParameters(InMID);
+		}
+	}
+}
+
 void UDynamicMaterialModel::PostLoad()
 {
 	Super::PostLoad();
@@ -610,7 +629,7 @@ void UDynamicMaterialModel::FixGlobalParameterValues()
 
 		const FName Name = ParameterPair.Value->GetFName();
 
-		UDMMaterialValue* Local = FindObjectFast<UDMMaterialValue>(this, Name, /* Exact class */ false);
+		UDMMaterialValue* Local = FindObjectFast<UDMMaterialValue>(this, Name);
 
 		if (!Local)
 		{
@@ -625,16 +644,18 @@ void UDynamicMaterialModel::FixGlobalParameterValues()
 #if WITH_EDITOR
 void UDynamicMaterialModel::ReinitComponents()
 {
-	if (IsValid(DynamicMaterial))
+	if (IsValid(DynamicMaterial) && !DynamicMaterial->HasAnyFlags(RF_DuplicateTransient))
 	{
 		if (GUndo)
 		{
 			DynamicMaterial->Modify();
 		}
 
-		DynamicMaterial->AtomicallySetFlags(RF_DuplicateTransient);
+		DynamicMaterial->SetFlags(RF_DuplicateTransient);
 	}
 
+	FixGlobalVars();
+
 	// Clean up old parameters
 	ParameterMap.Empty();
 
@@ -661,8 +682,6 @@ void UDynamicMaterialModel::ReinitComponents()
 		}
 	}
 
-	FixGlobalVars();
-
 	if (IDynamicMaterialModelEditorOnlyDataInterface* ModelEditorOnlyData = GetEditorOnlyData())
 	{
 		ModelEditorOnlyData->ReinitComponents();
@@ -673,58 +692,26 @@ void UDynamicMaterialModel::FixGlobalVars()
 {
 	auto FixGlobalVar = [this](FName InValueName, FName InParameterName)
 		{
-			UDMMaterialValue* Property = GetGlobalParameterValue(InValueName);
+			UDMMaterialValue* GlobalVar = GetGlobalParameterValue(InValueName);
 
-			if (!Property)
+			if (!GlobalVar)
 			{
 				return;
 			}
 
-			UDMMaterialParameter* Parameter = Property->GetParameter();
+			GlobalVar->CachedParameterName = InParameterName;
 
-			if (!Parameter || Parameter->GetParameterName() != InParameterName)
-			{
-				if (Property->SetParameterName(InParameterName))
-				{
-					return;
-				}
+			UDMMaterialParameter* Parameter = GlobalVar->GetParameter();
 
-				Parameter = Property->GetParameter();
-
-				if (!Parameter)
-				{
-					return;
-				}
-			}
-	
-			if (const TWeakObjectPtr<UDMMaterialParameter>* FoundParameterPtr = ParameterMap.Find(InParameterName))
+			if (!Parameter)
 			{
-				UDMMaterialParameter* FoundParameter = FoundParameterPtr->Get();
-
-				if (!FoundParameter || FoundParameter->HasAnyFlags(RF_ArchetypeObject)
-					|| (Parameter && FoundParameter != Parameter))
-				{
-					if (Parameter)
-					{
-						ParameterMap[InParameterName] = Parameter;
-					}
-					else
-					{
-						ParameterMap.Remove(InParameterName);
-					}
-				}
+				GlobalVar->Parameter = NewObject<UDMMaterialParameter>(this, NAME_None, RF_Transactional);
+				Parameter = GlobalVar->Parameter;
 			}
 
-			if (Property && Parameter != Property->Parameter)
+			if (Parameter->GetParameterName() != InParameterName)
 			{
-				if (GUndo)
-				{
-					Parameter->Modify();
-					Property->Modify();
-				}
-
 				Parameter->ParameterName = InParameterName;
-				Property->Parameter = Parameter;
 			}
 		};
 
diff --git a/Engine/Plugins/Experimental/DynamicMaterial/Source/DynamicMaterial/Public/Components/DMMaterialComponent.h b/Engine/Plugins/Experimental/DynamicMaterial/Source/DynamicMaterial/Public/Components/DMMaterialComponent.h
--- a/Engine/Plugins/Experimental/DynamicMaterial/Source/DynamicMaterial/Public/Components/DMMaterialComponent.h
+++ b/Engine/Plugins/Experimental/DynamicMaterial/Source/DynamicMaterial/Public/Components/DMMaterialComponent.h
@@ -23,27 +23,34 @@ enum class EDMComponentLifetimeState : uint8
 /**
  * The base class for all material components. Has a few useful things.
  */
-UCLASS(Abstract, BlueprintType, meta = (DisplayName = "Material Designer Component"))
-class DYNAMICMATERIAL_API UDMMaterialComponent : public UObject, public FNotifyHook
+UCLASS(MinimalAPI, Abstract, BlueprintType, meta = (DisplayName = "Material Designer Component"))
+class UDMMaterialComponent : public UObject, public FNotifyHook
 {
 	GENERATED_BODY()
 
 	friend class UDynamicMaterialEditorSettings;
 
 public:
-	UDMMaterialComponent();
+#if WITH_EDITOR
+	/** Returns true if we can attempt to un-dirty a component. */
+	DYNAMICMATERIAL_API static bool CanClean();
 
-	/** Checks object flags and IsValid() */
-	UFUNCTION(BlueprintPure, Category = "Material Designer")
-	bool IsComponentValid() const;
+	/** Called to prevent cleaning for InDelayFor. **/
+	DYNAMICMATERIAL_API static void PreventClean(double InDelayFor = MinTimeBeforeClean);
+#endif
+
+	DYNAMICMATERIAL_API UDMMaterialComponent();
 
+	/** Checks object flags and IsValid() */
 	UFUNCTION(BlueprintPure, Category = "Material Designer")
-	virtual bool IsRootComponent() const { return false; }
+	DYNAMICMATERIAL_API bool IsComponentValid() const;
 
+	/** Searches the component for a specific component based on a path. */
 	UFUNCTION(BlueprintCallable, Category = "Material Designer")
-	UDMMaterialComponent* GetComponentByPath(const FString& InPath) const;
+	DYNAMICMATERIAL_API UDMMaterialComponent* GetComponentByPath(const FString& InPath) const;
 
-	UDMMaterialComponent* GetComponentByPath(FDMComponentPath& InPath) const;
+	/** Searches the component for a specific component based on a path. */
+	DYNAMICMATERIAL_API UDMMaterialComponent* GetComponentByPath(FDMComponentPath& InPath) const;
 
 	template<typename InComponentClass>
 	UDMMaterialComponent* GetComponentByPath(FDMComponentPath& InPath) const
@@ -51,19 +58,22 @@ public:
 		return Cast<InComponentClass>(GetComponentByPath(InPath));
 	}
 
+	/** Event that is triggered when this component, or a sub-component, changes to trigger other updates in the model. */
 	UFUNCTION(BlueprintCallable, Category = "Material Designer")
-	virtual void Update(EDMUpdateType InUpdateType);
+	DYNAMICMATERIAL_API virtual void Update(EDMUpdateType InUpdateType);
 
 #if WITH_EDITOR
-	UFUNCTION(BlueprintPure, Category = "Material Designer")
 	/** Returns the complete path from the model to this component. */
-	FString GetComponentPath() const;
+	UFUNCTION(BlueprintPure, Category = "Material Designer")
+	DYNAMICMATERIAL_API FString GetComponentPath() const;
 
+	/** Returns the component that owns this component in the model hierarchy. */
 	UFUNCTION(BlueprintCallable, Category = "Material Designer")
-	virtual UDMMaterialComponent* GetParentComponent() const;
+	DYNAMICMATERIAL_API virtual UDMMaterialComponent* GetParentComponent() const;
 
+	/** Returns the first in the model hierarchy above this component of the given type. */
 	UFUNCTION(BlueprintCallable, Category = "Material Designer")
-	UDMMaterialComponent* GetTypedParent(UClass* InParentClass, bool bInAllowSubclasses) const;
+	DYNAMICMATERIAL_API UDMMaterialComponent* GetTypedParent(UClass* InParentClass, bool bInAllowSubclasses) const;
 
 	template<class InParentClass>
 	InParentClass* GetTypedParent(bool bInAllowSubclasses) const
@@ -73,74 +83,83 @@ public:
 
 	/* Returns a description of this class/object. */
 	UFUNCTION(BlueprintPure, Category = "Material Designer")
-	virtual FText GetComponentDescription() const;
-
-	/** Returns true if we can attempt to un-dirty this component. */
-	static bool CanClean();
-
-	/** Called to prevent cleaning for MinTimeBeforeClean. **/
-	static void PreventClean(double DelayFor = MinTimeBeforeClean);
+	DYNAMICMATERIAL_API virtual FText GetComponentDescription() const;
 
 	/** Returns true if this component has been marked dirty. */
-	bool NeedsClean();
+	DYNAMICMATERIAL_API bool NeedsClean();
 
 	/** Performs whatever operation is involved in cleaning this component. */
-	virtual void DoClean();
+	DYNAMICMATERIAL_API virtual void DoClean();
 
 	UFUNCTION(BlueprintPure, Category = "Material Designer")
 	EDMComponentLifetimeState GetComponentState() const { return ComponentState; }
 
+	/** Returns true if this component is in the original "Created" state and has not been moved onto Added or Removed. */
 	UFUNCTION(BlueprintPure, Category = "Material Designer")
 	bool IsComponentCreated() const { return ComponentState == EDMComponentLifetimeState::Created; }
 
+	/** This is a kind of "useless" check, a component has _always_ been created. It's here for completeness. */
 	UFUNCTION(BlueprintPure, Category = "Material Designer")
 	bool HasComponentBeenCreated() const
 	{
-		// This is kind of a "useless" check, a component has _always_ been created.
 		return true;
 	}
 
+	/** Returns true if this component is in the Added state. */
 	UFUNCTION(BlueprintPure, Category = "Material Designer")
 	bool IsComponentAdded() const { return ComponentState == EDMComponentLifetimeState::Added; }
 
+	/** Returns true if this component is in the Added or greater state. */
 	UFUNCTION(BlueprintPure, Category = "Material Designer")
 	bool HasComponentBeenAdded() const { return ComponentState >= EDMComponentLifetimeState::Added; }
 
+	/** Returns true if this component is in the Removed state. */
 	UFUNCTION(BlueprintPure, Category = "Material Designer")
 	bool IsComponentRemoved() const { return ComponentState == EDMComponentLifetimeState::Removed; }
 
+	/** Returns true if this component is in the Removed or greater state. */
 	UFUNCTION(BlueprintPure, Category = "Material Designer")
 	bool HasComponentBeenRemoved() const { return ComponentState >= EDMComponentLifetimeState::Removed; }
 
+	/** Changes the component state to a new one. Should not be used to set it back to Created. */
 	UFUNCTION(BlueprintCallable, Category = "Material Designer")
-	void SetComponentState(EDMComponentLifetimeState NewState);
-
-	//~ Begin UObject
-	virtual bool Modify(bool bInAlwaysMarkDirty = true) override;
-	virtual void PostLoad() override;
-	//~ End UObject
-
-	/*** Begin FNotifyHook */
-	virtual void NotifyPreChange(class FEditPropertyChain* PropertyAboutToChange) {}
-	/*** End FNotifyHook */
+	DYNAMICMATERIAL_API void SetComponentState(EDMComponentLifetimeState InNewState);
 
+	/** Returns a list of FNames for this component representing editable UPROPERTYs. */
 	UFUNCTION(BlueprintPure, Category = "Material Designer")
 	const TArray<FName>& GetEditableProperties() const { return EditableProperties; }
 
+	/** Returns true the given UPROPERTY name is editable. */
 	UFUNCTION(BlueprintPure, Category = "Material Designer")
-	virtual bool IsPropertyVisible(FName Property) const { return true; }
+	virtual bool IsPropertyVisible(FName InProperty) const { return true; }
 
 	/** Returns the part of the component representing just this object */
-	virtual FString GetComponentPathComponent() const;
+	DYNAMICMATERIAL_API virtual FString GetComponentPathComponent() const;
 
-	virtual void PostEditorDuplicate(UDynamicMaterialModel* InMaterialModel, UDMMaterialComponent* InParent);
+	/** Called to ensure that the object hierarchy is correct. */
+	DYNAMICMATERIAL_API virtual void PostEditorDuplicate(UDynamicMaterialModel* InMaterialModel, UDMMaterialComponent* InParent);
 
 	DECLARE_MULTICAST_DELEGATE_TwoParams(FOnUpdate, UDMMaterialComponent*, EDMUpdateType)
-	FOnUpdate& GetOnUpdate() { return OnUpdate; }
+
+	/** Delegate called when this component's Update method is called. */
+	FOnUpdate::RegistrationType& GetOnUpdate() { return OnUpdate; }
 
 	DECLARE_MULTICAST_DELEGATE_TwoParams(FOnLifetimeStateChanged, UDMMaterialComponent*, EDMComponentLifetimeState)
-	FOnLifetimeStateChanged& GetOnAdded() { return OnAdded; }
-	FOnLifetimeStateChanged& GetOnRemoved() { return OnRemoved; }
+
+	/** Event called when this component's state changes to Added. */
+	FOnLifetimeStateChanged::RegistrationType& GetOnAdded() { return OnAdded; }
+
+	/** Event called when this component's state changes to Removed. */
+	FOnLifetimeStateChanged::RegistrationType& GetOnRemoved() { return OnRemoved; }
+
+	//~ Begin FNotifyHook
+	virtual void NotifyPreChange(FEditPropertyChain* InPropertyAboutToChange) {}
+	//~ End FNotifyHook
+
+	//~ Begin UObject
+	DYNAMICMATERIAL_API virtual bool Modify(bool bInAlwaysMarkDirty = true) override;
+	DYNAMICMATERIAL_API virtual void PostLoad() override;
+	//~ End UObject
 #endif
 
 protected:
@@ -162,19 +181,27 @@ protected:
 	FOnLifetimeStateChanged OnRemoved;
 #endif
 
+	/** Does some checks to see whether the out is safe to retrieve and retrieves it. */
 	UFUNCTION(BlueprintPure, Category = "Material Designer")
-	UObject* GetOuterSafe() const;
+	DYNAMICMATERIAL_API UObject* GetOuterSafe() const;
 
-	virtual UDMMaterialComponent* GetSubComponentByPath(FDMComponentPath& InPath, const FDMComponentPathSegment& InPathSegment) const;
+	/** Searches the component for a specific component based on a path. */
+	DYNAMICMATERIAL_API virtual UDMMaterialComponent* GetSubComponentByPath(FDMComponentPath& InPath, const FDMComponentPathSegment& InPathSegment) const;
 
-	void MarkComponentDirty();
+	/** Marks this component as needed to be cleaned. Preview materials updates, etc. */
+	DYNAMICMATERIAL_API void MarkComponentDirty();
 
 #if WITH_EDITOR
 	/** Allows this object to modify the child path when generating a path. */
-	virtual void GetComponentPathInternal(TArray<FString>& OutChildComponentPathComponents) const;
+	DYNAMICMATERIAL_API virtual void GetComponentPathInternal(TArray<FString>& OutChildComponentPathComponents) const;
+
+	/** Called when a component state changes to call the enter state functions. */
+	DYNAMICMATERIAL_API virtual void OnComponentStateChange(EDMComponentLifetimeState InNewState);
+
+	/** Called when a component enters the Added state. */
+	DYNAMICMATERIAL_API virtual void OnComponentAdded();
 
-	virtual void OnComponentStateChange(EDMComponentLifetimeState NewState);
-	virtual void OnComponentAdded();
-	virtual void OnComponentRemoved();
+	/** Called when a component enters the Removed state. */
+	DYNAMICMATERIAL_API virtual void OnComponentRemoved();
 #endif
 };
diff --git a/Engine/Plugins/Experimental/DynamicMaterial/Source/DynamicMaterial/Public/Components/DMMaterialLinkedComponent.h b/Engine/Plugins/Experimental/DynamicMaterial/Source/DynamicMaterial/Public/Components/DMMaterialLinkedComponent.h
--- a/Engine/Plugins/Experimental/DynamicMaterial/Source/DynamicMaterial/Public/Components/DMMaterialLinkedComponent.h
+++ b/Engine/Plugins/Experimental/DynamicMaterial/Source/DynamicMaterial/Public/Components/DMMaterialLinkedComponent.h
@@ -5,21 +5,23 @@
 #include "DMMaterialComponent.h"
 #include "DMMaterialLinkedComponent.generated.h"
 
-UCLASS(Abstract, BlueprintType, meta = (DisplayName = "Material Designer Linked Component"))
-class DYNAMICMATERIAL_API UDMMaterialLinkedComponent : public UDMMaterialComponent
+/** A component which links to a specific parent component in the hierarchy instead of its Outer. */
+UCLASS(MinimalAPI, Abstract, BlueprintType, meta = (DisplayName = "Material Designer Linked Component"))
+class UDMMaterialLinkedComponent : public UDMMaterialComponent
 {
 	GENERATED_BODY()
 
 public:
-	UDMMaterialLinkedComponent() = default;
-
 #if WITH_EDITOR
-	virtual UDMMaterialComponent* GetParentComponent() const override;
+	/** Returns the linked parent component. */
+	DYNAMICMATERIAL_API virtual UDMMaterialComponent* GetParentComponent() const override;
 
+	/** Sets the linked parent component. */
 	UFUNCTION(BlueprintCallable, Category = "Material Designer")
-	void SetParentComponent(UDMMaterialComponent* InParentComponent);
+	DYNAMICMATERIAL_API void SetParentComponent(UDMMaterialComponent* InParentComponent);
 
-	virtual void PostEditorDuplicate(UDynamicMaterialModel* InMaterialModel, UDMMaterialComponent* InParent) override;
+	/** Sets the parent component to the InParent parameter. */
+	DYNAMICMATERIAL_API virtual void PostEditorDuplicate(UDynamicMaterialModel* InMaterialModel, UDMMaterialComponent* InParent) override;
 #endif
 
 protected:
@@ -28,4 +30,3 @@ protected:
 	TObjectPtr<UDMMaterialComponent> ParentComponent = nullptr;
 #endif
 };
-
