Material Designer

diff --git a/Engine/Plugins/Experimental/DynamicMaterial/Source/DynamicMaterial/Public/Components/DMMaterialParameter.h b/Engine/Plugins/Experimental/DynamicMaterial/Source/DynamicMaterial/Public/Components/DMMaterialParameter.h
--- a/Engine/Plugins/Experimental/DynamicMaterial/Source/DynamicMaterial/Public/Components/DMMaterialParameter.h
+++ b/Engine/Plugins/Experimental/DynamicMaterial/Source/DynamicMaterial/Public/Components/DMMaterialParameter.h
@@ -10,30 +10,34 @@ class UDynamicMaterialModel;
 /**
  * A parameter on a Material Designer Instance.
  */
-UCLASS(BlueprintType, meta = (DisplayName = "Material Designer Parameter"))
-class DYNAMICMATERIAL_API UDMMaterialParameter : public UDMMaterialLinkedComponent
+UCLASS(MinimalAPI, BlueprintType, meta = (DisplayName = "Material Designer Parameter"))
+class UDMMaterialParameter : public UDMMaterialLinkedComponent
 {
 	GENERATED_BODY()
 
-public:
 	friend class UDynamicMaterialModel;
 
+public:
+	DYNAMICMATERIAL_API UDMMaterialParameter();
+
 	UFUNCTION(BlueprintPure, Category = "Material Designer")
-	UDynamicMaterialModel* GetMaterialModel() const;
+	DYNAMICMATERIAL_API UDynamicMaterialModel* GetMaterialModel() const;
 
 	UFUNCTION(BlueprintPure, Category = "Material Designer")
 	FName GetParameterName() const { return ParameterName; }
 
 #if WITH_EDITOR
+	/** Changes the parameter name registered with the Model. */
 	UFUNCTION(BlueprintCallable, Category = "Material Designer")
-	void RenameParameter(FName InBaseParameterName);
+	DYNAMICMATERIAL_API void RenameParameter(FName InBaseParameterName);
 
 	//~ Begin UObject
-	virtual void BeginDestroy() override;
+	/** Unregisters the parameter name with the model. */
+	DYNAMICMATERIAL_API virtual void BeginDestroy() override;
 	//~ End UObject
 
 	//~ Begin UDMMaterialComponent
-	virtual void PostEditorDuplicate(UDynamicMaterialModel* InMaterialModel, UDMMaterialComponent* InParent) override;
+	DYNAMICMATERIAL_API virtual void PostEditorDuplicate(UDynamicMaterialModel* InMaterialModel, UDMMaterialComponent* InParent) override;
 	//~ End UDMMaterialComponent
 #endif
 
@@ -42,11 +46,10 @@ protected:
 	FName ParameterName;
 
 private:
-	UDMMaterialParameter();
-
 	#if WITH_EDITOR
 	//~ Begin UDMMaterialComponent
-	virtual void OnComponentRemoved() override;
+	/** Unregisters the parameter name with the model. */
+	DYNAMICMATERIAL_API virtual void OnComponentRemoved() override;
 	//~ End UDMMaterialComponent
 #endif
 };
diff --git a/Engine/Plugins/Experimental/DynamicMaterial/Source/DynamicMaterial/Public/Components/DMMaterialValue.h b/Engine/Plugins/Experimental/DynamicMaterial/Source/DynamicMaterial/Public/Components/DMMaterialValue.h
--- a/Engine/Plugins/Experimental/DynamicMaterial/Source/DynamicMaterial/Public/Components/DMMaterialValue.h
+++ b/Engine/Plugins/Experimental/DynamicMaterial/Source/DynamicMaterial/Public/Components/DMMaterialValue.h
@@ -17,47 +17,52 @@ class UMaterialExpressionParameter;
 class UMaterialInstanceDynamic;
 
 #if WITH_EDITOR
+class UDMMaterialValueDynamic;
+class UDynamicMaterialModelDynamic;
+enum class EDMMaterialParameterGroup : uint8;
 struct IDMMaterialBuildStateInterface;
 #endif
 
 /**
- * A value used inside a material. Can be exported as a material parameter.
+ * A value used in a material. Manages its own parameter.
  */
-UCLASS(BlueprintType, Abstract, ClassGroup = "Material Designer", meta = (DisplayName = "Material Designer Value"))
-class DYNAMICMATERIAL_API UDMMaterialValue : public UDMMaterialLinkedComponent, public IDMJsonSerializable
+UCLASS(MinimalAPI, BlueprintType, Abstract, ClassGroup = "Material Designer", meta = (DisplayName = "Material Designer Value"))
+class UDMMaterialValue : public UDMMaterialLinkedComponent, public IDMJsonSerializable
 {
 	GENERATED_BODY()
- 
-	friend class UDynamicMaterialModel;
 
-#if WITH_EDITOR
-	friend class SDMMaterialEditor;
-#endif
+	friend class UDynamicMaterialModel;
  
 public:
-	static const FString ParameterPathToken;
+	DYNAMICMATERIAL_API static const FString ParameterPathToken;
 
 #if WITH_EDITOR
-	static const FName ValueName;
-
-	static UDMMaterialValue* CreateMaterialValue(UDynamicMaterialModel* InMaterialModel, const FString& InName, EDMValueType InValueType, bool bInLocal);
+	/* Name that should be used by all child classes for its value property ("Value"). */
+	DYNAMICMATERIAL_API static const FName ValueName;
 
+	/** Creates a new material value and initializes it with the given material model. */
 	UFUNCTION(BlueprintCallable, Category = "Material Designer")
-	static UDMMaterialValue* CreateMaterialValue(UDynamicMaterialModel* InMaterialModel, const FString& InName, TSubclassOf<UDMMaterialValue> InValueClass, bool bInLocal);
+	static DYNAMICMATERIAL_API UDMMaterialValue* CreateMaterialValue(UDynamicMaterialModel* InMaterialModel, const FString& InName, TSubclassOf<UDMMaterialValue> InValueClass, bool bInLocal);
 #endif
  
+	DYNAMICMATERIAL_API UDMMaterialValue();
+
+	/** Returns the owning material model (this object's Outer). */
 	UFUNCTION(BlueprintPure, Category = "Material Designer")
-	UDynamicMaterialModel* GetMaterialModel() const;
- 
+	DYNAMICMATERIAL_API UDynamicMaterialModel* GetMaterialModel() const;
+
+	/** Returns the type of value as respresented by the possible base type enums. */
 	UFUNCTION(BlueprintPure, Category = "Material Designer")
 	EDMValueType GetType() const { return Type; }
  
 #if WITH_EDITOR
+	/** Uses the Value Definition Library to get the type name of the base enum type. */
 	UFUNCTION(BlueprintPure, Category = "Material Designer")
-	FText GetTypeName() const;
- 
+	DYNAMICMATERIAL_API FText GetTypeName() const;
+
+	/** Combination of parameter name and type name. May be an automatically generated parameter name. */
 	UFUNCTION(BlueprintPure, Category = "Material Designer")
-	FText GetDescription() const;
+	DYNAMICMATERIAL_API FText GetDescription() const;
 
 	UFUNCTION(BlueprintPure, Category = "Material Designer")
 	virtual bool IsDefaultValue() const PURE_VIRTUAL(UDMMaterialValue::IsDefaultValue, return false;)
@@ -79,39 +84,46 @@ public:
 	virtual void ResetDefaultValue() PURE_VIRTUAL(UDMMaterialValue::ResetDefaultValue)
 #endif
 
+	/** Returns true is this value is not a "global" value belonging directly to the Model, but belongs to a specific component. */
 	UFUNCTION(BlueprintPure, Category = "Material Designer")
 	bool IsLocal() const { return bLocal; }
- 
+
+	/** Returns the parameter component managed by this value. Will only exist if a parameter name has been set. */
 	UFUNCTION(BlueprintPure, Category = "Material Designer")
 	UDMMaterialParameter* GetParameter() const { return Parameter; }
- 
+
+	/** Returns the specifically set parameter name or an automatically generated parameter name. */
 	UFUNCTION(BlueprintPure, Category = "Material Designer")
-	FName GetMaterialParameterName() const;
+	DYNAMICMATERIAL_API FName GetMaterialParameterName() const;
  
 #if WITH_EDITOR
+	/**
+	 * Sets a specific parameter name, overriding the automatic one if a base name is provided, or resetting back to the original
+	 * if a base name is not provider (NAME_None). When setting a specific name a parameter component is created and registered
+	 * with the Model.
+	 */
+	UFUNCTION(BlueprintCallable, Category = "Material Designer")
+	DYNAMICMATERIAL_API bool SetParameterName(FName InBaseName);
+
+	/** Determines the group this value's parameter will be in in the Material. */
+	DYNAMICMATERIAL_API EDMMaterialParameterGroup GetParameterGroup() const;
+
+	UFUNCTION(BlueprintPure, Category = "Material Designer")
+	bool GetShouldExposeParameter() const { return bExposeParameter; }
+
 	UFUNCTION(BlueprintCallable, Category = "Material Designer")
-	bool SetParameterName(FName InBaseName);
+	DYNAMICMATERIAL_API void SetShouldExposeParameter(bool bInExpose);
 
 	virtual void GenerateExpression(const TSharedRef<IDMMaterialBuildStateInterface>& InBuildState) const
 		PURE_VIRTUAL(UDMMaterialStageSource::GenerateExpressions)
 #endif
  
-#if WITH_EDITOR
-	UFUNCTION(BlueprintPure, Category = "Material Designer")
-	int32 FindIndex() const;
- 
-	int32 FindIndexSafe() const;
- 
-	UFUNCTION(BlueprintPure, Category = "Material Designer")
-	UMaterial* GetPreviewMaterial();
- 
+#if WITH_EDITOR 
 	/**
 	 * Returns the output index (channel WHOLE_CHANNEL) if this expression has pre-masked outputs.
 	 * Returns INDEX_NONE if it is not supported.
 	 */
-	virtual int32 GetInnateMaskOutput(int32 OutputChannels) const;
-
-	virtual FName GetMainPropertyName() const { return NAME_None; }
+	DYNAMICMATERIAL_API virtual int32 GetInnateMaskOutput(int32 OutputChannels) const;
 
 	/** Return true if, when setting the base stage, the same value should be applied to the mask stage. */
 	virtual bool IsWholeLayerValue() const { return false; }
@@ -120,30 +132,33 @@ public:
 	virtual bool AllowEditValue() const { return true; }
 #endif
  
-	virtual void SetMIDParameter(UMaterialInstanceDynamic* InMID) const;
+	virtual void SetMIDParameter(UMaterialInstanceDynamic* InMID) const PURE_VIRTUAL(UDMMaterialStageSource::SetMIDParameter)
 
-	//~ Begin IDMJsonSerializable
-	virtual TSharedPtr<FJsonValue> JsonSerialize() const override;
-	virtual bool JsonDeserialize(const TSharedPtr<FJsonValue>& InJsonValue) override;
-	//~ End IDMJsonSerializable
+#if WITH_EDITOR
+	virtual UDMMaterialValueDynamic* ToDynamic(UDynamicMaterialModelDynamic* InMaterialModelDynamic) PURE_VIRTUAL(UDMMaterialStageSource::ToDynamic, return nullptr;)
+#endif
  
 	//~ Begin UDMMaterialComponent
-	virtual void Update(EDMUpdateType InUpdateType) override;
+	DYNAMICMATERIAL_API virtual void Update(EDMUpdateType InUpdateType) override;
 #if WITH_EDITOR
-	virtual void DoClean() override;
-	virtual void PostEditorDuplicate(UDynamicMaterialModel* InMaterialModel, UDMMaterialComponent* InParent) override;
+	DYNAMICMATERIAL_API virtual void PostEditorDuplicate(UDynamicMaterialModel* InMaterialModel, UDMMaterialComponent* InParent) override;
 #endif
 	//~ End UDMMaterialComponent
+
+	/** Non-editor implementation. */
+	//~ Begin IDMJsonSerializable
+	DYNAMICMATERIAL_API virtual TSharedPtr<FJsonValue> JsonSerialize() const override;
+	DYNAMICMATERIAL_API virtual bool JsonDeserialize(const TSharedPtr<FJsonValue>& InJsonValue) override;
+	//~ End IDMJsonSerializable
  
 #if WITH_EDITOR
 	//~ Begin UObject
-	virtual bool Modify(bool bInAlwaysMarkDirty = true) override;
-	virtual void PostEditUndo() override;
-	virtual void PostEditChangeProperty(FPropertyChangedEvent& PropertyChangedEvent) override;
-	virtual void PostCDOContruct() override;
-	virtual void PostLoad() override;
-	virtual void PostEditImport() override;
-	virtual void BeginDestroy() override;
+	DYNAMICMATERIAL_API virtual bool Modify(bool bInAlwaysMarkDirty = true) override;
+	DYNAMICMATERIAL_API virtual void PostEditUndo() override;
+	DYNAMICMATERIAL_API virtual void PostEditChangeProperty(FPropertyChangedEvent& InPropertyChangedEvent) override;
+	DYNAMICMATERIAL_API virtual void PostCDOContruct() override;
+	DYNAMICMATERIAL_API virtual void PostLoad() override;
+	DYNAMICMATERIAL_API virtual void PostEditImport() override;
 	//~ End UObject
 #endif
  
@@ -155,41 +170,47 @@ protected:
  
 	/**
 	 * True: The value is local to the stage it is used in. 
-	 * False: The value is a global value that can be used anywhere
+	 * False: The value is a global value owned directly by the Model.
 	 */
 	UPROPERTY(VisibleInstanceOnly, BlueprintReadOnly, Category = "Material Designer")
 	bool bLocal;
 
 	/**
 	 * The parameter name used to expose this value in a material.
-	 * If it isn't provided, GetFName() will be used instead.
+	 * If it isn't provided, an automatic name will be generated.
 	 */
 	UPROPERTY(VisibleInstanceOnly, BlueprintReadOnly, Instanced, Category = "Material Designer")
 	TObjectPtr<UDMMaterialParameter> Parameter;
- 
+
+	UPROPERTY(VisibleInstanceOnly, BlueprintReadOnly, Category = "Material Designer")
+	FName CachedParameterName;
+
 #if WITH_EDITORONLY_DATA
-	UPROPERTY(VisibleInstanceOnly, BlueprintReadOnly, Transient, DuplicateTransient, TextExportTransient, Category = "Material Designer")
-	TObjectPtr<UMaterial> PreviewMaterial;
+	UPROPERTY(VisibleInstanceOnly, BlueprintReadOnly, Category = "Material Designer")
+	bool bExposeParameter;
 #endif
  
-	UDMMaterialValue(EDMValueType InType);
- 
-	virtual void OnValueChanged(EDMUpdateType InUpdateType);
+	DYNAMICMATERIAL_API UDMMaterialValue(EDMValueType InType);
+
+	/** Called when the value changes to call this component's and the parent component's Update if possible. */
+	DYNAMICMATERIAL_API virtual void OnValueChanged(EDMUpdateType InUpdateType);
+
+#if WITH_EDITOR
+	/** Generates an automatic path name based on the component hierarchy */
+	DYNAMICMATERIAL_API FName GenerateAutomaticParameterName() const;
+
+	/** Updates the cached parameter name based on the Parameter object or the above method. */
+	void UpdateCachedParameterName();
+#endif
  
 	//~ Begin UDMMaterialComponent
-	virtual UDMMaterialComponent* GetSubComponentByPath(FDMComponentPath& InPath, const FDMComponentPathSegment& InPathSegment) const override;
+	DYNAMICMATERIAL_API virtual UDMMaterialComponent* GetSubComponentByPath(FDMComponentPath& InPath, const FDMComponentPathSegment& InPathSegment) const override;
 	//~ End UDMMaterialComponent
 
 	#if WITH_EDITOR
-	void CreatePreviewMaterial();
-	void UpdatePreviewMaterial();
- 
 	//~ Begin UDMMaterialComponent
-	virtual void OnComponentAdded() override;
-	virtual void OnComponentRemoved() override;
+	DYNAMICMATERIAL_API virtual void OnComponentAdded() override;
+	DYNAMICMATERIAL_API virtual void OnComponentRemoved() override;
 	//~ End UDMMaterialComponent
 #endif
-
-private:
-	UDMMaterialValue();
 };
diff --git a/Engine/Plugins/Experimental/DynamicMaterial/Source/DynamicMaterial/Public/Components/DMRenderTargetRenderer.h b/Engine/Plugins/Experimental/DynamicMaterial/Source/DynamicMaterial/Public/Components/DMRenderTargetRenderer.h
--- a/Engine/Plugins/Experimental/DynamicMaterial/Source/DynamicMaterial/Public/Components/DMRenderTargetRenderer.h
+++ b/Engine/Plugins/Experimental/DynamicMaterial/Source/DynamicMaterial/Public/Components/DMRenderTargetRenderer.h
@@ -9,16 +9,20 @@
 class UDMMaterialValueRenderTarget;
 template <typename T> class TSubclassOf;
 
-UCLASS(Abstract, BlueprintType, Blueprintable, ClassGroup = "Material Designer", meta = (DisplayName = "Material Designer Render Target Renderer"))
-class DYNAMICMATERIAL_API UDMRenderTargetRenderer : public UDMMaterialComponent, public IDMJsonSerializable
+/**
+ * A component responsible for rendering something to a texture render target (value).
+ */
+UCLASS(MinimalAPI, Abstract, BlueprintType, Blueprintable, ClassGroup = "Material Designer", meta = (DisplayName = "Material Designer Render Target Renderer"))
+class UDMRenderTargetRenderer : public UDMMaterialComponent, public IDMJsonSerializable
 {
 	GENERATED_BODY()
 
 public:
 	FDelegateHandle EndOfFrameDelegateHandle;
 
+	/** Creates a render of the given class and initializes it. */
 	UFUNCTION(BlueprintCallable, Category = "Material Designer")
-	static UDMRenderTargetRenderer* CreateRenderTargetRenderer(TSubclassOf<UDMRenderTargetRenderer>
+	static DYNAMICMATERIAL_API UDMRenderTargetRenderer* CreateRenderTargetRenderer(TSubclassOf<UDMRenderTargetRenderer>
 		InRendererClass, UDMMaterialValueRenderTarget* InRenderTargetValue);
 
 	template<typename InRendererClass
@@ -28,36 +32,41 @@ public:
 		return Cast<InRendererClass>(CreateRenderTargetRenderer(InRendererClass::StaticClass(), InRenderTargetValue));
 	}
 
+	/** Returns the render target value (the object's Outer). */
 	UFUNCTION(BlueprintPure, Category = "Material Designer")
-	UDMMaterialValueRenderTarget* GetRenderTargetValue() const;
+	DYNAMICMATERIAL_API UDMMaterialValueRenderTarget* GetRenderTargetValue() const;
 
+	/** Updates the contents of the render target, redrawing and possibly resizing it. */
 	UFUNCTION(BlueprintCallable, Category = "Material Designer")
-	void UpdateRenderTarget();
+	DYNAMICMATERIAL_API void UpdateRenderTarget();
 
+	/** Updates the contents of the render target, redrawing and possibly resizing it, at the end of the frame. */
 	UFUNCTION(BlueprintCallable, Category = "Material Designer")
-	void AsyncUpdateRenderTarget();
+	DYNAMICMATERIAL_API void AsyncUpdateRenderTarget();
 
-	/** Will trigger the end of frame update if it is currently queued */
-	void FlushUpdateRenderTarget();
+	/** Will trigger the end of frame update if it is currently queued. */
+	DYNAMICMATERIAL_API void FlushUpdateRenderTarget();
 
+	/** Returns true is this target is currently being re-rendering. */
 	UFUNCTION(BlueprintPure, Category = "Material Designer")
 	bool IsUpdating() const { return bUpdating; }
 
 	//~ Begin IDMJsonSerializable
-	virtual TSharedPtr<FJsonValue> JsonSerialize() const override;
-	virtual bool JsonDeserialize(const TSharedPtr<FJsonValue>& InJsonValue) override;
+	DYNAMICMATERIAL_API virtual TSharedPtr<FJsonValue> JsonSerialize() const override;
+	DYNAMICMATERIAL_API virtual bool JsonDeserialize(const TSharedPtr<FJsonValue>& InJsonValue) override;
 	//~ End IDMJsonSerializable
 
 	//~ Begin UDMMaterialComponent
-	virtual void Update(EDMUpdateType InUpdateType) override;
+	DYNAMICMATERIAL_API virtual void Update(EDMUpdateType InUpdateType) override;
 	//~ End UDMMaterialComponent
 
 	//~ Begin UObject
-	virtual void PostLoad() override;
+	DYNAMICMATERIAL_API virtual void PostLoad() override;
 	//~ End UObject
 
 protected:
 	bool bUpdating = false;
 
-	virtual void UpdateRenderTarget_Internal() {}
+	/** Extend this to perform the actual render target rendering in a subclass. */
+	virtual void UpdateRenderTarget_Internal() PURE_VIRTUAL(UDMRenderTargetRenderer::UpdateRenderTarget_Internal)
 };
diff --git a/Engine/Plugins/Experimental/DynamicMaterial/Source/DynamicMaterial/Public/Components/DMTextureUV.h b/Engine/Plugins/Experimental/DynamicMaterial/Source/DynamicMaterial/Public/Components/DMTextureUV.h
--- a/Engine/Plugins/Experimental/DynamicMaterial/Source/DynamicMaterial/Public/Components/DMTextureUV.h
+++ b/Engine/Plugins/Experimental/DynamicMaterial/Source/DynamicMaterial/Public/Components/DMTextureUV.h
@@ -16,12 +16,20 @@
 
 class UDMMaterialParameter;
 class UDMTextureUV;
+class UDMTextureUVDynamic;
 class UDynamicMaterialModel;
 class UMaterialInstanceDynamic;
 
+#if WITH_EDITOR
+class UDynamicMaterialModelDynamic;
+class UMaterialInstanceDynamic;
+enum class EDMMaterialParameterGroup : uint8;
+#endif
+
 namespace UE::DynamicMaterial::ParamID
 {
 	// Individual parameters
+	constexpr int32 Invalid = -1;
 	constexpr int32 PivotX = 0;
 	constexpr int32 PivotY = 1;
 	constexpr int32 TilingX = 2;
@@ -37,8 +45,11 @@ namespace UE::DynamicMaterial::ParamID
 	constexpr int32 Offset = OffsetX;
 }
 
-UCLASS(BlueprintType, ClassGroup = "Material Designer", meta = (DisplayName = "Material Designer Texture UV"))
-class DYNAMICMATERIAL_API UDMTextureUV : public UDMMaterialLinkedComponent
+/*
+ * Represents a Texture UV material function with the following parameters: offset, tiling, pivot and rotation.
+ */
+UCLASS(MinimalAPI, BlueprintType, ClassGroup = "Material Designer", meta = (DisplayName = "Material Designer Texture UV"))
+class UDMTextureUV : public UDMMaterialLinkedComponent
 {
 	GENERATED_BODY()
 
@@ -47,34 +58,35 @@ class DYNAMICMATERIAL_API UDMTextureUV : public UDMMaterialLinkedComponent
 
 public:
 #if WITH_EDITOR
-	static const FName NAME_UVSource;
-	static const FName NAME_bMirrorOnX;
-	static const FName NAME_bMirrorOnY;
-#endif
-
-	static const FName NAME_Offset;
-	static const FName NAME_Pivot;
-	static const FName NAME_Rotation;
-	static const FName NAME_Tiling;
-
-	static const FString OffsetXPathToken;
-	static const FString OffsetYPathToken;
-	static const FString PivotXPathToken;
-	static const FString PivotYPathToken;
-	static const FString RotationPathToken;
-	static const FString TilingXPathToken;
-	static const FString TilingYPathToken;
+	DYNAMICMATERIAL_API static const FName NAME_UVSource;
+	DYNAMICMATERIAL_API static const FName NAME_bMirrorOnX;
+	DYNAMICMATERIAL_API static const FName NAME_bMirrorOnY;
+#endif
+
+	DYNAMICMATERIAL_API static const FName NAME_Offset;
+	DYNAMICMATERIAL_API static const FName NAME_Pivot;
+	DYNAMICMATERIAL_API static const FName NAME_Rotation;
+	DYNAMICMATERIAL_API static const FName NAME_Tiling;
+
+	DYNAMICMATERIAL_API static const FString OffsetXPathToken;
+	DYNAMICMATERIAL_API static const FString OffsetYPathToken;
+	DYNAMICMATERIAL_API static const FString PivotXPathToken;
+	DYNAMICMATERIAL_API static const FString PivotYPathToken;
+	DYNAMICMATERIAL_API static const FString RotationPathToken;
+	DYNAMICMATERIAL_API static const FString TilingXPathToken;
+	DYNAMICMATERIAL_API static const FString TilingYPathToken;
 	
 	static const FGuid GUID;
 
 #if WITH_EDITOR
 	// The bool is whether or not the property should be exposed to Sequencer as keyable.
-	static const TMap<FName, bool> TextureProperties;
+	DYNAMICMATERIAL_API static const TMap<FName, bool> TextureProperties;
 
-	static UDMTextureUV* CreateTextureUV(UObject* InOuter);
+	UFUNCTION(BlueprintCallable, Category = "Material Designer")
+	static DYNAMICMATERIAL_API UDMTextureUV* CreateTextureUV(UObject* InOuter);
 #endif
 
-	UDMTextureUV();
+	DYNAMICMATERIAL_API UDMTextureUV();
 
 #if WITH_EDITORONLY_DATA
 	UPROPERTY(EditInstanceOnly, BlueprintReadWrite, Category = "Material Designer")
@@ -82,110 +94,126 @@ public:
 #endif
 
 	UFUNCTION(BlueprintPure, Category = "Material Designer")
-	UDynamicMaterialModel* GetMaterialModel() const;
+	DYNAMICMATERIAL_API UDynamicMaterialModel* GetMaterialModel() const;
 
 #if WITH_EDITOR
 	UFUNCTION(BlueprintPure, Category = "Material Designer")
 	EDMUVSource GetUVSource() const { return UVSource; }
 
 	UFUNCTION(BlueprintCallable, Category = "Material Designer")
-	void SetUVSource(EDMUVSource InUVSource);
+	DYNAMICMATERIAL_API void SetUVSource(EDMUVSource InUVSource);
 #endif
 
 	UFUNCTION(BlueprintPure, Category = "Material Designer")
 	const FVector2D& GetOffset() const { return Offset; }
 
 	UFUNCTION(BlueprintCallable, Category = "Material Designer")
-	void SetOffset(const FVector2D& InOffset);
+	DYNAMICMATERIAL_API void SetOffset(const FVector2D& InOffset);
 
 	UFUNCTION(BlueprintPure, Category = "Material Designer")
 	const FVector2D& GetPivot() const { return Pivot; }
 
 	UFUNCTION(BlueprintCallable, Category = "Material Designer")
-	void SetPivot(const FVector2D& InPivot);
+	DYNAMICMATERIAL_API void SetPivot(const FVector2D& InPivot);
 
 	UFUNCTION(BlueprintPure, Category = "Material Designer")
 	float GetRotation() const { return Rotation; }
 
 	UFUNCTION(BlueprintCallable, Category = "Material Designer")
-	void SetRotation(float InRotation);
+	DYNAMICMATERIAL_API void SetRotation(float InRotation);
 
 	UFUNCTION(BlueprintPure, Category = "Material Designer")
 	const FVector2D& GetTiling() const { return Tiling; }
 
 	UFUNCTION(BlueprintCallable, Category = "Material Designer")
-	void SetTiling(const FVector2D& InTiling);
+	DYNAMICMATERIAL_API void SetTiling(const FVector2D& InTiling);
 
 #if WITH_EDITOR
 	UFUNCTION(BlueprintPure, Category = "Material Designer")
 	bool GetMirrorOnX() const { return bMirrorOnX; }
 
 	UFUNCTION(BlueprintCallable, Category = "Material Designer")
-	void SetMirrorOnX(bool bInMirrorOnX);
+	DYNAMICMATERIAL_API void SetMirrorOnX(bool bInMirrorOnX);
 
 	UFUNCTION(BlueprintPure, Category = "Material Designer")
 	bool GetMirrorOnY() const { return bMirrorOnY; }
 
 	UFUNCTION(BlueprintCallable, Category = "Material Designer")
-	void SetMirrorOnY(bool bInMirrorOnY);
+	DYNAMICMATERIAL_API void SetMirrorOnY(bool bInMirrorOnY);
 #endif
 
 	UFUNCTION(BlueprintPure, Category = "Material Designer")
-	TArray<UDMMaterialParameter*> GetParameters() const;
+	DYNAMICMATERIAL_API TArray<UDMMaterialParameter*> GetParameters() const;
 
 	UFUNCTION(BlueprintPure, Category = "Material Designer")
-	UDMMaterialParameter* GetMaterialParameter(FName InPropertyName, int32 InComponent) const;
+	DYNAMICMATERIAL_API UDMMaterialParameter* GetMaterialParameter(FName InPropertyName, int32 InComponent) const;
 
 	UFUNCTION(BlueprintPure, Category = "Material Designer")
-	FName GetMaterialParameterName(FName InPropertyName, int32 InComponent) const;
+	DYNAMICMATERIAL_API FName GetMaterialParameterName(FName InPropertyName, int32 InComponent) const;
 
 #if WITH_EDITOR
 	UFUNCTION(BlueprintCallable, Category = "Material Designer")
-	bool SetMaterialParameterName(FName InPropertyName, int32 InComponent, FName InNewName);
+	DYNAMICMATERIAL_API bool SetMaterialParameterName(FName InPropertyName, int32 InComponent, FName InNewName);
+
+	DYNAMICMATERIAL_API EDMMaterialParameterGroup GetParameterGroup(FName InPropertyName, int32 InComponent) const;
+
+	UFUNCTION(BlueprintPure, Category = "Material Designer")
+	DYNAMICMATERIAL_API bool GetShouldExposeParameter(FName InPropertyName, int32 InComponent) const;
+
+	UFUNCTION(BlueprintCallable, Category = "Material Designer")
+	DYNAMICMATERIAL_API void SetShouldExposeParameter(FName InPropertyName, int32 InComponent, bool bInExpose);
 #endif
 
-	void SetMIDParameters(UMaterialInstanceDynamic* InMID);
+	DYNAMICMATERIAL_API void SetMIDParameters(UMaterialInstanceDynamic* InMID);
+
+#if WITH_EDITOR
+	DYNAMICMATERIAL_API UDMTextureUVDynamic* ToDynamic(UDynamicMaterialModelDynamic* InMaterialModelDynamic);
+#endif
 
 	//~ Begin UObject
 #if WITH_EDITOR
-	virtual bool Modify(bool bInAlwaysMarkDirty = true) override;
-	virtual void PostLoad() override;
-	virtual void PostEditImport() override;
-	virtual void PostEditChangeProperty(struct FPropertyChangedEvent& PropertyChangedEvent) override;
-	virtual void PreEditUndo() override;
-	virtual void PostEditUndo() override;
-#endif
-	virtual void Serialize(FArchive& Ar) override;
+	DYNAMICMATERIAL_API virtual bool Modify(bool bInAlwaysMarkDirty = true) override;
+	DYNAMICMATERIAL_API virtual void PostLoad() override;
+	DYNAMICMATERIAL_API virtual void PostEditImport() override;
+	DYNAMICMATERIAL_API virtual void PostEditChangeProperty(FPropertyChangedEvent& InPropertyChangedEvent) override;
+	DYNAMICMATERIAL_API virtual void PreEditUndo() override;
+	DYNAMICMATERIAL_API virtual void PostEditUndo() override;
+#endif
+	DYNAMICMATERIAL_API virtual void Serialize(FArchive& Ar) override;
 	//~ End UObject
 
-#if WITH_EDITOR
 	//~ Begin UDMMaterialComponent
-	virtual void PostEditorDuplicate(UDynamicMaterialModel* InMaterialModel, UDMMaterialComponent* InParent) override;
-	//~ End UDMMaterialComponent
+#if WITH_EDITOR
+	DYNAMICMATERIAL_API virtual FString GetComponentPathComponent() const override;
+	DYNAMICMATERIAL_API virtual void PostEditorDuplicate(UDynamicMaterialModel* InMaterialModel, UDMMaterialComponent* InParent) override;
 #endif
+	DYNAMICMATERIAL_API virtual void Update(EDMUpdateType InUpdateType) override;
+	//~ End UDMMaterialComponent
 
 protected:
+	static int32 PropertyComponentToParamId(FName InPropertyName, int32 InComponent);
+
 #if WITH_EDITORONLY_DATA
-	UPROPERTY(EditInstanceOnly, BlueprintReadWrite, BlueprintSetter = SetUVSource, Category = "Material Designer|Texture UV",
+	UPROPERTY(EditInstanceOnly, BlueprintReadWrite, Category = "Material Designer|Texture UV",
 		meta = (AllowPrivateAccess = "true", ToolTip = "Forces a material rebuild.", NotKeyframeable))
 	EDMUVSource UVSource = EDMUVSource::Texture;
 
 	EDMUVSource UVSource_PreUndo = UVSource;
 #endif
 
-	UPROPERTY(EditInstanceOnly, BlueprintReadWrite, Getter = GetOffset, Setter = SetOffset, BlueprintSetter = SetOffset, Category = "Material Designer|Texture UV",
+	UPROPERTY(EditInstanceOnly, BlueprintReadWrite, Getter, Setter, BlueprintSetter = SetOffset, Category = "Material Designer|Texture UV",
 		meta = (AllowPrivateAccess = "true", Delta = 0.001))
 	FVector2D Offset = FVector2D(0.f, 0.f);
 
-	UPROPERTY(EditInstanceOnly, BlueprintReadWrite, Getter = GetPivot, Setter = SetPivot, BlueprintSetter = SetPivot, Category = "Material Designer|Texture UV",
+	UPROPERTY(EditInstanceOnly, BlueprintReadWrite, Getter, Setter, BlueprintSetter = SetPivot, Category = "Material Designer|Texture UV",
 		meta = (AllowPrivateAccess = "true", ToolTip="Pivot for rotation and tiling.", Delta = 0.001))
 	FVector2D Pivot = FVector2D(0.5, 0.5);
 
-	UPROPERTY(EditInstanceOnly, BlueprintReadWrite, Getter = GetRotation, Setter = SetRotation, BlueprintSetter = SetRotation, Category = "Material Designer|Texture UV",
+	UPROPERTY(EditInstanceOnly, BlueprintReadWrite, Getter, Setter, BlueprintSetter = SetRotation, Category = "Material Designer|Texture UV",
 		meta = (AllowPrivateAccess = "true", Delta = 1.0))
 	float Rotation = 0.f;
 
-	UPROPERTY(EditInstanceOnly, BlueprintReadWrite, Getter = GetTiling, Setter = SetTiling, BlueprintSetter = SetTiling, Category = "Material Designer|Texture UV",
+	UPROPERTY(EditInstanceOnly, BlueprintReadWrite, Getter, Setter, BlueprintSetter = SetTiling, Category = "Material Designer|Texture UV",
 		meta = (AllowPrivateAccess = "true", AllowPreserveRatio, VectorRatioMode = 3, Delta = 0.001))
 	FVector2D Tiling = FVector2D(1.f, 1.f);
 
@@ -194,11 +222,11 @@ protected:
 	FVector2D Scale = FVector2D(1.f, 1.f);
 
 #if WITH_EDITORONLY_DATA
-	UPROPERTY(EditInstanceOnly, BlueprintReadWrite, BlueprintSetter = SetMirrorOnX, Category = "Material Designer|Texture UV",
+	UPROPERTY(EditInstanceOnly, BlueprintReadWrite, Category = "Material Designer|Texture UV",
 		meta = (AllowPrivateAccess = "true", ToolTip = "Forces a material rebuild.", NotKeyframeable))
 	bool bMirrorOnX = false;
 
-	UPROPERTY(EditInstanceOnly, BlueprintReadWrite, BlueprintSetter = SetMirrorOnY, Category = "Material Designer|Texture UV",
+	UPROPERTY(EditInstanceOnly, BlueprintReadWrite, Category = "Material Designer|Texture UV",
 		meta = (AllowPrivateAccess = "true", ToolTip = "Forces a material rebuild.", NotKeyframeable))
 	bool bMirrorOnY = false;
 
@@ -206,28 +234,44 @@ protected:
 	bool bMirrorOnY_PreUndo = bMirrorOnY;
 #endif
 
-	UPROPERTY(VisibleInstanceOnly, BlueprintReadOnly, DuplicateTransient, TextExportTransient, Category = "Material Designer")
+	UPROPERTY(VisibleInstanceOnly, BlueprintReadOnly, TextExportTransient, Category = "Material Designer")
 	TMap<int32, TObjectPtr<UDMMaterialParameter>> MaterialParameters;
 
-#if WITH_EDITOR
-	mutable TMap<int32, bool> MaterialNodesCreated;
+	UPROPERTY(VisibleInstanceOnly, BlueprintReadOnly, TextExportTransient, Category = "Material Designer")
+	TMap<int32, FName> CachedParameterNames;
+
+	UPROPERTY(VisibleInstanceOnly, BlueprintReadOnly, TextExportTransient, Category = "Material Designer")
+	TSet<int32> ExposedParameters;
 
+#if WITH_EDITOR
 	bool bNeedsPostLoadValueUpdate = false;
 	bool bNeedsPostLoadStructureUpdate = false;
 
-	void CreateParameterNames();
 	void RemoveParameterNames();
 #endif
 
 	void OnTextureUVChanged(EDMUpdateType InUpdateType);
 
+#if WITH_EDITOR
+	/** Generates an automatic path component based on the property name and component. */
+	FName GenerateAutomaticPathComponent(FName InPropertyName, int32 InComponent) const;
+
+	/** Generates an automatic path name based on the component hierarchy */
+	FName GenerateAutomaticParameterName(FName InPropertyName, int32 InComponent) const;
+
+	/** Update an individual cached parameter name. */
+	void UpdateCachedParameterName(FName InPropertyName, int32 InComponent);
+
+	/** Updates the cached parameter name based on the Parameter object or the above method. */
+	void UpdateCachedParameterNames();
+#endif
+
 	//~ Begin UDMMaterialComponent
-	virtual void Update(EDMUpdateType InUpdateType) override;
-	virtual UDMMaterialComponent* GetSubComponentByPath(FDMComponentPath& InPath, const FDMComponentPathSegment& InPathSegment) const override;
+	DYNAMICMATERIAL_API virtual UDMMaterialComponent* GetSubComponentByPath(FDMComponentPath& InPath, const FDMComponentPathSegment& InPathSegment) const override;
 #if WITH_EDITOR
-	virtual void GetComponentPathInternal(TArray<FString>& OutChildComponentPathComponents) const override;
-	virtual void OnComponentAdded() override;
-	virtual void OnComponentRemoved() override;
+	DYNAMICMATERIAL_API virtual void GetComponentPathInternal(TArray<FString>& OutChildComponentPathComponents) const override;
+	DYNAMICMATERIAL_API virtual void OnComponentAdded() override;
+	DYNAMICMATERIAL_API virtual void OnComponentRemoved() override;
 #endif
 	//~ End UDMMaterialComponent
 };
diff --git a/Engine/Plugins/Experimental/DynamicMaterial/Source/DynamicMaterial/Public/Components/MaterialValues/DMMaterialValueBool.h b/Engine/Plugins/Experimental/DynamicMaterial/Source/DynamicMaterial/Public/Components/MaterialValues/DMMaterialValueBool.h
--- a/Engine/Plugins/Experimental/DynamicMaterial/Source/DynamicMaterial/Public/Components/MaterialValues/DMMaterialValueBool.h
+++ b/Engine/Plugins/Experimental/DynamicMaterial/Source/DynamicMaterial/Public/Components/MaterialValues/DMMaterialValueBool.h
@@ -5,45 +5,57 @@
 #include "Components/DMMaterialValue.h"
 #include "DMMaterialValueBool.generated.h"
  
-UCLASS(BlueprintType, ClassGroup = "Material Designer")
-class DYNAMICMATERIAL_API UDMMaterialValueBool : public UDMMaterialValue
+/**
+ * Component representing a bool value. Manages its own parameter.
+ */
+UCLASS(MinimalAPI, BlueprintType, ClassGroup = "Material Designer")
+class UDMMaterialValueBool : public UDMMaterialValue
 {
 	GENERATED_BODY()
  
 public:
+	UDMMaterialValueBool();
+
 	UFUNCTION(BlueprintPure, Category = "Material Designer")
 	bool GetValue() const { return Value; }
  
 	UFUNCTION(BlueprintCallable, Category = "Material Designer")
-	void SetValue(bool InValue);
+	DYNAMICMATERIAL_API void SetValue(bool InValue);
  
 #if WITH_EDITOR
 	UFUNCTION(BlueprintPure, Category = "Material Designer")
 	bool GetDefaultValue() const { return bDefaultValue; }
  
 	UFUNCTION(BlueprintCallable, Category = "Material Designer")
-	void SetDefaultValue(bool bInDefaultValue);
+	DYNAMICMATERIAL_API void SetDefaultValue(bool bInDefaultValue);
 #endif
 
 #if WITH_EDITOR
 	//~ Begin IDMJsonSerializable
-	virtual TSharedPtr<FJsonValue> JsonSerialize() const override;
-	virtual bool JsonDeserialize(const TSharedPtr<FJsonValue>& InJsonValue) override;
+	DYNAMICMATERIAL_API virtual TSharedPtr<FJsonValue> JsonSerialize() const override;
+	DYNAMICMATERIAL_API virtual bool JsonDeserialize(const TSharedPtr<FJsonValue>& InJsonValue) override;
 	//~ End IDMJsonSerializable
 #endif
  
 	//~ Begin UDMMaterialValue
-	virtual void SetMIDParameter(UMaterialInstanceDynamic* InMID) const override;
+	DYNAMICMATERIAL_API virtual void SetMIDParameter(UMaterialInstanceDynamic* InMID) const override;
 #if WITH_EDITOR
-	virtual void GenerateExpression(const TSharedRef<IDMMaterialBuildStateInterface>& InBuildState) const override;
-	virtual bool IsDefaultValue() const override;
-	virtual void ApplyDefaultValue() override;
-	virtual void ResetDefaultValue() override;
+	DYNAMICMATERIAL_API virtual void GenerateExpression(const TSharedRef<IDMMaterialBuildStateInterface>& InBuildState) const override;
+	DYNAMICMATERIAL_API virtual bool IsDefaultValue() const override;
+	DYNAMICMATERIAL_API virtual void ApplyDefaultValue() override;
+	DYNAMICMATERIAL_API virtual void ResetDefaultValue() override;
+	DYNAMICMATERIAL_API virtual UDMMaterialValueDynamic* ToDynamic(UDynamicMaterialModelDynamic* InMaterialModelDynamic) override;
 #endif
-	// ~End UDMMaterialValue
+	//~ End UDMMaterialValue
+
+	//~ Begin UDMMaterialComponent
+#if WITH_EDITOR
+	DYNAMICMATERIAL_API virtual FString GetComponentPathComponent() const override;
+#endif
+	//~ End UDMMaterialComponent
  
 protected:
-	UPROPERTY(EditInstanceOnly, BlueprintReadWrite, Getter = GetValue, Setter = SetValue, BlueprintSetter = SetValue, Category = "Material Designer",
+	UPROPERTY(EditInstanceOnly, BlueprintReadWrite, Getter, Setter, BlueprintSetter = SetValue, Category = "Material Designer",
 		meta = (AllowPrivateAccess = "true", DisplayName = "Bool"))
 	bool Value;
  
@@ -52,6 +64,4 @@ protected:
 		meta = (AllowPrivateAccess = "true"))
 	bool bDefaultValue;
 #endif
- 
-	UDMMaterialValueBool();
 };
diff --git a/Engine/Plugins/Experimental/DynamicMaterial/Source/DynamicMaterial/Public/Components/MaterialValues/DMMaterialValueColorAtlas.h b/Engine/Plugins/Experimental/DynamicMaterial/Source/DynamicMaterial/Public/Components/MaterialValues/DMMaterialValueColorAtlas.h
--- a/Engine/Plugins/Experimental/DynamicMaterial/Source/DynamicMaterial/Public/Components/MaterialValues/DMMaterialValueColorAtlas.h
+++ b/Engine/Plugins/Experimental/DynamicMaterial/Source/DynamicMaterial/Public/Components/MaterialValues/DMMaterialValueColorAtlas.h
@@ -10,59 +10,71 @@ class UCurveLinearColorAtlas;
 class UCurveLinearColor;
 #endif
 
-UCLASS(BlueprintType, ClassGroup = "Material Designer")
-class DYNAMICMATERIAL_API UDMMaterialValueColorAtlas : public UDMMaterialValue
+/**
+ * Component representing a color atlas value. Manages its own parameter.
+ */
+UCLASS(MinimalAPI, BlueprintType, ClassGroup = "Material Designer")
+class UDMMaterialValueColorAtlas : public UDMMaterialValue
 {
 	GENERATED_BODY()
 
 public:
+	DYNAMICMATERIAL_API UDMMaterialValueColorAtlas();
+
 	UFUNCTION(BlueprintPure, Category = "Material Designer")
 	float GetValue() const { return Value; }
 
 	UFUNCTION(BlueprintCallable, Category = "Material Designer")
-	void SetValue(float InValue);
+	DYNAMICMATERIAL_API void SetValue(float InValue);
 
 #if WITH_EDITOR
 	UFUNCTION(BlueprintPure, Category = "Material Designer")
 	float GetDefaultValue() const { return DefaultValue; }
 
 	UFUNCTION(BlueprintCallable, Category = "Material Designer")
-	void SetDefaultValue(float InDefaultValue);
+	DYNAMICMATERIAL_API void SetDefaultValue(float InDefaultValue);
 
 	UFUNCTION(BlueprintPure, Category = "Material Designer")
 	UCurveLinearColorAtlas* GetAtlas() const { return Atlas; }
 
 	UFUNCTION(BlueprintCallable, Category = "Material Designer")
-	void SetAtlas(UCurveLinearColorAtlas* InAtlas);
+	DYNAMICMATERIAL_API void SetAtlas(UCurveLinearColorAtlas* InAtlas);
 
 	UFUNCTION(BlueprintPure, Category = "Material Designer")
 	UCurveLinearColor* GetCurve() const { return Curve; }
 
 	UFUNCTION(BlueprintCallable, Category = "Material Designer")
-	void SetCurve(UCurveLinearColor* InCurve);
+	DYNAMICMATERIAL_API void SetCurve(UCurveLinearColor* InCurve);
 
 	virtual bool IsWholeLayerValue() const override { return true; }
 #endif
 
 #if WITH_EDITOR
 	//~ Begin IDMJsonSerializable
-	virtual TSharedPtr<FJsonValue> JsonSerialize() const override;
-	virtual bool JsonDeserialize(const TSharedPtr<FJsonValue>& InJsonValue) override;
+	DYNAMICMATERIAL_API virtual TSharedPtr<FJsonValue> JsonSerialize() const override;
+	DYNAMICMATERIAL_API virtual bool JsonDeserialize(const TSharedPtr<FJsonValue>& InJsonValue) override;
 	//~ End IDMJsonSerializable
 #endif
 
 	//~ Begin UDMMaterialValue
-	virtual void SetMIDParameter(UMaterialInstanceDynamic* InMID) const override;
+	DYNAMICMATERIAL_API virtual void SetMIDParameter(UMaterialInstanceDynamic* InMID) const override;
 #if WITH_EDITOR
-	virtual void GenerateExpression(const TSharedRef<IDMMaterialBuildStateInterface>& InBuildState) const override;
-	virtual bool IsDefaultValue() const override;
-	virtual void ApplyDefaultValue() override;
-	virtual void ResetDefaultValue() override;
+	DYNAMICMATERIAL_API virtual void GenerateExpression(const TSharedRef<IDMMaterialBuildStateInterface>& InBuildState) const override;
+	DYNAMICMATERIAL_API virtual bool IsDefaultValue() const override;
+	DYNAMICMATERIAL_API virtual void ApplyDefaultValue() override;
+	DYNAMICMATERIAL_API virtual void ResetDefaultValue() override;
+	DYNAMICMATERIAL_API virtual UDMMaterialValueDynamic* ToDynamic(UDynamicMaterialModelDynamic* InMaterialModelDynamic) override;
 #endif
 	//~ End UDMMaterialValue
 
+	//~ Begin UDMMaterialComponent
+#if WITH_EDITOR
+	DYNAMICMATERIAL_API virtual FString GetComponentPathComponent() const override;
+#endif
+	//~ End UDMMaterialComponent
+
 protected:
-	UPROPERTY(EditInstanceOnly, BlueprintReadWrite, Getter = GetValue, Setter = SetValue, BlueprintSetter = SetValue, Category = "Material Designer",
+	UPROPERTY(EditInstanceOnly, BlueprintReadWrite, Getter, Setter, BlueprintSetter = SetValue, Category = "Material Designer",
 		meta = (DisplayName = "Alpha", AllowPrivateAccess = "true", ClampMin = 0, UIMin = 0, ClampMax = 1, UIMax = 1))
 	float Value;
 
@@ -79,6 +91,4 @@ protected:
 		meta = (AllowPrivateAccess = "true", NotKeyframeable, NoCreate))
 	TObjectPtr<UCurveLinearColor> Curve;
 #endif
-
-	UDMMaterialValueColorAtlas();
 };
diff --git a/Engine/Plugins/Experimental/DynamicMaterial/Source/DynamicMaterial/Public/Components/MaterialValues/DMMaterialValueFloat.h b/Engine/Plugins/Experimental/DynamicMaterial/Source/DynamicMaterial/Public/Components/MaterialValues/DMMaterialValueFloat.h
--- a/Engine/Plugins/Experimental/DynamicMaterial/Source/DynamicMaterial/Public/Components/MaterialValues/DMMaterialValueFloat.h
+++ b/Engine/Plugins/Experimental/DynamicMaterial/Source/DynamicMaterial/Public/Components/MaterialValues/DMMaterialValueFloat.h
@@ -4,27 +4,34 @@
  
 #include "Components/DMMaterialValue.h"
 #include "DMMaterialValueFloat.generated.h"
- 
-UCLASS(Abstract, BlueprintType, ClassGroup = "Material Designer")
-class DYNAMICMATERIAL_API UDMMaterialValueFloat : public UDMMaterialValue
+
+/**
+ * Base class for float/scalar values in Materials.
+ */
+UCLASS(MinimalAPI, Abstract, BlueprintType, ClassGroup = "Material Designer")
+class UDMMaterialValueFloat : public UDMMaterialValue
 {
 	GENERATED_BODY()
  
 public:
+	UDMMaterialValueFloat();
+
+	/** The value range for any float component. If both values are the same, it is not set. */
 	UFUNCTION(BlueprintPure, Category = "Material Designer")
 	const FFloatInterval& GetValueRange() const { return ValueRange; }
- 
+
+	/** Returns true if a value range has been set. This is true if the min and max aren't the same. */
 	UFUNCTION(BlueprintPure, Category = "Material Designer")
 	bool HasValueRange() const { return ValueRange.Min != ValueRange.Max; }
- 
+
+	/** Sets the range of possible values. Set min and max to the same to disable. */
 	UFUNCTION(BlueprintCallable, Category = "Material Designer")
 	void SetValueRange(const FFloatInterval& InValueRange) { ValueRange = InValueRange; }
  
 protected:
-	UPROPERTY(VisibleInstanceOnly, BlueprintReadWrite, Getter = GetValueRange, Setter = SetValueRange, BlueprintSetter = SetValueRange, Category = "Material Designer",
+	UPROPERTY(VisibleInstanceOnly, BlueprintReadWrite, Getter, Setter, BlueprintSetter = SetValueRange, Category = "Material Designer",
 		meta = (AllowPrivateAccess = "true", DisplayName = "Value Range"))
 	FFloatInterval ValueRange;
- 
-	UDMMaterialValueFloat(); 
-	UDMMaterialValueFloat(EDMValueType InValueType);
+
+	DYNAMICMATERIAL_API UDMMaterialValueFloat(EDMValueType InValueType);
 };
diff --git a/Engine/Plugins/Experimental/DynamicMaterial/Source/DynamicMaterial/Public/Components/MaterialValues/DMMaterialValueFloat1.h b/Engine/Plugins/Experimental/DynamicMaterial/Source/DynamicMaterial/Public/Components/MaterialValues/DMMaterialValueFloat1.h
--- a/Engine/Plugins/Experimental/DynamicMaterial/Source/DynamicMaterial/Public/Components/MaterialValues/DMMaterialValueFloat1.h
+++ b/Engine/Plugins/Experimental/DynamicMaterial/Source/DynamicMaterial/Public/Components/MaterialValues/DMMaterialValueFloat1.h
@@ -5,45 +5,57 @@
 #include "DMMaterialValueFloat.h"
 #include "DMMaterialValueFloat1.generated.h"
  
-UCLASS(BlueprintType, ClassGroup = "Material Designer")
-class DYNAMICMATERIAL_API UDMMaterialValueFloat1 : public UDMMaterialValueFloat
+/**
+ * Component representing a float/scalar value. Manages its own parameter.
+ */
+UCLASS(MinimalAPI, BlueprintType, ClassGroup = "Material Designer")
+class UDMMaterialValueFloat1 : public UDMMaterialValueFloat
 {
 	GENERATED_BODY()
  
 public:
+	UDMMaterialValueFloat1();
+
 	UFUNCTION(BlueprintPure, Category = "Material Designer")
 	float GetValue() const { return Value; }
  
 	UFUNCTION(BlueprintCallable, Category = "Material Designer")
-	void SetValue(float InValue);
+	DYNAMICMATERIAL_API void SetValue(float InValue);
  
 #if WITH_EDITOR
 	UFUNCTION(BlueprintPure, Category = "Material Designer")
 	float GetDefaultValue() const { return DefaultValue; }
  
 	UFUNCTION(BlueprintCallable, Category = "Material Designer")
-	void SetDefaultValue(float InDefaultValue);
+	DYNAMICMATERIAL_API void SetDefaultValue(float InDefaultValue);
 #endif
 
 #if WITH_EDITOR
 	//~ Begin IDMJsonSerializable
-	virtual TSharedPtr<FJsonValue> JsonSerialize() const override;
-	virtual bool JsonDeserialize(const TSharedPtr<FJsonValue>& InJsonValue) override;
+	DYNAMICMATERIAL_API virtual TSharedPtr<FJsonValue> JsonSerialize() const override;
+	DYNAMICMATERIAL_API virtual bool JsonDeserialize(const TSharedPtr<FJsonValue>& InJsonValue) override;
 	//~ End IDMJsonSerializable
 #endif
  
 	//~ Begin UDMMaterialValue
-	virtual void SetMIDParameter(UMaterialInstanceDynamic* InMID) const override;
+	DYNAMICMATERIAL_API virtual void SetMIDParameter(UMaterialInstanceDynamic* InMID) const override;
 #if WITH_EDITOR
-	virtual void GenerateExpression(const TSharedRef<IDMMaterialBuildStateInterface>& InBuildState) const override;
-	virtual bool IsDefaultValue() const override;
-	virtual void ApplyDefaultValue() override;
-	virtual void ResetDefaultValue() override;
+	DYNAMICMATERIAL_API virtual void GenerateExpression(const TSharedRef<IDMMaterialBuildStateInterface>& InBuildState) const override;
+	DYNAMICMATERIAL_API virtual bool IsDefaultValue() const override;
+	DYNAMICMATERIAL_API virtual void ApplyDefaultValue() override;
+	DYNAMICMATERIAL_API virtual void ResetDefaultValue() override;
+	DYNAMICMATERIAL_API virtual UDMMaterialValueDynamic* ToDynamic(UDynamicMaterialModelDynamic* InMaterialModelDynamic) override;
 #endif
 	// ~End UDMMaterialValue
+
+	//~ Begin UDMMaterialComponent
+#if WITH_EDITOR
+	DYNAMICMATERIAL_API virtual FString GetComponentPathComponent() const override;
+#endif
+	//~ End UDMMaterialComponent
  
 protected:
-	UPROPERTY(EditInstanceOnly, BlueprintReadWrite, Getter = GetValue, Setter = SetValue, BlueprintSetter = SetValue, Category = "Material Designer",
+	UPROPERTY(EditInstanceOnly, BlueprintReadWrite, Getter, Setter, BlueprintSetter = SetValue, Category = "Material Designer",
 		meta = (AllowPrivateAccess = "true"), DisplayName = "Float")
 	float Value;
  
@@ -52,6 +64,4 @@ protected:
 		meta = (AllowPrivateAccess = "true"))
 	float DefaultValue;
 #endif
- 
-	UDMMaterialValueFloat1();
 };
diff --git a/Engine/Plugins/Experimental/DynamicMaterial/Source/DynamicMaterial/Public/Components/MaterialValues/DMMaterialValueFloat2.h b/Engine/Plugins/Experimental/DynamicMaterial/Source/DynamicMaterial/Public/Components/MaterialValues/DMMaterialValueFloat2.h
--- a/Engine/Plugins/Experimental/DynamicMaterial/Source/DynamicMaterial/Public/Components/MaterialValues/DMMaterialValueFloat2.h
+++ b/Engine/Plugins/Experimental/DynamicMaterial/Source/DynamicMaterial/Public/Components/MaterialValues/DMMaterialValueFloat2.h
@@ -5,47 +5,59 @@
 #include "DMMaterialValueFloat.h"
 #include "DMMaterialValueFloat2.generated.h"
  
-UCLASS(BlueprintType, ClassGroup = "Material Designer")
-class DYNAMICMATERIAL_API UDMMaterialValueFloat2 : public UDMMaterialValueFloat
+/**
+ * Component representing an FVector2D value. Manages its own parameter.
+ */
+UCLASS(MinimalAPI, BlueprintType, ClassGroup = "Material Designer")
+class UDMMaterialValueFloat2 : public UDMMaterialValueFloat
 {
 	GENERATED_BODY()
  
 public: 
+	UDMMaterialValueFloat2();
+
 	UFUNCTION(BlueprintPure, Category = "Material Designer")
 	const FVector2D& GetValue() const { return Value; }
  
 	UFUNCTION(BlueprintCallable, Category = "Material Designer")
-	void SetValue(const FVector2D& InValue);
+	DYNAMICMATERIAL_API void SetValue(const FVector2D& InValue);
  
 #if WITH_EDITOR
 	UFUNCTION(BlueprintPure, Category = "Material Designer")
-	FVector2D GetDefaultValue() const { return DefaultValue; }
+	const FVector2D& GetDefaultValue() const { return DefaultValue; }
  
 	UFUNCTION(BlueprintCallable, Category = "Material Designer")
-	void SetDefaultValue(FVector2D InDefaultValue);
+	DYNAMICMATERIAL_API void SetDefaultValue(const FVector2D& InDefaultValue);
 #endif
 
 #if WITH_EDITOR
 	//~ Begin IDMJsonSerializable
-	virtual TSharedPtr<FJsonValue> JsonSerialize() const override;
-	virtual bool JsonDeserialize(const TSharedPtr<FJsonValue>& InJsonValue) override;
+	DYNAMICMATERIAL_API virtual TSharedPtr<FJsonValue> JsonSerialize() const override;
+	DYNAMICMATERIAL_API virtual bool JsonDeserialize(const TSharedPtr<FJsonValue>& InJsonValue) override;
 	//~ End IDMJsonSerializable
 #endif
  
 	//~ Begin UDMMaterialValue
-	virtual void SetMIDParameter(UMaterialInstanceDynamic* InMID) const override;
+	DYNAMICMATERIAL_API virtual void SetMIDParameter(UMaterialInstanceDynamic* InMID) const override;
 #if WITH_EDITOR
-	virtual void GenerateExpression(const TSharedRef<IDMMaterialBuildStateInterface>& InBuildState) const override;
-	virtual int32 GetInnateMaskOutput(int32 OutputChannels) const override;
-	virtual bool IsDefaultValue() const override;
-	virtual void ApplyDefaultValue() override;
-	virtual void ResetDefaultValue() override;
+	DYNAMICMATERIAL_API virtual void GenerateExpression(const TSharedRef<IDMMaterialBuildStateInterface>& InBuildState) const override;
+	DYNAMICMATERIAL_API virtual int32 GetInnateMaskOutput(int32 OutputChannels) const override;
+	DYNAMICMATERIAL_API virtual bool IsDefaultValue() const override;
+	DYNAMICMATERIAL_API virtual void ApplyDefaultValue() override;
+	DYNAMICMATERIAL_API virtual void ResetDefaultValue() override;
+	DYNAMICMATERIAL_API virtual UDMMaterialValueDynamic* ToDynamic(UDynamicMaterialModelDynamic* InMaterialModelDynamic) override;
 #endif
 	// ~End UDMMaterialValue
+
+	//~ Begin UDMMaterialComponent
+#if WITH_EDITOR
+	DYNAMICMATERIAL_API virtual FString GetComponentPathComponent() const override;
+#endif
+	//~ End UDMMaterialComponent
  
 protected: 
-	UPROPERTY(EditInstanceOnly, BlueprintReadWrite, Getter = GetValue, Setter = SetValue, BlueprintSetter = SetValue, Category = "Material Designer",
-		meta = (AllowPrivateAccess = "true", DisplayName = "Vector", AllowPreserveRatio, Delta = 0.001))
+	UPROPERTY(EditInstanceOnly, BlueprintReadWrite, Getter, Setter, BlueprintSetter = SetValue, Category = "Material Designer",
+		meta = (AllowPrivateAccess = "true", DisplayName = "Vector2D", AllowPreserveRatio, Delta = 0.001))
 	FVector2D Value;
  
 #if WITH_EDITORONLY_DATA
@@ -53,6 +65,4 @@ protected:
 		meta = (AllowPrivateAccess = "true"))
 	FVector2D DefaultValue;
 #endif
- 
-	UDMMaterialValueFloat2();
 };
diff --git a/Engine/Plugins/Experimental/DynamicMaterial/Source/DynamicMaterial/Public/Components/MaterialValues/DMMaterialValueFloat3RGB.h b/Engine/Plugins/Experimental/DynamicMaterial/Source/DynamicMaterial/Public/Components/MaterialValues/DMMaterialValueFloat3RGB.h
--- a/Engine/Plugins/Experimental/DynamicMaterial/Source/DynamicMaterial/Public/Components/MaterialValues/DMMaterialValueFloat3RGB.h
+++ b/Engine/Plugins/Experimental/DynamicMaterial/Source/DynamicMaterial/Public/Components/MaterialValues/DMMaterialValueFloat3RGB.h
@@ -5,47 +5,58 @@
 #include "DMMaterialValueFloat.h"
 #include "DMMaterialValueFloat3RGB.generated.h"
  
-UCLASS(BlueprintType, ClassGroup = "Material Designer")
-class DYNAMICMATERIAL_API UDMMaterialValueFloat3RGB : public UDMMaterialValueFloat
+/**
+ * Component representing an FLinearColor (no alpha) value. Manages its own parameter.
+ */
+UCLASS(MinimalAPI, BlueprintType, ClassGroup = "Material Designer")
+class UDMMaterialValueFloat3RGB : public UDMMaterialValueFloat
 {
 	GENERATED_BODY()
  
 public: 
+	UDMMaterialValueFloat3RGB();
+
 	UFUNCTION(BlueprintPure, Category = "Material Designer")
 	const FLinearColor& GetValue() const { return Value; }
  
 	UFUNCTION(BlueprintCallable, Category = "Material Designer")
-	void SetValue(const FLinearColor& InValue);
+	DYNAMICMATERIAL_API void SetValue(const FLinearColor& InValue);
  
 #if WITH_EDITOR
 	UFUNCTION(BlueprintPure, Category = "Material Designer")
-	FLinearColor GetDefaultValue() const { return DefaultValue; }
+	const FLinearColor& GetDefaultValue() const { return DefaultValue; }
  
 	UFUNCTION(BlueprintCallable, Category = "Material Designer")
-	void SetDefaultValue(FLinearColor InDefaultValue);
+	DYNAMICMATERIAL_API void SetDefaultValue(const FLinearColor& InDefaultValue);
 #endif
 
 #if WITH_EDITOR
 	//~ Begin IDMJsonSerializable
-	virtual TSharedPtr<FJsonValue> JsonSerialize() const override;
-	virtual bool JsonDeserialize(const TSharedPtr<FJsonValue>& InJsonValue) override;
+	DYNAMICMATERIAL_API virtual TSharedPtr<FJsonValue> JsonSerialize() const override;
+	DYNAMICMATERIAL_API virtual bool JsonDeserialize(const TSharedPtr<FJsonValue>& InJsonValue) override;
 	//~ End IDMJsonSerializable
 #endif
  
 	//~ Begin UDMMaterialValue
-	virtual void SetMIDParameter(UMaterialInstanceDynamic* InMID) const override;
+	DYNAMICMATERIAL_API virtual void SetMIDParameter(UMaterialInstanceDynamic* InMID) const override;
 #if WITH_EDITOR
-	virtual void GenerateExpression(const TSharedRef<IDMMaterialBuildStateInterface>& InBuildState) const override;
-	virtual int32 GetInnateMaskOutput(int32 OutputChannels) const override;
-	virtual bool IsDefaultValue() const override;
-	virtual void ApplyDefaultValue() override;
-	virtual void ResetDefaultValue() override;
-	virtual FName GetMainPropertyName() const override { return ValueName; }
+	DYNAMICMATERIAL_API virtual void GenerateExpression(const TSharedRef<IDMMaterialBuildStateInterface>& InBuildState) const override;
+	DYNAMICMATERIAL_API virtual int32 GetInnateMaskOutput(int32 OutputChannels) const override;
+	DYNAMICMATERIAL_API virtual bool IsDefaultValue() const override;
+	DYNAMICMATERIAL_API virtual void ApplyDefaultValue() override;
+	DYNAMICMATERIAL_API virtual void ResetDefaultValue() override;
+	DYNAMICMATERIAL_API virtual UDMMaterialValueDynamic* ToDynamic(UDynamicMaterialModelDynamic* InMaterialModelDynamic) override;
 	//~ End UDMMaterialValue
 #endif
 
+	//~ Begin UDMMaterialComponent
+#if WITH_EDITOR
+	DYNAMICMATERIAL_API virtual FString GetComponentPathComponent() const override;
+#endif
+	//~ End UDMMaterialComponent
+
 protected: 
-	UPROPERTY(EditInstanceOnly, BlueprintReadWrite, Getter = GetValue, Setter = SetValue, BlueprintSetter = SetValue, Category = "Material Designer",
+	UPROPERTY(EditInstanceOnly, BlueprintReadWrite, Getter, Setter, BlueprintSetter = SetValue, Category = "Material Designer",
 		meta = (AllowPrivateAccess = "true", DisplayName = "Color", HideAlphaChannel))
 	FLinearColor Value;
  
@@ -54,6 +65,4 @@ protected:
 		meta = (AllowPrivateAccess = "true"))
 	FLinearColor DefaultValue;
 #endif
- 
-	UDMMaterialValueFloat3RGB();
 };
diff --git a/Engine/Plugins/Experimental/DynamicMaterial/Source/DynamicMaterial/Public/Components/MaterialValues/DMMaterialValueFloat3RPY.h b/Engine/Plugins/Experimental/DynamicMaterial/Source/DynamicMaterial/Public/Components/MaterialValues/DMMaterialValueFloat3RPY.h
--- a/Engine/Plugins/Experimental/DynamicMaterial/Source/DynamicMaterial/Public/Components/MaterialValues/DMMaterialValueFloat3RPY.h
+++ b/Engine/Plugins/Experimental/DynamicMaterial/Source/DynamicMaterial/Public/Components/MaterialValues/DMMaterialValueFloat3RPY.h
@@ -5,46 +5,58 @@
 #include "DMMaterialValueFloat.h"
 #include "DMMaterialValueFloat3RPY.generated.h"
  
-UCLASS(BlueprintType, ClassGroup = "Material Designer")
-class DYNAMICMATERIAL_API UDMMaterialValueFloat3RPY : public UDMMaterialValueFloat
+/**
+ * Component representing an FRotator value. Manages its own parameter.
+ */
+UCLASS(MinimalAPI, BlueprintType, ClassGroup = "Material Designer")
+class UDMMaterialValueFloat3RPY : public UDMMaterialValueFloat
 {
 	GENERATED_BODY()
  
 public: 
+	UDMMaterialValueFloat3RPY();
+
 	UFUNCTION(BlueprintPure, Category = "Material Designer")
 	const FRotator& GetValue() const { return Value; }
  
 	UFUNCTION(BlueprintCallable, Category = "Material Designer")
-	void SetValue(const FRotator& InValue);
+	DYNAMICMATERIAL_API void SetValue(const FRotator& InValue);
  
 #if WITH_EDITOR
 	UFUNCTION(BlueprintPure, Category = "Material Designer")
-	FRotator GetDefaultValue() const { return DefaultValue; }
+	const FRotator& GetDefaultValue() const { return DefaultValue; }
  
 	UFUNCTION(BlueprintCallable, Category = "Material Designer")
-	void SetDefaultValue(FRotator InDefaultValue);
+	DYNAMICMATERIAL_API void SetDefaultValue(const FRotator& InDefaultValue);
 #endif
 
 #if WITH_EDITOR
 	//~ Begin IDMJsonSerializable
-	virtual TSharedPtr<FJsonValue> JsonSerialize() const override;
-	virtual bool JsonDeserialize(const TSharedPtr<FJsonValue>& InJsonValue) override;
+	DYNAMICMATERIAL_API virtual TSharedPtr<FJsonValue> JsonSerialize() const override;
+	DYNAMICMATERIAL_API virtual bool JsonDeserialize(const TSharedPtr<FJsonValue>& InJsonValue) override;
 	//~ End IDMJsonSerializable
 #endif
  
 	//~ Begin UDMMaterialValue
-	virtual void SetMIDParameter(UMaterialInstanceDynamic* InMID) const override;
+	DYNAMICMATERIAL_API virtual void SetMIDParameter(UMaterialInstanceDynamic* InMID) const override;
 #if WITH_EDITOR
-	virtual void GenerateExpression(const TSharedRef<IDMMaterialBuildStateInterface>& InBuildState) const override;
-	virtual int32 GetInnateMaskOutput(int32 OutputChannels) const override;
-	virtual bool IsDefaultValue() const override;
-	virtual void ApplyDefaultValue() override;
-	virtual void ResetDefaultValue() override;
+	DYNAMICMATERIAL_API virtual void GenerateExpression(const TSharedRef<IDMMaterialBuildStateInterface>& InBuildState) const override;
+	DYNAMICMATERIAL_API virtual int32 GetInnateMaskOutput(int32 OutputChannels) const override;
+	DYNAMICMATERIAL_API virtual bool IsDefaultValue() const override;
+	DYNAMICMATERIAL_API virtual void ApplyDefaultValue() override;
+	DYNAMICMATERIAL_API virtual void ResetDefaultValue() override;
+	DYNAMICMATERIAL_API virtual UDMMaterialValueDynamic* ToDynamic(UDynamicMaterialModelDynamic* InMaterialModelDynamic) override;
 #endif
 	// ~End UDMMaterialValue
+
+	//~ Begin UDMMaterialComponent
+#if WITH_EDITOR
+	DYNAMICMATERIAL_API virtual FString GetComponentPathComponent() const override;
+#endif
+	//~ End UDMMaterialComponent
  
 protected: 
-	UPROPERTY(EditInstanceOnly, BlueprintReadWrite, Getter = GetValue, Setter = SetValue, BlueprintSetter = SetValue, Category = "Material Designer",
+	UPROPERTY(EditInstanceOnly, BlueprintReadWrite, Getter, Setter, BlueprintSetter = SetValue, Category = "Material Designer",
 		meta = (AllowPrivateAccess = "true", DisplayName = "Rotation"))
 	FRotator Value;
  
@@ -53,6 +65,4 @@ protected:
 		meta = (AllowPrivateAccess = "true"))
 	FRotator DefaultValue;
 #endif
- 
-	UDMMaterialValueFloat3RPY();
 };
diff --git a/Engine/Plugins/Experimental/DynamicMaterial/Source/DynamicMaterial/Public/Components/MaterialValues/DMMaterialValueFloat3XYZ.h b/Engine/Plugins/Experimental/DynamicMaterial/Source/DynamicMaterial/Public/Components/MaterialValues/DMMaterialValueFloat3XYZ.h
--- a/Engine/Plugins/Experimental/DynamicMaterial/Source/DynamicMaterial/Public/Components/MaterialValues/DMMaterialValueFloat3XYZ.h
+++ b/Engine/Plugins/Experimental/DynamicMaterial/Source/DynamicMaterial/Public/Components/MaterialValues/DMMaterialValueFloat3XYZ.h
@@ -5,47 +5,59 @@
 #include "DMMaterialValueFloat.h"
 #include "DMMaterialValueFloat3XYZ.generated.h"
  
-UCLASS(BlueprintType, ClassGroup = "Material Designer")
-class DYNAMICMATERIAL_API UDMMaterialValueFloat3XYZ : public UDMMaterialValueFloat
+/**
+ * Component representing an FVector value. Manages its own parameter.
+ */
+UCLASS(MinimalAPI, BlueprintType, ClassGroup = "Material Designer")
+class UDMMaterialValueFloat3XYZ : public UDMMaterialValueFloat
 {
 	GENERATED_BODY()
  
 public: 
+	UDMMaterialValueFloat3XYZ();
+
 	UFUNCTION(BlueprintPure, Category = "Material Designer")
 	const FVector& GetValue() const { return Value; }
  
 	UFUNCTION(BlueprintCallable, Category = "Material Designer")
-	void SetValue(const FVector& InValue);
+	DYNAMICMATERIAL_API void SetValue(const FVector& InValue);
  
 #if WITH_EDITOR
 	UFUNCTION(BlueprintPure, Category = "Material Designer")
-	FVector GetDefaultValue() const { return DefaultValue; }
+	const FVector& GetDefaultValue() const { return DefaultValue; }
  
 	UFUNCTION(BlueprintCallable, Category = "Material Designer")
-	void SetDefaultValue(FVector InDefaultValue);
+	DYNAMICMATERIAL_API void SetDefaultValue(const FVector& InDefaultValue);
 #endif
 
 #if WITH_EDITOR
 	//~ Begin IDMJsonSerializable
-	virtual TSharedPtr<FJsonValue> JsonSerialize() const override;
-	virtual bool JsonDeserialize(const TSharedPtr<FJsonValue>& InJsonValue) override;
+	DYNAMICMATERIAL_API virtual TSharedPtr<FJsonValue> JsonSerialize() const override;
+	DYNAMICMATERIAL_API virtual bool JsonDeserialize(const TSharedPtr<FJsonValue>& InJsonValue) override;
 	//~ End IDMJsonSerializable
 #endif
  
 	//~ Begin UDMMaterialValue
-	virtual void SetMIDParameter(UMaterialInstanceDynamic* InMID) const override;
+	DYNAMICMATERIAL_API virtual void SetMIDParameter(UMaterialInstanceDynamic* InMID) const override;
 #if WITH_EDITOR
-	virtual void GenerateExpression(const TSharedRef<IDMMaterialBuildStateInterface>& InBuildState) const override;
-	virtual int32 GetInnateMaskOutput(int32 OutputChannels) const override;
-	virtual bool IsDefaultValue() const override;
-	virtual void ApplyDefaultValue() override;
-	virtual void ResetDefaultValue() override;
+	DYNAMICMATERIAL_API virtual void GenerateExpression(const TSharedRef<IDMMaterialBuildStateInterface>& InBuildState) const override;
+	DYNAMICMATERIAL_API virtual int32 GetInnateMaskOutput(int32 OutputChannels) const override;
+	DYNAMICMATERIAL_API virtual bool IsDefaultValue() const override;
+	DYNAMICMATERIAL_API virtual void ApplyDefaultValue() override;
+	DYNAMICMATERIAL_API virtual void ResetDefaultValue() override;
+	DYNAMICMATERIAL_API virtual UDMMaterialValueDynamic* ToDynamic(UDynamicMaterialModelDynamic* InMaterialModelDynamic) override;
 #endif
 	// ~End UDMMaterialValue
+
+	//~ Begin UDMMaterialComponent
+#if WITH_EDITOR
+	DYNAMICMATERIAL_API virtual FString GetComponentPathComponent() const override;
+#endif
+	//~ End UDMMaterialComponent
  
 protected: 
-	UPROPERTY(EditInstanceOnly, BlueprintReadWrite, Getter = GetValue, Setter = SetValue, BlueprintSetter = SetValue, Category = "Material Designer",
-		meta = (AllowPrivateAccess = "true", DisplayName = "Vector", AllowPreserveRatio))
+	UPROPERTY(EditInstanceOnly, BlueprintReadWrite, Getter, Setter, BlueprintSetter = SetValue, Category = "Material Designer",
+		meta = (AllowPrivateAccess = "true", DisplayName = "Vector3D", AllowPreserveRatio))
 	FVector Value;
  
 #if WITH_EDITORONLY_DATA
@@ -53,6 +65,4 @@ protected:
 		meta = (AllowPrivateAccess = "true"), Category = "Material Designer")
 	FVector DefaultValue;
 #endif
- 
-	UDMMaterialValueFloat3XYZ();
 };
diff --git a/Engine/Plugins/Experimental/DynamicMaterial/Source/DynamicMaterial/Public/Components/MaterialValues/DMMaterialValueFloat4.h b/Engine/Plugins/Experimental/DynamicMaterial/Source/DynamicMaterial/Public/Components/MaterialValues/DMMaterialValueFloat4.h
--- a/Engine/Plugins/Experimental/DynamicMaterial/Source/DynamicMaterial/Public/Components/MaterialValues/DMMaterialValueFloat4.h
+++ b/Engine/Plugins/Experimental/DynamicMaterial/Source/DynamicMaterial/Public/Components/MaterialValues/DMMaterialValueFloat4.h
@@ -5,48 +5,59 @@
 #include "DMMaterialValueFloat.h"
 #include "DMMaterialValueFloat4.generated.h"
  
-UCLASS(BlueprintType, ClassGroup = "Material Designer")
-class DYNAMICMATERIAL_API UDMMaterialValueFloat4 : public UDMMaterialValueFloat
+/**
+ * Component representing a float4 value. Manages its own parameter.
+ */
+UCLASS(MinimalAPI, BlueprintType, ClassGroup = "Material Designer")
+class UDMMaterialValueFloat4 : public UDMMaterialValueFloat
 {
 	GENERATED_BODY()
  
 public: 
+	UDMMaterialValueFloat4();
+
 	UFUNCTION(BlueprintPure, Category = "Material Designer")
 	const FLinearColor& GetValue() const { return Value; }
  
 	UFUNCTION(BlueprintCallable, Category = "Material Designer")
-	void SetValue(const FLinearColor& InValue);
+	DYNAMICMATERIAL_API void SetValue(const FLinearColor& InValue);
  
 #if WITH_EDITOR
 	UFUNCTION(BlueprintPure, Category = "Material Designer")
-	FLinearColor GetDefaultValue() const { return DefaultValue; }
+	const FLinearColor& GetDefaultValue() const { return DefaultValue; }
  
 	UFUNCTION(BlueprintCallable, Category = "Material Designer")
-	void SetDefaultValue(FLinearColor InDefaultValue);
+	DYNAMICMATERIAL_API void SetDefaultValue(const FLinearColor& InDefaultValue);
 #endif
 
 #if WITH_EDITOR
 	//~ Begin IDMJsonSerializable
-	virtual TSharedPtr<FJsonValue> JsonSerialize() const override;
-	virtual bool JsonDeserialize(const TSharedPtr<FJsonValue>& InJsonValue) override;
+	DYNAMICMATERIAL_API virtual TSharedPtr<FJsonValue> JsonSerialize() const override;
+	DYNAMICMATERIAL_API virtual bool JsonDeserialize(const TSharedPtr<FJsonValue>& InJsonValue) override;
 	//~ End IDMJsonSerializable
 #endif
  
 	//~ Begin UDMMaterialValue
-	virtual void SetMIDParameter(UMaterialInstanceDynamic* InMID) const override;
+	DYNAMICMATERIAL_API virtual void SetMIDParameter(UMaterialInstanceDynamic* InMID) const override;
 #if WITH_EDITOR
-	virtual void GenerateExpression(const TSharedRef<IDMMaterialBuildStateInterface>& InBuildState) const override;
-	virtual int32 GetInnateMaskOutput(int32 OutputChannels) const override;
-	virtual bool IsDefaultValue() const override;
-	virtual void ApplyDefaultValue() override;
-	virtual void ResetDefaultValue() override;
-	virtual FName GetMainPropertyName() const override { return ValueName; }
+	DYNAMICMATERIAL_API virtual void GenerateExpression(const TSharedRef<IDMMaterialBuildStateInterface>& InBuildState) const override;
+	DYNAMICMATERIAL_API virtual int32 GetInnateMaskOutput(int32 OutputChannels) const override;
+	DYNAMICMATERIAL_API virtual bool IsDefaultValue() const override;
+	DYNAMICMATERIAL_API virtual void ApplyDefaultValue() override;
+	DYNAMICMATERIAL_API virtual void ResetDefaultValue() override;
 	virtual bool IsWholeLayerValue() const override { return true; }
+	DYNAMICMATERIAL_API virtual UDMMaterialValueDynamic* ToDynamic(UDynamicMaterialModelDynamic* InMaterialModelDynamic) override;
 	//~ End UDMMaterialValue
 #endif
 
+	//~ Begin UDMMaterialComponent
+#if WITH_EDITOR
+	DYNAMICMATERIAL_API virtual FString GetComponentPathComponent() const override;
+#endif
+	//~ End UDMMaterialComponent
+
 protected: 
-	UPROPERTY(EditInstanceOnly, BlueprintReadWrite, Getter = GetValue, Setter = SetValue, BlueprintSetter = SetValue, Category = "Material Designer",
+	UPROPERTY(EditInstanceOnly, BlueprintReadWrite, Getter, Setter, BlueprintSetter = SetValue, Category = "Material Designer",
 		meta = (AllowPrivateAccess = "true", DisplayName = "Color"))
 	FLinearColor Value;
  
@@ -55,6 +66,4 @@ protected:
 		meta = (AllowPrivateAccess = "true"))
 	FLinearColor DefaultValue;
 #endif
- 
-	UDMMaterialValueFloat4();
 };
diff --git a/Engine/Plugins/Experimental/DynamicMaterial/Source/DynamicMaterial/Public/Components/MaterialValues/DMMaterialValueRenderTarget.h b/Engine/Plugins/Experimental/DynamicMaterial/Source/DynamicMaterial/Public/Components/MaterialValues/DMMaterialValueRenderTarget.h
--- a/Engine/Plugins/Experimental/DynamicMaterial/Source/DynamicMaterial/Public/Components/MaterialValues/DMMaterialValueRenderTarget.h
+++ b/Engine/Plugins/Experimental/DynamicMaterial/Source/DynamicMaterial/Public/Components/MaterialValues/DMMaterialValueRenderTarget.h
@@ -11,44 +11,47 @@ class UTextureRenderTarget2D;
 enum ETextureRenderTargetFormat : int;
 struct FLinearColor;
 
-UCLASS(BlueprintType, ClassGroup = "Material Designer")
-class DYNAMICMATERIAL_API UDMMaterialValueRenderTarget : public UDMMaterialValueTexture
+/**
+ * Component representing a render target texture value. Manages its own parameter.
+ */
+UCLASS(MinimalAPI, BlueprintType, ClassGroup = "Material Designer")
+class UDMMaterialValueRenderTarget : public UDMMaterialValueTexture
 {
 	GENERATED_BODY()
 
 public:
-	static const FString RendererPathToken;
+	DYNAMICMATERIAL_API static const FString RendererPathToken;
 
 	UDMMaterialValueRenderTarget();
 
 	virtual ~UDMMaterialValueRenderTarget() override;
 
 	UFUNCTION(BlueprintPure, Category = "Material Designer")
-	UTextureRenderTarget2D* GetRenderTarget() const;
+	DYNAMICMATERIAL_API UTextureRenderTarget2D* GetRenderTarget() const;
 
 	UFUNCTION(BlueprintPure, Category = "Material Designer")
-	const FIntPoint& GetTextureSize() const;
+	DYNAMICMATERIAL_API const FIntPoint& GetTextureSize() const;
 
 	UFUNCTION(BlueprintCallable, Category = "Material Designer")
-	void SetTextureSize(const FIntPoint& InTextureSize);
+	DYNAMICMATERIAL_API void SetTextureSize(const FIntPoint& InTextureSize);
 
 	UFUNCTION(BlueprintPure, Category = "Material Designer")
-	ETextureRenderTargetFormat GetTextureFormat() const;
+	DYNAMICMATERIAL_API ETextureRenderTargetFormat GetTextureFormat() const;
 
 	UFUNCTION(BlueprintCallable, Category = "Material Designer")
-	void SetTextureFormat(ETextureRenderTargetFormat InTextureFormat);
+	DYNAMICMATERIAL_API void SetTextureFormat(ETextureRenderTargetFormat InTextureFormat);
 
 	UFUNCTION(BlueprintPure, Category = "Material Designer")
-	const FLinearColor& GetClearColor() const;
+	DYNAMICMATERIAL_API const FLinearColor& GetClearColor() const;
 
 	UFUNCTION(BlueprintCallable, Category = "Material Designer")
-	void SetClearColor(const FLinearColor& InClearColor);
+	DYNAMICMATERIAL_API void SetClearColor(const FLinearColor& InClearColor);
 
 	UFUNCTION(BlueprintPure, Category = "Material Designer")
-	UDMRenderTargetRenderer* GetRenderer() const;
+	DYNAMICMATERIAL_API UDMRenderTargetRenderer* GetRenderer() const;
 
 	UFUNCTION(BlueprintCallable, Category = "Material Designer")
-	void SetRenderer(UDMRenderTargetRenderer* InRenderer);
+	DYNAMICMATERIAL_API void SetRenderer(UDMRenderTargetRenderer* InRenderer);
 
 	/**
 	 * Allows outside objects to ensure our render target is valid.
@@ -61,33 +64,33 @@ public:
 
 #if WITH_EDITOR
 	//~ Begin IDMJsonSerializable
-	virtual TSharedPtr<FJsonValue> JsonSerialize() const override;
-	virtual bool JsonDeserialize(const TSharedPtr<FJsonValue>& InJsonValue) override;
+	DYNAMICMATERIAL_API virtual TSharedPtr<FJsonValue> JsonSerialize() const override;
+	DYNAMICMATERIAL_API virtual bool JsonDeserialize(const TSharedPtr<FJsonValue>& InJsonValue) override;
 	//~ End IDMJsonSerializable
 
 	//~ Begin UDMMaterialValue
 	/** Render target is handled internally. */
 	virtual bool AllowEditValue() const override { return false; }
+	DYNAMICMATERIAL_API virtual UDMMaterialValueDynamic* ToDynamic(UDynamicMaterialModelDynamic* InMaterialModelDynamic) override;
 	//~ End UDMMaterialValue
 #endif
 
 	//~ Begin UDMMaterialComponent
-	virtual void Update(EDMUpdateType InUpdateType) override;
+	DYNAMICMATERIAL_API virtual void Update(EDMUpdateType InUpdateType) override;
 #if WITH_EDITOR
-	virtual void PostEditorDuplicate(UDynamicMaterialModel* InMaterialModel, UDMMaterialComponent* InParent) override;
+	DYNAMICMATERIAL_API virtual FString GetComponentPathComponent() const override;
+	DYNAMICMATERIAL_API virtual void PostEditorDuplicate(UDynamicMaterialModel* InMaterialModel, UDMMaterialComponent* InParent) override;
 #endif
 	//~ End UDMMaterialComponent
 
 	//~ Begin UObject
 #if WITH_EDITOR
-	virtual void PostEditChangeProperty(FPropertyChangedEvent& InPropertyChangedEvent) override;
+	DYNAMICMATERIAL_API virtual void PostEditChangeProperty(FPropertyChangedEvent& InPropertyChangedEvent) override;
 #endif
-	virtual void PostLoad() override;
+	DYNAMICMATERIAL_API virtual void PostLoad() override;
 	//~ End UObject
 
 protected:
-	FDelegateHandle EndOfFrameDelegateHandle;
-
 	UPROPERTY(EditAnywhere, BlueprintReadWrite, Setter = SetTextureSize, BlueprintSetter = SetTextureSize, Category = "Material Designer|Render Target",
 		meta = (AllowPrivateAccess = "true", NotKeyframeable))
 	FIntPoint TextureSize;
@@ -104,6 +107,9 @@ protected:
 		meta = (AllowPrivateAccess = "true", NotKeyframeable, NoCreate))
 	TObjectPtr<UDMRenderTargetRenderer> Renderer;
 
+	/** Used to asynchronously update the render target. */
+	FDelegateHandle EndOfFrameDelegateHandle;
+
 	void AsyncCreateRenderTarget();
 
 	void CreateRenderTarget();
@@ -111,10 +117,10 @@ protected:
 	void UpdateRenderTarget();
 
 	//~ Begin UDMMaterialComponent
-	virtual UDMMaterialComponent* GetSubComponentByPath(FDMComponentPath& InPath, const FDMComponentPathSegment& InPathSegment) const override;
+	DYNAMICMATERIAL_API virtual UDMMaterialComponent* GetSubComponentByPath(FDMComponentPath& InPath, const FDMComponentPathSegment& InPathSegment) const override;
 #if WITH_EDITOR
-	virtual void OnComponentAdded() override;
-	virtual void OnComponentRemoved() override;
+	DYNAMICMATERIAL_API virtual void OnComponentAdded() override;
+	DYNAMICMATERIAL_API virtual void OnComponentRemoved() override;
 #endif
 	//~ End UDMMaterialComponent
 };
diff --git a/Engine/Plugins/Experimental/DynamicMaterial/Source/DynamicMaterial/Public/Components/MaterialValues/DMMaterialValueTexture.h b/Engine/Plugins/Experimental/DynamicMaterial/Source/DynamicMaterial/Public/Components/MaterialValues/DMMaterialValueTexture.h
--- a/Engine/Plugins/Experimental/DynamicMaterial/Source/DynamicMaterial/Public/Components/MaterialValues/DMMaterialValueTexture.h
+++ b/Engine/Plugins/Experimental/DynamicMaterial/Source/DynamicMaterial/Public/Components/MaterialValues/DMMaterialValueTexture.h
@@ -10,72 +10,83 @@ class UTexture;
 #if WITH_EDITOR
 namespace UE::DynamicMaterial::Private
 {
+	/** Return true if the texture has an alpha channel. */
 	bool DYNAMICMATERIAL_API HasAlpha(UTexture* InTexture);
 }
 
 DECLARE_DELEGATE_RetVal(UTexture*, FDMGetDefaultRGBTexture);
 #endif
 
-UCLASS(BlueprintType, ClassGroup = "Material Designer")
-class DYNAMICMATERIAL_API UDMMaterialValueTexture : public UDMMaterialValue
+/**
+ * Component representing a texture value. Manages its own parameter.
+ */
+UCLASS(MinimalAPI, BlueprintType, ClassGroup = "Material Designer")
+class UDMMaterialValueTexture : public UDMMaterialValue
 {
 	GENERATED_BODY()
  
 public:
+	UDMMaterialValueTexture();
+
 #if WITH_EDITOR
-	static FDMGetDefaultRGBTexture GetDefaultRGBTexture;
+	DYNAMICMATERIAL_API static FDMGetDefaultRGBTexture GetDefaultRGBTexture;
 
 	UFUNCTION(BlueprintPure, Category = "Material Designer")
-	static UDMMaterialValueTexture* CreateMaterialValueTexture(UObject* Outer, UTexture* InTexture);
+	static DYNAMICMATERIAL_API UDMMaterialValueTexture* CreateMaterialValueTexture(UObject* InOuter, UTexture* InTexture);
 #endif
 
 	UFUNCTION(BlueprintPure, Category = "Material Designer")
 	UTexture* GetValue() const { return Value; }
  
 	UFUNCTION(BlueprintCallable, Category = "Material Designer")
-	void SetValue(UTexture* InValue);
+	DYNAMICMATERIAL_API void SetValue(UTexture* InValue);
 
 #if WITH_EDITOR
+	/** Return true if the texture value has an alpha channel. */
 	UFUNCTION(BlueprintPure, Category = "Material Designer")
-	bool HasAlpha() const;
+	DYNAMICMATERIAL_API bool HasAlpha() const;
 
 	UFUNCTION(BlueprintPure, Category = "Material Designer")
 	UTexture* GetDefaultValue() const { return DefaultValue; }
 
 	UFUNCTION(BlueprintCallable, Category = "Material Designer")
-	void SetDefaultValue(UTexture* InDefaultValue);
+	DYNAMICMATERIAL_API void SetDefaultValue(UTexture* InDefaultValue);
 #endif // WITH_EDITOR
 
 #if WITH_EDITOR
 	//~ Begin IDMJsonSerializable
-	virtual TSharedPtr<FJsonValue> JsonSerialize() const override;
-	virtual bool JsonDeserialize(const TSharedPtr<FJsonValue>& InJsonValue) override;
+	DYNAMICMATERIAL_API virtual TSharedPtr<FJsonValue> JsonSerialize() const override;
+	DYNAMICMATERIAL_API virtual bool JsonDeserialize(const TSharedPtr<FJsonValue>& InJsonValue) override;
 	//~ End IDMJsonSerializable
 #endif
 
 	//~ Begin UDMMaterialValue
-	virtual void SetMIDParameter(UMaterialInstanceDynamic* InMID) const override;
+	DYNAMICMATERIAL_API virtual void SetMIDParameter(UMaterialInstanceDynamic* InMID) const override;
 #if WITH_EDITOR
-	virtual FName GetMainPropertyName() const override { return ValueName; }
-	virtual void GenerateExpression(const TSharedRef<IDMMaterialBuildStateInterface>& InBuildState) const override;
-	virtual bool IsDefaultValue() const override;
-	virtual void ApplyDefaultValue() override;
-	virtual void ResetDefaultValue() override;
+	DYNAMICMATERIAL_API virtual void GenerateExpression(const TSharedRef<IDMMaterialBuildStateInterface>& InBuildState) const override;
+	DYNAMICMATERIAL_API virtual bool IsDefaultValue() const override;
+	DYNAMICMATERIAL_API virtual void ApplyDefaultValue() override;
+	DYNAMICMATERIAL_API virtual void ResetDefaultValue() override;
+	DYNAMICMATERIAL_API virtual UDMMaterialValueDynamic* ToDynamic(UDynamicMaterialModelDynamic* InMaterialModelDynamic) override;
 #endif
 	//~ End UDMMaterialValue
 
+	//~ Begin UDMMaterialComponent
+#if WITH_EDITOR
+	DYNAMICMATERIAL_API virtual FString GetComponentPathComponent() const override;
+#endif
+	//~ End UDMMaterialComponent
+
 	//~ Begin UObject
 #if WITH_EDITOR
-	virtual void PreEditChange(FProperty* InPropertyAboutToChange) override;
-	virtual void PostEditChangeProperty(FPropertyChangedEvent& InPropertyChangedEvent) override;
+	DYNAMICMATERIAL_API virtual void PreEditChange(FProperty* InPropertyAboutToChange) override;
+	DYNAMICMATERIAL_API virtual void PostEditChangeProperty(FPropertyChangedEvent& InPropertyChangedEvent) override;
 #endif
-	virtual void PostLoad() override;
+	DYNAMICMATERIAL_API virtual void PostLoad() override;
 	//~ End UObject
 
 protected:
-	UDMMaterialValueTexture();
- 
-	UPROPERTY(EditAnywhere, BlueprintReadWrite, NoClear, Getter = GetValue, Setter = SetValue, BlueprintSetter = SetValue, Category = "Material Designer",
+	UPROPERTY(EditAnywhere, BlueprintReadWrite, NoClear, Getter, Setter, BlueprintSetter = SetValue, Category = "Material Designer",
 		meta = (DisplayThumbnail = true, DisplayName = "Texture", AllowPrivateAccess = "true", HighPriority, NoCreate))
 	TObjectPtr<UTexture> Value;
 
@@ -84,6 +95,7 @@ protected:
 		meta = (AllowPrivateAccess = "true"))
 	TObjectPtr<UTexture> DefaultValue;
 
+	/** Used for PostEditChangeProperty */
 	UPROPERTY()
 	TObjectPtr<UTexture> OldValue;
 #endif
diff --git a/Engine/Plugins/Experimental/DynamicMaterial/Source/DynamicMaterial/Public/Components/RenderTargetRenderers/DMRenderTargetTextRenderer.h b/Engine/Plugins/Experimental/DynamicMaterial/Source/DynamicMaterial/Public/Components/RenderTargetRenderers/DMRenderTargetTextRenderer.h
--- a/Engine/Plugins/Experimental/DynamicMaterial/Source/DynamicMaterial/Public/Components/RenderTargetRenderers/DMRenderTargetTextRenderer.h
+++ b/Engine/Plugins/Experimental/DynamicMaterial/Source/DynamicMaterial/Public/Components/RenderTargetRenderers/DMRenderTargetTextRenderer.h
@@ -30,275 +30,283 @@ struct FDMTextLine
 	TSharedPtr<STextBlock> Widget;
 };
 
-UCLASS(BlueprintType, ClassGroup = "Material Designer", meta = (DisplayName = "Material Designer Render Target Text Renderer"))
-class DYNAMICMATERIAL_API UDMRenderTargetTextRenderer : public UDMRenderTargetWidgetRendererBase
+/**
+ * Renderer that renders an STextBlock widget and exposes all its parameters/properties.
+ */
+UCLASS(MinimalAPI, BlueprintType, ClassGroup = "Material Designer", meta = (DisplayName = "Material Designer Render Target Text Renderer"))
+class UDMRenderTargetTextRenderer : public UDMRenderTargetWidgetRendererBase
 {
 	GENERATED_BODY()
 
 	friend struct FDMRenderTargetTextRenderer;
 
 public:
-	UDMRenderTargetTextRenderer();
+	DYNAMICMATERIAL_API UDMRenderTargetTextRenderer();
 
 	UFUNCTION(BlueprintPure, Category = "Material Designer")
-	const FSlateFontInfo& GetFontInfo() const;
+	DYNAMICMATERIAL_API const FSlateFontInfo& GetFontInfo() const;
 
 	UFUNCTION(BlueprintCallable, Category = "Material Designer")
-	void SetFontInfo(const FSlateFontInfo& InFontInfo);
+	DYNAMICMATERIAL_API void SetFontInfo(const FSlateFontInfo& InFontInfo);
 
 	UFUNCTION(BlueprintPure, Category = "Material Designer")
-	const FText& GetText() const;
+	DYNAMICMATERIAL_API const FText& GetText() const;
 
 	UFUNCTION(BlueprintCallable, Category = "Material Designer")
-	void SetText(const FText& InText);
+	DYNAMICMATERIAL_API void SetText(const FText& InText);
 
 	UFUNCTION(BlueprintPure, Category = "Material Designer")
-	const FLinearColor& GetTextColor() const;
+	DYNAMICMATERIAL_API const FLinearColor& GetTextColor() const;
 
 	UFUNCTION(BlueprintCallable, Category = "Material Designer")
-	void SetTextColor(const FLinearColor& InColor);
+	DYNAMICMATERIAL_API void SetTextColor(const FLinearColor& InColor);
 
 	UFUNCTION(BlueprintPure, Category = "Material Designer")
-	const FLinearColor& GetBackgroundColor() const;
+	DYNAMICMATERIAL_API const FLinearColor& GetBackgroundColor() const;
 
 	UFUNCTION(BlueprintCallable, Category = "Material Designer")
-	void SetBackgroundColor(const FLinearColor& InBackgroundColor);
+	DYNAMICMATERIAL_API void SetBackgroundColor(const FLinearColor& InBackgroundColor);
 
 	UFUNCTION(BlueprintPure, Category = "Material Designer")
-	bool GetHasHighlight() const;
+	DYNAMICMATERIAL_API bool GetHasHighlight() const;
 
 	UFUNCTION(BlueprintCallable, Category = "Material Designer")
-	void SetHasHighlight(bool bInHasHighlight);
+	DYNAMICMATERIAL_API void SetHasHighlight(bool bInHasHighlight);
 
 	UFUNCTION(BlueprintPure, Category = "Material Designer")
-	const FLinearColor& GetHighlightColor() const;
+	DYNAMICMATERIAL_API const FLinearColor& GetHighlightColor() const;
 
 	UFUNCTION(BlueprintCallable, Category = "Material Designer")
-	void SetHighlightColor(const FLinearColor& InHighlightColor);
+	DYNAMICMATERIAL_API void SetHighlightColor(const FLinearColor& InHighlightColor);
 
 	UFUNCTION(BlueprintPure, Category = "Material Designer")
-	bool GetHasShadow() const;
+	DYNAMICMATERIAL_API bool GetHasShadow() const;
 
 	UFUNCTION(BlueprintCallable, Category = "Material Designer")
-	void SetHasShadow(bool bInHasShadow);
+	DYNAMICMATERIAL_API void SetHasShadow(bool bInHasShadow);
 
 	UFUNCTION(BlueprintPure, Category = "Material Designer")
-	const FLinearColor& GetShadowColor() const;
+	DYNAMICMATERIAL_API const FLinearColor& GetShadowColor() const;
 
 	UFUNCTION(BlueprintCallable, Category = "Material Designer")
-	void SetShadowColor(const FLinearColor& InShadowColor);
+	DYNAMICMATERIAL_API void SetShadowColor(const FLinearColor& InShadowColor);
 
 	UFUNCTION(BlueprintPure, Category = "Material Designer")
-	const FVector2D& GetShadowOffset() const;
+	DYNAMICMATERIAL_API const FVector2D& GetShadowOffset() const;
 
 	UFUNCTION(BlueprintCallable, Category = "Material Designer")
-	void SetShadowOffset(const FVector2D& InShadowOffset);
+	DYNAMICMATERIAL_API void SetShadowOffset(const FVector2D& InShadowOffset);
 
 	UFUNCTION(BlueprintPure, Category = "Material Designer")
-	bool GetAutoWrapText() const;
+	DYNAMICMATERIAL_API bool GetAutoWrapText() const;
 
 	UFUNCTION(BlueprintCallable, Category = "Material Designer")
-	void SetAutoWrapText(bool bInAutoWrap);
+	DYNAMICMATERIAL_API void SetAutoWrapText(bool bInAutoWrap);
 
 	UFUNCTION(BlueprintPure, Category = "Material Designer")
-	float GetWrapTextAt() const;
+	DYNAMICMATERIAL_API float GetWrapTextAt() const;
 
 	UFUNCTION(BlueprintCallable, Category = "Material Designer")
-	void SetWrapTextAt(float InWrapAt);
+	DYNAMICMATERIAL_API void SetWrapTextAt(float InWrapAt);
 
 	UFUNCTION(BlueprintPure, Category = "Material Designer")
-	ETextWrappingPolicy GetWrappingPolicy() const;
+	DYNAMICMATERIAL_API ETextWrappingPolicy GetWrappingPolicy() const;
 
 	UFUNCTION(BlueprintCallable, Category = "Material Designer")
-	void SetWrappingPolicy(ETextWrappingPolicy InWrappingPolicy);
+	DYNAMICMATERIAL_API void SetWrappingPolicy(ETextWrappingPolicy InWrappingPolicy);
 
 	UFUNCTION(BlueprintPure, Category = "Material Designer")
-	ETextJustify::Type GetJustify() const;
+	DYNAMICMATERIAL_API ETextJustify::Type GetJustify() const;
 
 	UFUNCTION(BlueprintCallable, Category = "Material Designer")
-	void SetJustify(ETextJustify::Type InJustify);
+	DYNAMICMATERIAL_API void SetJustify(ETextJustify::Type InJustify);
 
 	UFUNCTION(BlueprintPure, Category = "Material Designer")
-	ETextTransformPolicy GetTransformPolicy() const;
+	DYNAMICMATERIAL_API ETextTransformPolicy GetTransformPolicy() const;
 
 	UFUNCTION(BlueprintCallable, Category = "Material Designer")
-	void SetTransformPolicy(ETextTransformPolicy InTransformPolicy);
+	DYNAMICMATERIAL_API void SetTransformPolicy(ETextTransformPolicy InTransformPolicy);
 
 	UFUNCTION(BlueprintPure, Category = "Material Designer")
-	ETextFlowDirection GetFlowDirection() const;
+	DYNAMICMATERIAL_API ETextFlowDirection GetFlowDirection() const;
 
 	UFUNCTION(BlueprintCallable, Category = "Material Designer")
-	void SetFlowDirection(ETextFlowDirection InFlowDirection);
+	DYNAMICMATERIAL_API void SetFlowDirection(ETextFlowDirection InFlowDirection);
 
 	UFUNCTION(BlueprintPure, Category = "Material Designer")
-	ETextShapingMethod GetShapingMethod() const;
+	DYNAMICMATERIAL_API ETextShapingMethod GetShapingMethod() const;
 
 	UFUNCTION(BlueprintCallable, Category = "Material Designer")
-	void SetShapingMethod(ETextShapingMethod InShapingMethod);
+	DYNAMICMATERIAL_API void SetShapingMethod(ETextShapingMethod InShapingMethod);
 
 	UFUNCTION(BlueprintPure, Category = "Material Designer")
-	const TInstancedStruct<FSlateBrush>& GetStrikeBrush() const;
+	DYNAMICMATERIAL_API const TInstancedStruct<FSlateBrush>& GetStrikeBrush() const;
 
 	UFUNCTION(BlueprintCallable, Category = "Material Designer")
-	void SetStrikeBrush(const TInstancedStruct<FSlateBrush>& InStrikeBrush);
+	DYNAMICMATERIAL_API void SetStrikeBrush(const TInstancedStruct<FSlateBrush>& InStrikeBrush);
 
 	UFUNCTION(BlueprintPure, Category = "Material Designer")
-	float GetLineHeight() const;
+	DYNAMICMATERIAL_API float GetLineHeight() const;
 
 	UFUNCTION(BlueprintCallable, Category = "Material Designer")
-	void SetLineHeight(float InLineHeight);
+	DYNAMICMATERIAL_API void SetLineHeight(float InLineHeight);
 
 	UFUNCTION(BlueprintPure, Category = "Material Designer")
-	float GetPaddingLeft() const;
+	DYNAMICMATERIAL_API float GetPaddingLeft() const;
 
 	UFUNCTION(BlueprintCallable, Category = "Material Designer")
-	void SetPaddingLeft(float InPadding);
+	DYNAMICMATERIAL_API void SetPaddingLeft(float InPadding);
 
 	UFUNCTION(BlueprintPure, Category = "Material Designer")
-	float GetPaddingRight() const;
+	DYNAMICMATERIAL_API float GetPaddingRight() const;
 
 	UFUNCTION(BlueprintCallable, Category = "Material Designer")
-	void SetPaddingRight(float InPadding);
+	DYNAMICMATERIAL_API void SetPaddingRight(float InPadding);
 
 	UFUNCTION(BlueprintPure, Category = "Material Designer")
-	float GetPaddingTop() const;
+	DYNAMICMATERIAL_API float GetPaddingTop() const;
 
 	UFUNCTION(BlueprintCallable, Category = "Material Designer")
-	void SetPaddingTop(float InPadding);
+	DYNAMICMATERIAL_API void SetPaddingTop(float InPadding);
 
 	UFUNCTION(BlueprintPure, Category = "Material Designer")
-	float GetPaddingBottom() const;
+	DYNAMICMATERIAL_API float GetPaddingBottom() const;
 
 	UFUNCTION(BlueprintCallable, Category = "Material Designer")
-	void SetPaddingBottom(float InPadding);
+	DYNAMICMATERIAL_API void SetPaddingBottom(float InPadding);
 
 	UFUNCTION(BlueprintPure, Category = "Material Designer")
-	bool IsOverridingRenderTargetSize() const;
+	DYNAMICMATERIAL_API bool IsOverridingRenderTargetSize() const;
 
 	UFUNCTION(BlueprintCallable, Category = "Material Designer")
-	void SetOverrideRenderTargetSize(bool bInOverride);
+	DYNAMICMATERIAL_API void SetOverrideRenderTargetSize(bool bInOverride);
 
 #if WITH_EDITOR
 	//~ Begin IDMJsonSerializable
-	virtual TSharedPtr<FJsonValue> JsonSerialize() const override;
-	virtual bool JsonDeserialize(const TSharedPtr<FJsonValue>& InJsonValue) override;
+	DYNAMICMATERIAL_API virtual TSharedPtr<FJsonValue> JsonSerialize() const override;
+	DYNAMICMATERIAL_API virtual bool JsonDeserialize(const TSharedPtr<FJsonValue>& InJsonValue) override;
 	//~ End IDMJsonSerializable
 
 	//~ Begin UObject
-	virtual void PostEditChangeProperty(FPropertyChangedEvent& InPropertyChangedEvent) override;
+	DYNAMICMATERIAL_API virtual void PostEditChangeProperty(FPropertyChangedEvent& InPropertyChangedEvent) override;
 	//~ End UObject
 #endif
 
 protected:
-	UPROPERTY(EditInstanceOnly, BlueprintReadWrite, Setter = SetFontInfo, BlueprintSetter = SetFontInfo, Category = "Material Designer|Text",
+	UPROPERTY(EditInstanceOnly, BlueprintReadWrite, Getter, Setter, BlueprintSetter = SetFontInfo, Category = "Material Designer|Text",
 		meta = (NotKeyframeable, AllowPrivateAccess = "true"))
 	FSlateFontInfo FontInfo;
 
-	UPROPERTY(EditInstanceOnly, BlueprintReadWrite, Setter = SetText, BlueprintSetter = SetText, Category = "Material Designer|Text",
+	UPROPERTY(EditInstanceOnly, BlueprintReadWrite, Getter, Setter, BlueprintSetter = SetText, Category = "Material Designer|Text",
 		meta = (MultiLine, NotKeyframeable, AllowPrivateAccess = "true"))
 	FText Text = FText::GetEmpty();
 
+	/** Text broken down into individual lines and their corresponding widgets. */
 	UPROPERTY()
 	TArray<FDMTextLine> Lines = {};
 
-	UPROPERTY(EditInstanceOnly, BlueprintReadWrite, Setter = SetTextColor, BlueprintSetter = SetTextColor, Category = "Material Designer|Text",
+	UPROPERTY(EditInstanceOnly, BlueprintReadWrite, Getter, Setter, BlueprintSetter = SetTextColor, Category = "Material Designer|Text",
 		meta = (AllowPrivateAccess = "true"))
 	FLinearColor TextColor = FLinearColor::White;
 
-	UPROPERTY(EditInstanceOnly, BlueprintReadWrite, Setter = SetHasHighlight, BlueprintSetter = SetHasHighlight, Category = "Material Designer|Text",
+	UPROPERTY(EditInstanceOnly, BlueprintReadWrite, Getter = GetHasHighlight, Setter = SetHasHighlight, BlueprintSetter = SetHasHighlight, Category = "Material Designer|Text",
 		meta = (NotKeyframeable, AllowPrivateAccess = "true"))
 	bool bHasHighlight = false;
 
-	UPROPERTY(EditInstanceOnly, BlueprintReadWrite, Setter = SetHighlightColor, BlueprintSetter = SetHighlightColor, Category = "Material Designer|Text",
+	UPROPERTY(EditInstanceOnly, BlueprintReadWrite, Getter, Setter, BlueprintSetter = SetHighlightColor, Category = "Material Designer|Text",
 		meta = (AllowPrivateAccess = "true", EditCondition = "bHasHighlight", EditConditionHides = true))
 	FLinearColor HighlightColor = FLinearColor::Black;
 
-	UPROPERTY(EditInstanceOnly, BlueprintReadWrite, Setter = SetHasShadow, BlueprintSetter = SetHasShadow, Category = "Material Designer|Text",
+	UPROPERTY(EditInstanceOnly, BlueprintReadWrite, Getter = GetHasShadow, Setter = SetHasShadow, BlueprintSetter = SetHasShadow, Category = "Material Designer|Text",
 		meta = (NotKeyframeable, AllowPrivateAccess = "true"))
 	bool bHasShadow = false;
 
-	UPROPERTY(EditInstanceOnly, BlueprintReadWrite, Setter = SetShadowColor, BlueprintSetter = SetShadowColor, Category = "Material Designer|Text",
+	UPROPERTY(EditInstanceOnly, BlueprintReadWrite, Getter, Setter, BlueprintSetter = SetShadowColor, Category = "Material Designer|Text",
 		meta = (AllowPrivateAccess = "true", EditCondition = "bHasShadow", EditConditionHides = true))
 	FLinearColor ShadowColor = FLinearColor::Black;
 
-	UPROPERTY(EditInstanceOnly, BlueprintReadWrite, Setter = SetShadowOffset, BlueprintSetter = SetShadowOffset, Category = "Material Designer|Text",
+	UPROPERTY(EditInstanceOnly, BlueprintReadWrite, Getter, Setter, BlueprintSetter = SetShadowOffset, Category = "Material Designer|Text",
 		meta = (AllowPrivateAccess = "true", EditCondition = "bHasShadow", EditConditionHides = true))
 	FVector2D ShadowOffset = FVector2D(1.0, 1.0);
 
-	UPROPERTY(EditInstanceOnly, BlueprintReadWrite, Setter = SetAutoWrapText, BlueprintSetter = SetAutoWrapText, Category = "Material Designer|Text",
+	UPROPERTY(EditInstanceOnly, BlueprintReadWrite, Getter = GetAutoWrapText, Setter = SetAutoWrapText, BlueprintSetter = SetAutoWrapText, Category = "Material Designer|Text",
 		meta = (NotKeyframeable, AllowPrivateAccess = "true"))
 	bool bAutoWrapText = false;
 
-	UPROPERTY(EditInstanceOnly, BlueprintReadWrite, Setter = SetWrapTextAt, BlueprintSetter = SetWrapTextAt, Category = "Material Designer|Text",
+	UPROPERTY(EditInstanceOnly, BlueprintReadWrite, Getter, Setter, BlueprintSetter = SetWrapTextAt, Category = "Material Designer|Text",
 		meta = (NotKeyframeable, AllowPrivateAccess = "true", EditCondition="bAutoWrapText", EditConditionHides))
 	float WrapTextAt = 0.f;
 
-	UPROPERTY(EditInstanceOnly, BlueprintReadWrite, Setter = SetWrappingPolicy, BlueprintSetter = SetWrappingPolicy, Category = "Material Designer|Text",
+	UPROPERTY(EditInstanceOnly, BlueprintReadWrite, Getter, Setter, BlueprintSetter = SetWrappingPolicy, Category = "Material Designer|Text",
 		meta = (NotKeyframeable, AllowPrivateAccess = "true", EditCondition = "bAutoWrapText", EditConditionHides))
 	ETextWrappingPolicy WrappingPolicy = ETextWrappingPolicy::DefaultWrapping;
 
-	UPROPERTY(EditInstanceOnly, BlueprintReadWrite, Setter = SetJustify, BlueprintSetter = SetJustify, Category = "Material Designer|Text",
+	UPROPERTY(EditInstanceOnly, BlueprintReadWrite, Getter, Setter, BlueprintSetter = SetJustify, Category = "Material Designer|Text",
 		meta = (NotKeyframeable, AllowPrivateAccess = "true"))
 	TEnumAsByte<ETextJustify::Type> Justify = ETextJustify::Left;
 
-	UPROPERTY(EditInstanceOnly, BlueprintReadWrite, Setter = SetTransformPolicy, BlueprintSetter = SetTransformPolicy, Category = "Material Designer|Text",
+	UPROPERTY(EditInstanceOnly, BlueprintReadWrite, Getter, Setter, BlueprintSetter = SetTransformPolicy, Category = "Material Designer|Text",
 		meta = (NotKeyframeable, AllowPrivateAccess = "true"))
 	ETextTransformPolicy TransformPolicy = ETextTransformPolicy::None;
 
-	UPROPERTY(EditInstanceOnly, BlueprintReadWrite, Setter = SetFlowDirection, BlueprintSetter = SetFlowDirection, Category = "Material Designer|Text",
+	UPROPERTY(EditInstanceOnly, BlueprintReadWrite, Getter, Setter, BlueprintSetter = SetFlowDirection, Category = "Material Designer|Text",
 		meta = (NotKeyframeable, AllowPrivateAccess = "true"))
 	ETextFlowDirection FlowDirection = ETextFlowDirection::Auto;
 
-	UPROPERTY(EditInstanceOnly, BlueprintReadWrite, Setter = SetShapingMethod, BlueprintSetter = SetShapingMethod, Category = "Material Designer|Text",
+	UPROPERTY(EditInstanceOnly, BlueprintReadWrite, Getter, Setter, BlueprintSetter = SetShapingMethod, Category = "Material Designer|Text",
 		meta = (NotKeyframeable, AllowPrivateAccess = "true"))
 	ETextShapingMethod ShapingMethod = ETextShapingMethod::Auto;
 
-	UPROPERTY(EditInstanceOnly, BlueprintReadWrite, Setter = SetStrikeBrush, BlueprintSetter = SetStrikeBrush, Category = "Material Designer|Text",
+	UPROPERTY(EditInstanceOnly, BlueprintReadWrite, Getter, Setter, BlueprintSetter = SetStrikeBrush, Category = "Material Designer|Text",
 		meta = (NotKeyframeable, AllowPrivateAccess = "true"))
 	TInstancedStruct<FSlateBrush> StrikeBrush;
 
 	/** Multiplier on the base font height. */
-	UPROPERTY(EditInstanceOnly, BlueprintReadWrite, Setter = SetLineHeight, BlueprintSetter = SetLineHeight, Category = "Material Designer|Text",
+	UPROPERTY(EditInstanceOnly, BlueprintReadWrite, Getter, Setter, BlueprintSetter = SetLineHeight, Category = "Material Designer|Text",
 		meta = (AllowPrivateAccess = "true", UIMin = 0.1, UIMax = 2.0))
 	float LineHeight = 1.f;
 
 	/** Extra space adding beyond the edge of the glyphs. Useful for shadows, glows, etc. */
-	UPROPERTY(EditInstanceOnly, BlueprintReadWrite, Setter = SetPaddingLeft, BlueprintSetter = SetPaddingLeft, Category = "Material Designer|Text",
+	UPROPERTY(EditInstanceOnly, BlueprintReadWrite, Getter, Setter, BlueprintSetter = SetPaddingLeft, Category = "Material Designer|Text",
 		meta = (AllowPrivateAccess = "true"))
 	float PaddingLeft = 0.f;
 
 	/** Extra space adding beyond the edge of the glyphs. Useful for shadows, glows, etc. */
-	UPROPERTY(EditInstanceOnly, BlueprintReadWrite, Setter = SetPaddingRight, BlueprintSetter = SetPaddingRight, Category = "Material Designer|Text",
+	UPROPERTY(EditInstanceOnly, BlueprintReadWrite, Getter, Setter, BlueprintSetter = SetPaddingRight, Category = "Material Designer|Text",
 		meta = (AllowPrivateAccess = "true"))
 	float PaddingRight = 0.f;
 
 	/** Extra space adding beyond the edge of the glyphs. Useful for shadows, glows, etc. */
-	UPROPERTY(EditInstanceOnly, BlueprintReadWrite, Setter = SetPaddingTop, BlueprintSetter = SetPaddingTop, Category = "Material Designer|Text",
+	UPROPERTY(EditInstanceOnly, BlueprintReadWrite, Getter, Setter, BlueprintSetter = SetPaddingTop, Category = "Material Designer|Text",
 		meta = (AllowPrivateAccess = "true"))
 	float PaddingTop = 0.f;
 
 	/** Extra space adding beyond the edge of the glyphs. Useful for shadows, glows, etc. */
-	UPROPERTY(EditInstanceOnly, BlueprintReadWrite, Setter = SetPaddingBottom, BlueprintSetter = SetPaddingBottom, Category = "Material Designer|Text",
+	UPROPERTY(EditInstanceOnly, BlueprintReadWrite, Getter, Setter, BlueprintSetter = SetPaddingBottom, Category = "Material Designer|Text",
 		meta = (AllowPrivateAccess = "true"))
 	float PaddingBottom = 0.f;
 
-	UPROPERTY(EditInstanceOnly, BlueprintReadWrite, Setter = SetOverrideRenderTargetSize, BlueprintSetter = SetOverrideRenderTargetSize, Category = "Material Designer|Text",
-		meta = (NotKeyframeable, AllowPrivateAccess = "true", ToolTip = "When true, will change the size of the render target to fit the text."))
+	/** When true, will change the size of the render target to fit the text. */
+	UPROPERTY(EditInstanceOnly, BlueprintReadWrite, Getter = IsOverridingRenderTargetSize, Setter = SetOverrideRenderTargetSize, BlueprintSetter = SetOverrideRenderTargetSize, Category = "Material Designer|Text",
+		meta = (NotKeyframeable, AllowPrivateAccess = "true"))
 	bool bOverrideRenderTargetSize = true;
 
 	bool bRecalculateTextSize = false;
 
+	/** Takes the Text and creates a new TextLines array. */
 	void UpdateTextLines();
 
+	/** Creates the widget for an individual line. */
 	TSharedRef<STextBlock> CreateTextWidget(const FText& InText) const;
 
+	/** Updates the texture size, if needed, for the given text. */
 	void SetCustomTextureSize();
 
 	//~ Begin UDMRenderTargetWidgetRendererBase
-	virtual void CreateWidgetInstance() override;
+	DYNAMICMATERIAL_API virtual void CreateWidgetInstance() override;
 	//~ End UDMRenderTargetWidgetRendererBase
 
 	//~ Begin UDMRenderTargetRenderer
-	virtual void UpdateRenderTarget_Internal() override;
+	DYNAMICMATERIAL_API virtual void UpdateRenderTarget_Internal() override;
 	//~ End UDMRenderTargetRenderer
 };
