Material Designer

diff --git a/Engine/Plugins/Experimental/DynamicMaterial/Source/DynamicMaterial/Public/Components/RenderTargetRenderers/DMRenderTargetUMGWidgetRenderer.h b/Engine/Plugins/Experimental/DynamicMaterial/Source/DynamicMaterial/Public/Components/RenderTargetRenderers/DMRenderTargetUMGWidgetRenderer.h
--- a/Engine/Plugins/Experimental/DynamicMaterial/Source/DynamicMaterial/Public/Components/RenderTargetRenderers/DMRenderTargetUMGWidgetRenderer.h
+++ b/Engine/Plugins/Experimental/DynamicMaterial/Source/DynamicMaterial/Public/Components/RenderTargetRenderers/DMRenderTargetUMGWidgetRenderer.h
@@ -8,28 +8,28 @@
 class FWidgetRenderer;
 class UWidget;
 
-UCLASS(BlueprintType, ClassGroup = "Material Designer", meta = (DisplayName = "Material Designer Render Target Widget Renderer"))
-class DYNAMICMATERIAL_API UDMRenderTargetUMGWidgetRenderer : public UDMRenderTargetWidgetRendererBase
+UCLASS(MinimalAPI, BlueprintType, ClassGroup = "Material Designer", meta = (DisplayName = "Material Designer Render Target Widget Renderer"))
+class UDMRenderTargetUMGWidgetRenderer : public UDMRenderTargetWidgetRendererBase
 {
 	GENERATED_BODY()
 
 public:
-	UDMRenderTargetUMGWidgetRenderer();
+	DYNAMICMATERIAL_API UDMRenderTargetUMGWidgetRenderer();
 
 	UFUNCTION(BlueprintPure, Category = "Material Designer")
 	TSubclassOf<UWidget> GetWidgetClass() const { return WidgetClass; }
 
 	UFUNCTION(BlueprintCallable, Category = "Material Designer")
-	void SetWidgetClass(TSubclassOf<UWidget> InWidgetClass);
+	DYNAMICMATERIAL_API void SetWidgetClass(TSubclassOf<UWidget> InWidgetClass);
 
 #if WITH_EDITOR
 	//~ Begin IDMJsonSerializable
-	virtual TSharedPtr<FJsonValue> JsonSerialize() const override;
-	virtual bool JsonDeserialize(const TSharedPtr<FJsonValue>& InJsonValue) override;
+	DYNAMICMATERIAL_API virtual TSharedPtr<FJsonValue> JsonSerialize() const override;
+	DYNAMICMATERIAL_API virtual bool JsonDeserialize(const TSharedPtr<FJsonValue>& InJsonValue) override;
 	//~ End IDMJsonSerializable
 
 	//~ Begin UObject
-	virtual void PostEditChangeProperty(FPropertyChangedEvent& InPropertyChangedEvent) override;
+	DYNAMICMATERIAL_API virtual void PostEditChangeProperty(FPropertyChangedEvent& InPropertyChangedEvent) override;
 	//~ End UObject
 #endif
 
@@ -42,6 +42,6 @@ protected:
 	TObjectPtr<UWidget> WidgetInstance;
 
 	//~ Begin UDMRenderTargetWidgetRendererBase
-	virtual void CreateWidgetInstance() override;
+	DYNAMICMATERIAL_API virtual void CreateWidgetInstance() override;
 	//~ End UDMRenderTargetWidgetRendererBase
 };
diff --git a/Engine/Plugins/Experimental/DynamicMaterial/Source/DynamicMaterial/Public/Components/RenderTargetRenderers/DMRenderTargetWidgetRendererBase.h b/Engine/Plugins/Experimental/DynamicMaterial/Source/DynamicMaterial/Public/Components/RenderTargetRenderers/DMRenderTargetWidgetRendererBase.h
--- a/Engine/Plugins/Experimental/DynamicMaterial/Source/DynamicMaterial/Public/Components/RenderTargetRenderers/DMRenderTargetWidgetRendererBase.h
+++ b/Engine/Plugins/Experimental/DynamicMaterial/Source/DynamicMaterial/Public/Components/RenderTargetRenderers/DMRenderTargetWidgetRendererBase.h
@@ -9,13 +9,16 @@ class FWidgetRenderer;
 class SWidget;
 class UWidget;
 
-UCLASS(Abstract, BlueprintType, ClassGroup = "Material Designer", meta = (DisplayName = "Material Designer Render Target Widget Renderer"))
-class DYNAMICMATERIAL_API UDMRenderTargetWidgetRendererBase : public UDMRenderTargetRenderer
+/**
+ * Renderer that renders UWidgets to render targets.
+ */
+UCLASS(MinimalAPI, Abstract, BlueprintType, ClassGroup = "Material Designer", meta = (DisplayName = "Material Designer Render Target Widget Renderer"))
+class UDMRenderTargetWidgetRendererBase : public UDMRenderTargetRenderer
 {
 	GENERATED_BODY()
 
 public:
-	UDMRenderTargetWidgetRendererBase();
+	DYNAMICMATERIAL_API UDMRenderTargetWidgetRendererBase();
 
 protected:
 	TSharedPtr<SWidget> Widget;
@@ -24,6 +27,6 @@ protected:
 	virtual void CreateWidgetInstance() PURE_VIRTUAL(UDMRenderTargetUMGWidgetRenderer::CreateWidgetInstance);
 
 	//~ Begin UDMRenderTargetRenderer
-	virtual void UpdateRenderTarget_Internal() override;
+	DYNAMICMATERIAL_API virtual void UpdateRenderTarget_Internal() override;
 	//~ End UDMRenderTargetRenderer
 };
diff --git a/Engine/Plugins/Experimental/DynamicMaterial/Source/DynamicMaterial/Public/DMComponentPath.h b/Engine/Plugins/Experimental/DynamicMaterial/Source/DynamicMaterial/Public/DMComponentPath.h
--- a/Engine/Plugins/Experimental/DynamicMaterial/Source/DynamicMaterial/Public/DMComponentPath.h
+++ b/Engine/Plugins/Experimental/DynamicMaterial/Source/DynamicMaterial/Public/DMComponentPath.h
@@ -7,41 +7,41 @@
 struct FDMComponentPathSegment;
 
 /** A path such as Name.Component.Component[2].Value */
-struct DYNAMICMATERIAL_API FDMComponentPath
+struct FDMComponentPath
 {
 	static constexpr TCHAR Separator = '.';
 	static constexpr TCHAR ParameterOpen = '[';
 	static constexpr TCHAR ParameterClose = ']';
 
 	FDMComponentPath() = delete;
-	FDMComponentPath(FStringView InPath);
-	FDMComponentPath(const FString& InPathString);
+	DYNAMICMATERIAL_API FDMComponentPath(FStringView InPath);
+	DYNAMICMATERIAL_API FDMComponentPath(const FString& InPathString);
 
-	bool IsLeaf() const;
+	DYNAMICMATERIAL_API bool IsLeaf() const;
 
 	/**
 	 * Extracts the first component and removes it from the path
 	 */
-	FDMComponentPathSegment GetFirstSegment();
+	DYNAMICMATERIAL_API FDMComponentPathSegment GetFirstSegment();
 
 protected:
 	FStringView Path;
 };
 
 /** Represents a single part of a component path */
-struct DYNAMICMATERIAL_API FDMComponentPathSegment
+struct FDMComponentPathSegment
 {
 	friend struct FDMComponentPath;
 
-	FDMComponentPathSegment(FStringView InToken, FStringView InParameter);
+	DYNAMICMATERIAL_API FDMComponentPathSegment(FStringView InToken, FStringView InParameter);
 
 	FStringView GetToken() const { return Token; }
 
-	bool HasParameter() const;
+	DYNAMICMATERIAL_API bool HasParameter() const;
 
-	bool GetParameter(int32& OutParameter) const;
+	DYNAMICMATERIAL_API bool GetParameter(int32& OutParameter) const;
 
-	bool GetParameter(FString& OutParameter) const;
+	DYNAMICMATERIAL_API bool GetParameter(FString& OutParameter) const;
 
 protected:
 	FStringView Token;
diff --git a/Engine/Plugins/Experimental/DynamicMaterial/Source/DynamicMaterial/Public/DMDefs.h b/Engine/Plugins/Experimental/DynamicMaterial/Source/DynamicMaterial/Public/DMDefs.h
--- a/Engine/Plugins/Experimental/DynamicMaterial/Source/DynamicMaterial/Public/DMDefs.h
+++ b/Engine/Plugins/Experimental/DynamicMaterial/Source/DynamicMaterial/Public/DMDefs.h
@@ -95,7 +95,7 @@ enum class EDMLocationType : uint8
  * An individual component of a connector (e.g. G from RGB.)
  */
 USTRUCT(BlueprintType, Category = "Material Designer", meta = (DisplayName = "Material Designer Stage Connector Channel"))
-struct DYNAMICMATERIAL_API FDMMaterialStageConnectorChannel
+struct FDMMaterialStageConnectorChannel
 {
 	GENERATED_BODY()
 
@@ -134,23 +134,24 @@ struct DYNAMICMATERIAL_API FDMMaterialStageConnectorChannel
 	UPROPERTY(VisibleInstanceOnly, BlueprintReadOnly, Category = "Material Designer")
 	int32 OutputChannel = WHOLE_CHANNEL;
 
-	bool operator==(const FDMMaterialStageConnectorChannel& Other) const
+	bool operator==(const FDMMaterialStageConnectorChannel& InOther) const
 	{
-		return MaterialProperty == Other.MaterialProperty
-			&& SourceIndex == Other.SourceIndex
-			&& OutputIndex == Other.OutputIndex
-			&& OutputChannel == Other.OutputChannel;
+		return MaterialProperty == InOther.MaterialProperty
+			&& SourceIndex == InOther.SourceIndex
+			&& OutputIndex == InOther.OutputIndex
+			&& OutputChannel == InOther.OutputChannel;
 	}
 };
 
-struct DYNAMICMATERIAL_API FDMUpdateGuard
+/** Used to prevent Update from running during batch operations. */
+struct FDMUpdateGuard
 {
 	FDMUpdateGuard()
 	{
 		++GuardCount;
 	}
 
-	virtual ~FDMUpdateGuard()
+	~FDMUpdateGuard()
 	{
 		--GuardCount;
 	}
@@ -161,17 +162,31 @@ struct DYNAMICMATERIAL_API FDMUpdateGuard
 	}
 
 private:
-	static int32 GuardCount;
+	DYNAMICMATERIAL_API static int32 GuardCount;
 };
 
-struct DYNAMICMATERIAL_API FDMInitializationGuard
+/** Used to prevent initialization from running during batch operations. */
+struct FDMInitializationGuard
 {
-public:
-	static bool IsInitializing();
+	static bool IsInitializing()
+	{
+		return GuardCount > 0;
+	}
+
+	FDMInitializationGuard()
+	{
+		// Used the struct name to make it clear it's a static variable.
+		++GuardCount;
+	}
 
-	FDMInitializationGuard();
-	virtual ~FDMInitializationGuard();
+	~FDMInitializationGuard()
+	{
+		if (FDMInitializationGuard::GuardCount > 0)
+		{
+			--GuardCount;
+		}
+	}
 
 private:
-	static uint32 GuardCount;
+	DYNAMICMATERIAL_API static uint32 GuardCount;
 };
diff --git a/Engine/Plugins/Experimental/DynamicMaterial/Source/DynamicMaterial/Public/DMValueDefinition.h b/Engine/Plugins/Experimental/DynamicMaterial/Source/DynamicMaterial/Public/DMValueDefinition.h
--- a/Engine/Plugins/Experimental/DynamicMaterial/Source/DynamicMaterial/Public/DMValueDefinition.h
+++ b/Engine/Plugins/Experimental/DynamicMaterial/Source/DynamicMaterial/Public/DMValueDefinition.h
@@ -10,8 +10,11 @@
 
 class UDMMaterialValue;
 
+/**
+ * Stores information about basic value types, such as EDMValueType::Float1.
+ */
 USTRUCT(BlueprintType)
-struct DYNAMICMATERIAL_API FDMValueDefinition
+struct FDMValueDefinition
 {
 	GENERATED_BODY()
 
@@ -42,20 +45,22 @@ struct DYNAMICMATERIAL_API FDMValueDefinition
 
 	EDMValueType GetType() const { return Type; }
 
+	/** Will return 0 for non-float/any-float types. */
 	uint8 GetFloatCount() const { return FloatCount; }
 
 	const FText& GetDisplayName() const { return DisplayName; }
 
 	const TArray<FText>& GetChannelNames() const { return ChannelNames; }
 
-	bool IsFloatType() const;
+	/** Returns the base class of this type. */
+	TSubclassOf<UDMMaterialValue> GetValueClass() const { return ValueClass; }
 
-	bool IsFloat3Type() const;
+	DYNAMICMATERIAL_API bool IsFloatType() const;
 
-	/** To be consistent without OutputChannel, 1 is the first channel, not 0. */
-	const FText& GetChannelName(int32 InChannel) const;
+	DYNAMICMATERIAL_API bool IsFloat3Type() const;
 
-	TSubclassOf<UDMMaterialValue> GetValueClass() const;
+	/** To be consistent without OutputChannel, 1 is the first channel, not 0. */
+	DYNAMICMATERIAL_API const FText& GetChannelName(int32 InChannel) const;
 
 protected:
 	UPROPERTY(VisibleAnywhere, BlueprintReadOnly, Category = "Material Designer")
@@ -74,32 +79,36 @@ protected:
 	TSubclassOf<UDMMaterialValue> ValueClass;
 };
 
-UCLASS(BlueprintType)
-class DYNAMICMATERIAL_API UDMValueDefinitionLibrary : public UBlueprintFunctionLibrary
+UCLASS(MinimalAPI, BlueprintType)
+class UDMValueDefinitionLibrary : public UBlueprintFunctionLibrary
 {
 	GENERATED_BODY()
 
 public:
+	/** Returns the array of available Value Types, excluding generics like None or Max. */
 	UFUNCTION(BlueprintPure, CallInEditor, Category = "Material Designer")
-	static const TArray<EDMValueType>& GetValueTypes();
+	static DYNAMICMATERIAL_API const TArray<EDMValueType>& GetValueTypes();
 
+	/** Returns a value definition for the given value type. */
 	UFUNCTION(BlueprintPure, CallInEditor, Category = "Material Designer")
-	static const FDMValueDefinition& GetValueDefinition(EDMValueType InValueType);
+	static DYNAMICMATERIAL_API const FDMValueDefinition& GetValueDefinition(EDMValueType InValueType);
 
+	/** Returns whether the given types can be connected together as input/output. */
 	UFUNCTION(BlueprintPure, CallInEditor, Category = "Material Designer", meta = (DisplayName = "Are Types Compatible"))
 	static bool BP_AreTypesCompatible(EDMValueType A, EDMValueType B, int32 AChannel, int32 BChannel)
 	{
 		return AreTypesCompatible(A, B, AChannel, BChannel);
 	}
 
-	static bool AreTypesCompatible(EDMValueType A, EDMValueType B,
-		int32 AChannel = FDMMaterialStageConnectorChannel::WHOLE_CHANNEL,
-		int32 BChannel = FDMMaterialStageConnectorChannel::WHOLE_CHANNEL);
+	/** Returns whether the given types can be connected together as input/output. */
+	DYNAMICMATERIAL_API static bool AreTypesCompatible(EDMValueType InA, EDMValueType InB,
+		int32 InAChannel = FDMMaterialStageConnectorChannel::WHOLE_CHANNEL,
+		int32 InBChannel = FDMMaterialStageConnectorChannel::WHOLE_CHANNEL);
 
 	/** Converts a number of floats into the value type */
-	static const FDMValueDefinition& GetTypeForFloatCount(uint8 FloatCount);
+	DYNAMICMATERIAL_API static const FDMValueDefinition& GetTypeForFloatCount(uint8 InFloatCount);
 
 	/** Converts a number of floats into the value type */
 	UFUNCTION(BlueprintPure, CallInEditor, Category = "Material Designer")
-	static const FDMValueDefinition& GetTypeForFloatCount(int32 FloatCount);
+	static DYNAMICMATERIAL_API const FDMValueDefinition& GetTypeForFloatCount(int32 InFloatCount);
 };
diff --git a/Engine/Plugins/Experimental/DynamicMaterial/Source/DynamicMaterial/Public/DynamicMaterialModule.h b/Engine/Plugins/Experimental/DynamicMaterial/Source/DynamicMaterial/Public/DynamicMaterialModule.h
--- a/Engine/Plugins/Experimental/DynamicMaterial/Source/DynamicMaterial/Public/DynamicMaterialModule.h
+++ b/Engine/Plugins/Experimental/DynamicMaterial/Source/DynamicMaterial/Public/DynamicMaterialModule.h
@@ -24,16 +24,20 @@ DECLARE_DELEGATE_RetVal_OneParam(TScriptInterface<IDynamicMaterialModelEditorOnl
 /**
  * Material Designer - Build your own materials in a slimline editor!
  */
-class DYNAMICMATERIAL_API FDynamicMaterialModule : public IModuleInterface
+class FDynamicMaterialModule : public IModuleInterface
 {
 public:
-	static FDynamicMaterialModule& Get();
-	static bool AreUObjectsSafe();
-	static bool IsMaterialExportEnabled();
+	DYNAMICMATERIAL_API static FDynamicMaterialModule& Get();
+
+	/** Returns true if UObjects are currently safe to use. */
+	DYNAMICMATERIAL_API static bool AreUObjectsSafe();
+
+	/** Returns true if the material export flag has been enabled. @See DM.ExportMaterials */
+	DYNAMICMATERIAL_API static bool IsMaterialExportEnabled();
 
 #if WITH_EDITOR
-	static TScriptInterface<IDynamicMaterialModelEditorOnlyDataInterface> CreateEditorOnlyData(UDynamicMaterialModel* InMaterialModel);
-	static FDMCreateEditorOnlyDataDelegate& GetCreateEditorOnlyDataDelegate() { return CreateEditorOnlyDataDelegate; }
+	DYNAMICMATERIAL_API static TScriptInterface<IDynamicMaterialModelEditorOnlyDataInterface> CreateEditorOnlyData(UDynamicMaterialModel* InMaterialModel);
+	static FDMCreateEditorOnlyDataDelegate::RegistrationType& GetCreateEditorOnlyDataDelegate() { return CreateEditorOnlyDataDelegate; }
 #endif
 
 	//~ Begin IDynamicMaterialModule
@@ -45,7 +49,7 @@ protected:
 	static bool bIsEngineExiting;
 
 #if WITH_EDITOR
-	static FDMCreateEditorOnlyDataDelegate CreateEditorOnlyDataDelegate;
+	DYNAMICMATERIAL_API static FDMCreateEditorOnlyDataDelegate CreateEditorOnlyDataDelegate;
 #endif
 	
 	static void HandleEnginePreExit();
diff --git a/Engine/Plugins/Experimental/DynamicMaterial/Source/DynamicMaterial/Public/Material/DynamicMaterialInstance.h b/Engine/Plugins/Experimental/DynamicMaterial/Source/DynamicMaterial/Public/Material/DynamicMaterialInstance.h
--- a/Engine/Plugins/Experimental/DynamicMaterial/Source/DynamicMaterial/Public/Material/DynamicMaterialInstance.h
+++ b/Engine/Plugins/Experimental/DynamicMaterial/Source/DynamicMaterial/Public/Material/DynamicMaterialInstance.h
@@ -6,40 +6,47 @@
 #include "DynamicMaterialInstance.generated.h"
 
 class UDynamicMaterialModel;
+class UDynamicMaterialModelBase;
 
 #if WITH_EDITOR
 class UDMMaterialStageInputTextureUV;
 class UDMMaterialValue;
 #endif
 
-class UMaterial;
-
-UCLASS(ClassGroup = "Material Designer", DefaultToInstanced, BlueprintType, meta = (DisplayThumbnail = "true"))
-class DYNAMICMATERIAL_API UDynamicMaterialInstance : public UMaterialInstanceDynamic
+/** A Material Instance Dynamic with its own integrated Material Designer Model that generates the based Material. */
+UCLASS(MinimalAPI, ClassGroup = "Material Designer", DefaultToInstanced, BlueprintType, meta = (DisplayThumbnail = "true"))
+class UDynamicMaterialInstance : public UMaterialInstanceDynamic
 {
 	GENERATED_BODY()
 
 public:
 	UDynamicMaterialInstance();
 
-	UDynamicMaterialModel* GetMaterialModel() { return MaterialModel; }
+	/** Returns the Material Model associated with this Material Designer Instance. */
+	UFUNCTION(BlueprintPure, Category = "Material Designer")
+	DYNAMICMATERIAL_API UDynamicMaterialModelBase* GetMaterialModelBase();
+
+	/** Resolves the base Material Model used with this Instance and returns it. */
+	UFUNCTION(BlueprintPure, Category = "Material Designer")
+	DYNAMICMATERIAL_API UDynamicMaterialModel* GetMaterialModel();
 
 #if WITH_EDITOR
-	void SetMaterialModel(UDynamicMaterialModel* InMaterialModel) { MaterialModel = InMaterialModel; }
+	/** Sets the Material Model used for this Instance. */
+	DYNAMICMATERIAL_API void SetMaterialModel(UDynamicMaterialModelBase* InMaterialModel);
 
-	void OnMaterialBuilt(UDynamicMaterialModel* InMaterialModel);
-	void InitializeMIDPublic();
+	/** Event called when the base material is build. */
+	DYNAMICMATERIAL_API void OnMaterialBuilt(UDynamicMaterialModelBase* InMaterialModel);
+
+	/** Initialises the base MID object with the current Material Model's generated material.*/
+	DYNAMICMATERIAL_API void InitializeMIDPublic();
 
 	//~ Begin UObject
-	virtual void PostDuplicate(bool bDuplicateForPIE) override;
-	virtual void PostEditImport() override;
+	DYNAMICMATERIAL_API virtual void PostDuplicate(bool bDuplicateForPIE) override;
+	DYNAMICMATERIAL_API virtual void PostEditImport() override;
 	//~ End UObject
 #endif
 
 protected:
-	UPROPERTY(VisibleInstanceOnly, BlueprintReadOnly, Instanced, DuplicateTransient, TextExportTransient, Category = "Material Designer")
-	TObjectPtr<UMaterial> BaseMaterial;
-
 	UPROPERTY(VisibleInstanceOnly, BlueprintReadOnly, Instanced, Category = "Material Designer")
-	TObjectPtr<UDynamicMaterialModel> MaterialModel;
+	TObjectPtr<UDynamicMaterialModelBase> MaterialModelBase;
 };
diff --git a/Engine/Plugins/Experimental/DynamicMaterial/Source/DynamicMaterial/Public/Model/DynamicMaterialModel.h b/Engine/Plugins/Experimental/DynamicMaterial/Source/DynamicMaterial/Public/Model/DynamicMaterialModel.h
--- a/Engine/Plugins/Experimental/DynamicMaterial/Source/DynamicMaterial/Public/Model/DynamicMaterialModel.h
+++ b/Engine/Plugins/Experimental/DynamicMaterial/Source/DynamicMaterial/Public/Model/DynamicMaterialModel.h
@@ -2,7 +2,7 @@
 
 #pragma once
 
-#include "UObject/Object.h"
+#include "Model/DynamicMaterialModelBase.h"
 #include "Containers/Array.h"
 #include "Containers/Set.h"
 #include "Templates/SubclassOf.h"
@@ -17,175 +17,189 @@ class UDMMaterialComponent;
 class UDMMaterialParameter;
 class UDMMaterialValue;
 class UDMMaterialValueFloat1;
-class UDMMaterialValueFloat2;
 class UDMTextureUV;
-class UDynamicMaterialInstance;
+class UMaterialExpression;
+class UMaterialInstanceDynamic;
 enum class EDMMaterialPropertyType : uint8;
 enum class EDMMaterialShadingModel : uint8;
 enum class EDMUpdateType : uint8;
-enum class EDMValueType : uint8;
 struct FDMComponentPath;
 
 #if WITH_EDITORONLY_DATA
 class IDynamicMaterialModelEditorOnlyDataInterface;
-class UDMMaterialProperty;
-class UDMMaterialSlot;
 class UMaterial;
-class UMaterialExpression;
 enum EBlendMode : int;
 #endif
 
 DECLARE_MULTICAST_DELEGATE_TwoParams(FDMOnValueUpdated, UDynamicMaterialModel*, UDMMaterialValue*);
 DECLARE_MULTICAST_DELEGATE_TwoParams(FDMOnTextureUVUpdated, UDynamicMaterialModel*, UDMTextureUV*);
 
-UCLASS(ClassGroup = "Material Designer", DefaultToInstanced, BlueprintType, meta = (DisplayThumbnail = "true"))
-class DYNAMICMATERIAL_API UDynamicMaterialModel : public UObject
+UCLASS(MinimalAPI, ClassGroup = "Material Designer", DefaultToInstanced, BlueprintType, meta = (DisplayThumbnail = "true"))
+class UDynamicMaterialModel : public UDynamicMaterialModelBase
 {
 	GENERATED_BODY()
 
 public:
-	static const FString ValuesPathToken;
-	static const FString ParametersPathToken;
-	static const FLazyName GlobalOpacityValueName;
-	static const FLazyName GlobalOpacityParameterName;
-	static const FLazyName GlobalMetallicValueName;
-	static const FLazyName GlobalMetallicParameterName;
-	static const FLazyName GlobalRoughnessValueName;
-	static const FLazyName GlobalRoughnessParameterName;
-	static const FLazyName GlobalNormalValueName;
-	static const FLazyName GlobalNormalParameterName;
-	static const FLazyName GlobalSpecularValueName;
-	static const FLazyName GlobalSpecularParameterName;
-	static const FLazyName GlobalAnisotropyValueName;
-	static const FLazyName GlobalAnisotropyParameterName;
-	static const FLazyName GlobalWorldPositionOffsetValueName;
-	static const FLazyName GlobalWorldPositionOffsetParameterName;
-	static const FLazyName GlobalAmbientOcclusionValueName;
-	static const FLazyName GlobalAmbientOcclusionParameterName;
-	static const FLazyName GlobalRefractionValueName;
-	static const FLazyName GlobalRefractionParameterName;
-	static const FLazyName GlobalPixelDepthOffsetValueName;
-	static const FLazyName GlobalPixelDepthOffsetParameterName;
-	static const FLazyName GlobalOffsetValueName;
-	static const FLazyName GlobalOffsetParameterName;
-	static const FLazyName GlobalTilingValueName;
-	static const FLazyName GlobalTilingParameterName;
-	static const FLazyName GlobalRotationValueName;
-	static const FLazyName GlobalRotationParameterName;
+	/** Tokens used for Components. */
+	DYNAMICMATERIAL_API static const FString ValuesPathToken;
+	DYNAMICMATERIAL_API static const FString ParametersPathToken;
+
+	/** FNames for global parameters values. */
+	DYNAMICMATERIAL_API static const FLazyName GlobalOpacityValueName;
+	DYNAMICMATERIAL_API static const FLazyName GlobalOpacityParameterName;
+	DYNAMICMATERIAL_API static const FLazyName GlobalMetallicValueName;
+	DYNAMICMATERIAL_API static const FLazyName GlobalMetallicParameterName;
+	DYNAMICMATERIAL_API static const FLazyName GlobalRoughnessValueName;
+	DYNAMICMATERIAL_API static const FLazyName GlobalRoughnessParameterName;
+	DYNAMICMATERIAL_API static const FLazyName GlobalNormalValueName;
+	DYNAMICMATERIAL_API static const FLazyName GlobalNormalParameterName;
+	DYNAMICMATERIAL_API static const FLazyName GlobalSpecularValueName;
+	DYNAMICMATERIAL_API static const FLazyName GlobalSpecularParameterName;
+	DYNAMICMATERIAL_API static const FLazyName GlobalAnisotropyValueName;
+	DYNAMICMATERIAL_API static const FLazyName GlobalAnisotropyParameterName;
+	DYNAMICMATERIAL_API static const FLazyName GlobalWorldPositionOffsetValueName;
+	DYNAMICMATERIAL_API static const FLazyName GlobalWorldPositionOffsetParameterName;
+	DYNAMICMATERIAL_API static const FLazyName GlobalAmbientOcclusionValueName;
+	DYNAMICMATERIAL_API static const FLazyName GlobalAmbientOcclusionParameterName;
+	DYNAMICMATERIAL_API static const FLazyName GlobalRefractionValueName;
+	DYNAMICMATERIAL_API static const FLazyName GlobalRefractionParameterName;
+	DYNAMICMATERIAL_API static const FLazyName GlobalPixelDepthOffsetValueName;
+	DYNAMICMATERIAL_API static const FLazyName GlobalPixelDepthOffsetParameterName;
+	DYNAMICMATERIAL_API static const FLazyName GlobalOffsetValueName;
+	DYNAMICMATERIAL_API static const FLazyName GlobalOffsetParameterName;
+	DYNAMICMATERIAL_API static const FLazyName GlobalTilingValueName;
+	DYNAMICMATERIAL_API static const FLazyName GlobalTilingParameterName;
+	DYNAMICMATERIAL_API static const FLazyName GlobalRotationValueName;
+	DYNAMICMATERIAL_API static const FLazyName GlobalRotationParameterName;
 
 	UDynamicMaterialModel();
 
-	bool IsModelValid() const;
-
-	UDynamicMaterialInstance* GetDynamicMaterialInstance() const { return DynamicMaterialInstance; }
-	void SetDynamicMaterialInstance(UDynamicMaterialInstance* InDynamicMaterialInstance);
-
+	/** Returns this if this is IsValid() and isn't isn't being destroyed. */
 	UFUNCTION(BlueprintPure, Category = "Material Designer")
-	UMaterial* GetGeneratedMaterial() const { return DynamicMaterial; }
+	DYNAMICMATERIAL_API bool IsModelValid() const;
 
-	UE_DEPRECATED(5.5, "Added to GlobalParameterValues map.")
-	UFUNCTION(BlueprintPure, Category = "Material Designer")
-	UDMMaterialValueFloat1* GetGlobalOpacityValue() const;
+	/** Returns the map of global parameter values (such as global opacity). */
+	DYNAMICMATERIAL_API const TMap<FName, TObjectPtr<UDMMaterialValue>>& GetGlobalParameterValues() const;
 
+	/** Returns a specific global parameter value (such as global opacity) of the given object name (see global parameter FNames) or nullptr. */
 	UFUNCTION(BlueprintPure, Category = "Material Designer")
-	UDMMaterialValue* GetGlobalParameterValue(FName InName) const;
+	DYNAMICMATERIAL_API UDMMaterialValue* GetGlobalParameterValue(FName InName) const;
 
+	/** Returns a Cast version of a specific global value (such as global opacity) or nullptr. */
 	template<typename InValueClass>
 	InValueClass* GetTypedGlobalParameterValue(FName InName) const
 	{
 		return Cast<InValueClass>(GetGlobalParameterValue(InName));
 	}
 
+	/** Searches the model for a specific component based on a path. */
 	UFUNCTION(BlueprintCallable, Category = "Material Designer")
-	UDMMaterialComponent* GetComponentByPath(const FString& InPath) const;
+	DYNAMICMATERIAL_API UDMMaterialComponent* GetComponentByPath(const FString& InPath) const;
 
-	UDMMaterialComponent* GetComponentByPath(FDMComponentPath& InPath) const;
+	/** Searches the model for a specific component based on a path. */
+	DYNAMICMATERIAL_API UDMMaterialComponent* GetComponentByPath(FDMComponentPath& InPath) const;
 
+	/** Searches the model for a specific component based on a path and Casts it to the given type. */
 	template<typename InComponentClass>
-	UDMMaterialComponent* GetComponentByPath(FDMComponentPath& InPath) const
+	InComponentClass* GetComponentByPath(FDMComponentPath& InPath) const
 	{
 		return Cast<InComponentClass>(GetComponentByPath(InPath));
 	}
 
+	/** Returns an array of the (non-global parameter) values used in this Model. */
+	UFUNCTION(BlueprintPure, Category = "Material Designer")
+	const TArray<UDMMaterialValue*>& GetValues() const { return Values; }
+
+	/** Returns a specific (non-global parameter) value of the given object name. */
+	UFUNCTION(BlueprintPure, Category = "Material Designer")
+	DYNAMICMATERIAL_API UDMMaterialValue* GetValueByName(FName InName) const;
+
+	/** Returns a set of components that require a runtime reference, such as texture uvs. */
+	const TSet<TObjectPtr<UDMMaterialComponent>>& GetRuntimeComponents() const { return RuntimeComponents; }
+
 #if WITH_EDITOR
 	friend class UDynamicMaterialModelEditorOnlyData;
 	friend class UDynamicMaterialModelFactory;
 
+	/** Returns the editor only data for this model. */
 	UFUNCTION(BlueprintPure, Category = "Material Designer", Meta = (DisplayName = "Get Editor Only Data"))
 	TScriptInterface<IDynamicMaterialModelEditorOnlyDataInterface> BP_GetEditorOnlyData() const { return EditorOnlyDataSI; }
 
-	IDynamicMaterialModelEditorOnlyDataInterface* GetEditorOnlyData() const;
-
-	UFUNCTION(BlueprintPure, Category = "Material Designer")
-	const TArray<UDMMaterialValue*>& GetValues() const { return Values; }
-
-	UFUNCTION(BlueprintPure, Category = "Material Designer")
-	UDMMaterialValue* GetValueByName(FName InName) const;
-
-	UFUNCTION(BlueprintPure, Category = "Material Designer")
-	UDMMaterialValue* GetValueByIndex(int32 Index) const;
-
-	UDMMaterialValue* AddValue(EDMValueType InValueType);
+	/** Returns the editor only data for this model. */
+	DYNAMICMATERIAL_API IDynamicMaterialModelEditorOnlyDataInterface* GetEditorOnlyData() const;
 
+	/** Creates a new value of the given class and returns it. */
 	UFUNCTION(BlueprintCallable, Category = "Material Designer")
-	UDMMaterialValue* AddValue(TSubclassOf<UDMMaterialValue> InValueClass);
+	DYNAMICMATERIAL_API UDMMaterialValue* AddValue(TSubclassOf<UDMMaterialValue> InValueClass);
 
-	UFUNCTION(BlueprintCallable, Category = "Material Designer")
+	/** Adds a reference to a component so they don't get garbage collected at runtime. */
 	void AddRuntimeComponentReference(UDMMaterialComponent* InValue);
 
-	UFUNCTION(BlueprintCallable, Category = "Material Designer")
+	/** Removes a reference to a component. */
 	void RemoveRuntimeComponentReference(UDMMaterialComponent* InValue);
 
-	UFUNCTION(BlueprintCallable, Category = "Material Designer")
-	void RemoveValueByName(FName InName);
-
-	UFUNCTION(BlueprintCallable, Category = "Material Designer")
-	void RemoveValueByIndex(int32 Index);
+	/** Removes a value based on its parameter name. */
+	void RemoveValueByParameterName(FName InName);
 
+	/**
+	 * Returns true if a parameter has with the given name exists on this Model.
+	 * Will not include automatically generated component parameter names.
+	 */
 	UFUNCTION(BlueprintPure, Category = "Material Designer")
-	bool HasParameterName(FName InParameterName) const;
+	DYNAMICMATERIAL_API bool HasParameterName(FName InParameterName) const;
 
 	/** Creates a new parameter and assigns it a unique name. */
-	UFUNCTION(BlueprintCallable, Category = "Material Designer")
 	UDMMaterialParameter* CreateUniqueParameter(FName InBaseName);
 
 	/** Updates the name on an existing parameter. */
-	UFUNCTION(BlueprintCallable, Category = "Material Designer")
 	void RenameParameter(UDMMaterialParameter* InParameter, FName InBaseName);
 
 	/** Removes parameter by the name assigned to this parameter object. */
-	UFUNCTION(BlueprintCallable, Category = "Material Designer")
 	void FreeParameter(UDMMaterialParameter* InParameter);
 
 	/**
 	 * Removes this specific object from the parameter map if the name is in use by a different parameter.
 	 * Returns true if, after this call, the object is not in the parameter map.
 	 */
-	UFUNCTION(BlueprintCallable, Category = "Material Designer")
 	bool ConditionalFreeParameter(UDMMaterialParameter* InParameter);
-
-	UFUNCTION(BlueprintCallable, Category = "Material Designer")
-	void ResetData();
 #endif
 
-	virtual void OnValueUpdated(UDMMaterialValue* InValue, EDMUpdateType InUpdateType);
+	/** Called by the value when it updates. Updates Material Designer Instance and triggers the delegate. */
+	void OnValueUpdated(UDMMaterialValue* InValue, EDMUpdateType InUpdateType);
 
-	virtual void OnTextureUVUpdated(UDMTextureUV* InTextureUV);
+	/** Called by the texture uv when it updates. Updates Material Designer Instance and triggers the delegate. */
+	void OnTextureUVUpdated(UDMTextureUV* InTextureUV);
 
+	/** Returns the value update delegate to (un)subscribe. */
 	FDMOnValueUpdated::RegistrationType& GetOnValueUpdateDelegate() { return OnValueUpdateDelegate; }
+
+	/** Returns the texture uv update delegate to (un)subscribe. */
 	FDMOnTextureUVUpdated::RegistrationType& GetOnTextureUVUpdateDelegate() { return OnTextureUVUpdateDelegate; }
 
+	//~ Begin UDynamicMaterialModelBase
+	virtual UDynamicMaterialModel* ResolveMaterialModel() override { return this; }
+	virtual UDynamicMaterialInstance* GetDynamicMaterialInstance() const override { return DynamicMaterialInstance; }
+	virtual void SetDynamicMaterialInstance(UDynamicMaterialInstance* InDynamicMaterialInstance) override;
+	virtual UMaterial* GetGeneratedMaterial() const override { return DynamicMaterial; }
+	DYNAMICMATERIAL_API virtual void ApplyComponents(UMaterialInstanceDynamic* InMID) override;
+	//~ End UDynamicMaterialModelBase
+
 	//~ Begin UObject
-	virtual void PostLoad() override;
+	DYNAMICMATERIAL_API virtual void PostLoad() override;
 #if WITH_EDITOR
-	virtual void PostEditUndo() override;
-	virtual void PostEditImport() override;
-	virtual void PostDuplicate(bool bDuplicateForPIE) override;
+	DYNAMICMATERIAL_API virtual void PostEditUndo() override;
+	DYNAMICMATERIAL_API virtual void PostEditImport() override;
+	DYNAMICMATERIAL_API virtual void PostDuplicate(bool bDuplicateForPIE) override;
 	//~ End UObject
 
+	/** Called to ensure that the object hierarchy is correct. */
 	void PostEditorDuplicate();
 #endif
 
+	UE_DEPRECATED(5.5, "Added to GlobalParameterValues map.")
+	UFUNCTION(BlueprintPure, Category = "Material Designer")
+	UDMMaterialValueFloat1* GetGlobalOpacityValue() const;
+
 protected:
 	/** Global values */
 	UPROPERTY(VisibleInstanceOnly, BlueprintReadOnly, Instanced, Category = "Material Designer")
@@ -195,24 +209,32 @@ protected:
 	UPROPERTY(VisibleInstanceOnly, BlueprintReadOnly, Instanced, Category = "Material Designer")
 	TSet<TObjectPtr<UDMMaterialComponent>> RuntimeComponents;
 
+	/** Map of the global parameter values, such as global opacity. */
 	UPROPERTY(VisibleInstanceOnly, BlueprintReadOnly, Category = "Material Designer")
 	TMap<FName, TObjectPtr<UDMMaterialValue>> GlobalParameterValues;
 
-	UPROPERTY(VisibleInstanceOnly, DuplicateTransient, TextExportTransient, Category = "Material Designer")
+	/** Map of parameter names to the objects representing that parameter. */
+	UPROPERTY(VisibleInstanceOnly, TextExportTransient, Category = "Material Designer")
 	TMap<FName, TWeakObjectPtr<UDMMaterialParameter>> ParameterMap;
 
-	UPROPERTY(VisibleInstanceOnly, Instanced, BlueprintReadOnly, DuplicateTransient, TextExportTransient, Category = "Material Designer")
+	/** Material generated by the Model. */
+	UPROPERTY(VisibleInstanceOnly, Instanced, BlueprintReadOnly, TextExportTransient, Category = "Material Designer")
 	TObjectPtr<UMaterial> DynamicMaterial = nullptr;
 
-	UPROPERTY(VisibleInstanceOnly, BlueprintReadOnly, DuplicateTransient, TextExportTransient, Category = "Material Designer")
+	/** Material Designer Instance representing the MID for this Model. */
+	UPROPERTY(VisibleInstanceOnly, BlueprintReadOnly, TextExportTransient, Category = "Material Designer")
 	TObjectPtr<UDynamicMaterialInstance> DynamicMaterialInstance;
 
 #if WITH_EDITORONLY_DATA
+	/** Object holding the editor-only data used by this model. */
 	UPROPERTY(VisibleInstanceOnly, BlueprintReadOnly, Category = "Material Designer")
 	TScriptInterface<IDynamicMaterialModelEditorOnlyDataInterface> EditorOnlyDataSI;
 #endif
 
+	/** Delegate called when a value is updated. @See GetOnValueUpdateDelegate. */
 	FDMOnValueUpdated OnValueUpdateDelegate;
+
+	/** Delegate called when a value is updated. @See GetOnValueUpdateDelegate. */
 	FDMOnTextureUVUpdated OnTextureUVUpdateDelegate;
 
 	/**
@@ -222,10 +244,13 @@ protected:
 	void FixGlobalParameterValues();
 
 #if WITH_EDITOR
+	/** Checks the current parameters and returns the first parameter name that is not in used (BaseName1, 2, etc, etc.) */
 	FName CreateUniqueParameterName(FName InBaseName);
 
+	/** Called to ensure that all components are correctly initialized. Also calls the editor only data version, if applicable. */
 	void ReinitComponents();
 
+	/* Makes sure the global parameter values have the correct name. */
 	void FixGlobalVars();
 #endif
 
diff --git a/Engine/Plugins/Experimental/DynamicMaterial/Source/DynamicMaterial/Public/Model/IDMMaterialBuildStateInterface.h b/Engine/Plugins/Experimental/DynamicMaterial/Source/DynamicMaterial/Public/Model/IDMMaterialBuildStateInterface.h
--- a/Engine/Plugins/Experimental/DynamicMaterial/Source/DynamicMaterial/Public/Model/IDMMaterialBuildStateInterface.h
+++ b/Engine/Plugins/Experimental/DynamicMaterial/Source/DynamicMaterial/Public/Model/IDMMaterialBuildStateInterface.h
@@ -24,10 +24,13 @@ struct IDMMaterialBuildStateInterface
 {
 	virtual ~IDMMaterialBuildStateInterface() = default;
 
+	/** Returns the dynamic material being built. */
 	virtual UMaterial* GetDynamicMaterial() const = 0;
 
+	/** Returns the material model for the material being build. */
 	virtual UDynamicMaterialModel* GetMaterialModel() const = 0;
 
+	/** Returns the build utils object for this build state, holding helpful material utilities. */
 	virtual IDMMaterialBuildUtilsInterface& GetBuildUtils() const = 0;
 
 	/** Material Values */
diff --git a/Engine/Plugins/Experimental/DynamicMaterial/Source/DynamicMaterial/Public/Model/IDMMaterialBuildUtilsInterface.h b/Engine/Plugins/Experimental/DynamicMaterial/Source/DynamicMaterial/Public/Model/IDMMaterialBuildUtilsInterface.h
--- a/Engine/Plugins/Experimental/DynamicMaterial/Source/DynamicMaterial/Public/Model/IDMMaterialBuildUtilsInterface.h
+++ b/Engine/Plugins/Experimental/DynamicMaterial/Source/DynamicMaterial/Public/Model/IDMMaterialBuildUtilsInterface.h
@@ -21,6 +21,13 @@ class UMaterialExpressionComponentMask;
 struct FDMExpressionInput;
 struct FDMMaterialStageConnection;
 
+enum class EDMMaterialParameterGroup : uint8
+{
+	Property,
+	Global,
+	NotExposed
+};
+
 /**
  * BuildUtils provides some helper functions for creating UMaterialExpressions during the material build process.
  */
@@ -28,11 +35,22 @@ struct IDMMaterialBuildUtilsInterface
 {
 	virtual ~IDMMaterialBuildUtilsInterface() = default;
 
+	/**
+	 * Creates an expression of the given class and adds it to the material.
+	 * @param InComment Comment added to the material node.
+	 * @param InAsset Object assigned to the node, such as a texture.
+	 */
 	virtual UMaterialExpression* CreateExpression(TSubclassOf<UMaterialExpression> InExpressionClass, const FString& InComment, 
 		UObject* InAsset = nullptr) const = 0;
 
-	virtual UMaterialExpression* CreateExpressionParameter(TSubclassOf<UMaterialExpression> InExpressionClass, FName InParameterName, 
-		const FString& InComment, UObject* InAsset = nullptr) const = 0;
+	/**
+	 * Creates an expression of the given class as a parameter and adds it to the material.
+	 * @param InParameterName The name of the parameter.
+	 * @param InComment Comment added to the material node.
+	 * @param InAsset Object assigned to the node, such as a texture.
+	 */
+	virtual UMaterialExpression* CreateExpressionParameter(TSubclassOf<UMaterialExpression> InExpressionClass, FName InParameterName,
+		EDMMaterialParameterGroup InParameterGroup, const FString& InComment, UObject* InAsset = nullptr) const = 0;
 
 	template<typename InExpressionClass
 		UE_REQUIRES(std::derived_from<InExpressionClass, UMaterialExpression>)>
@@ -43,9 +61,10 @@ struct IDMMaterialBuildUtilsInterface
 
 	template<typename InExpressionClass
 		UE_REQUIRES(std::derived_from<InExpressionClass, UMaterialExpression>)>
-	InExpressionClass* CreateExpressionParameter(FName InParameterName, const FString& InComment, UObject* InAsset = nullptr) const
+	InExpressionClass* CreateExpressionParameter(FName InParameterName, EDMMaterialParameterGroup InParameterGroup, const FString& InComment, UObject* InAsset = nullptr) const
 	{
-		return Cast<InExpressionClass>(CreateExpressionParameter(TSubclassOf<UMaterialExpression>(InExpressionClass::StaticClass()), InParameterName, InComment, InAsset));
+		return Cast<InExpressionClass>(CreateExpressionParameter(TSubclassOf<UMaterialExpression>(InExpressionClass::StaticClass()), 
+			InParameterName, InParameterGroup, InComment, InAsset));
 	}
 
 	/** Creates a set of expressions merging all the inputs for each channel into a single output */
@@ -53,15 +72,18 @@ struct IDMMaterialBuildUtilsInterface
 		int32 InStageSourceInputIdx, const TArray<UDMMaterialStageInput*>& InStageInputs, int32& OutOutputIndex,
 		int32& OutOutputChannel) const = 0;
 
+	/** Creates a set of expressions that display this material stage input. */
 	virtual TArray<UMaterialExpression*> CreateExpressionInput(UDMMaterialStageInput* InInput) const = 0;
 
 	/** Creates a set of expressions merging all the inputs for each channel into a single output */
 	virtual UMaterialExpressionComponentMask* CreateExpressionBitMask(UMaterialExpression* InExpression, int32 InOutputIndex, 
 		int32 InOutputChannels) const = 0;
 
+	/** Creates an append expression, joining the output of 2 other expressions into a single vector. */
 	virtual UMaterialExpressionAppendVector* CreateExpressionAppend(UMaterialExpression* InExpressionA, int32 InOutputIndexA,
 		UMaterialExpression* InExpressionB, int32 InOutputIndexB) const = 0;
 
+	/** Updates a preview material, assigning the output of the "last expression" to an appropriate material property. */
 	virtual void UpdatePreviewMaterial(UMaterialExpression* InLastExpression, int32 InOutputIndex, int32 InOutputChannel, 
 		int32 InSize) const = 0;
 };
diff --git a/Engine/Plugins/Experimental/DynamicMaterial/Source/DynamicMaterial/Public/Model/IDynamicMaterialModelEditorOnlyDataInterface.h b/Engine/Plugins/Experimental/DynamicMaterial/Source/DynamicMaterial/Public/Model/IDynamicMaterialModelEditorOnlyDataInterface.h
--- a/Engine/Plugins/Experimental/DynamicMaterial/Source/DynamicMaterial/Public/Model/IDynamicMaterialModelEditorOnlyDataInterface.h
+++ b/Engine/Plugins/Experimental/DynamicMaterial/Source/DynamicMaterial/Public/Model/IDynamicMaterialModelEditorOnlyDataInterface.h
@@ -21,6 +21,7 @@ enum class EDMUpdateType : uint8;
 struct FDMComponentPathSegment;
 struct FDMComponentPath;
 
+/** Interface for the editor-only data so that editor-only parts of runtime components can interact with editor-only features. */
 UINTERFACE(MinimalAPI, BlueprintType, meta = (CannotImplementInterfaceInBlueprint))
 class UDynamicMaterialModelEditorOnlyDataInterface : public UInterface
 {
@@ -34,30 +35,40 @@ class IDynamicMaterialModelEditorOnlyDataInterface
 	friend class UDynamicMaterialModel;
 
 public:
-	UFUNCTION(BlueprintCallable, Category = "Material Designer")
-	virtual void RequestMaterialBuild() {}
-
+	/** Called to ensure that the object hierarchy is correct. */
 	virtual void PostEditorDuplicate() = 0;
 
+	/** Called when a value is updated. */
 	virtual void OnValueUpdated(UDMMaterialValue* InValue, EDMUpdateType InUpdateType) = 0;
 
+	/** Called when a value is added or removed. */
 	virtual void OnValueListUpdate() = 0;
 
+	/** Called when a texture uv is updated. */
 	virtual void OnTextureUVUpdated(UDMTextureUV* InTextureUV) = 0;
 
+	/** Called to load deprecated model data from 5.2 to 5.3 materials. */
 	virtual void LoadDeprecatedModelData(UDynamicMaterialModel* InMaterialModel) = 0;
 
 #if WITH_EDITOR
+	/** Called when the model needs to have the material rebuild. */
+	UFUNCTION(BlueprintCallable, Category = "Material Designer")
+	virtual void RequestMaterialBuild() PURE_VIRTUAL(UDynamicMaterialModelEditorOnlyDataInterface::RequestMaterialBuild)
+
+	/** Called to create the build state for building materials. */
 	virtual TSharedRef<IDMMaterialBuildStateInterface> CreateBuildStateInterface(UMaterial* InMaterialToBuild) const = 0;
 
+	/** Sets the component on a material property (such as a global parameter). */
 	virtual void SetPropertyComponent(EDMMaterialPropertyType InPropertyType, FName InComponentName, UDMMaterialComponent* InComponent) = 0;
 #endif
 
+	/** Searches the model editor only data for a specific component based on a path. */
 	virtual UDMMaterialComponent* GetSubComponentByPath(FDMComponentPath& InPath) const = 0;
+
+	/** Searches the model editor only data for a specific component based on a path. */
 	virtual UDMMaterialComponent* GetSubComponentByPath(FDMComponentPath& InPath, const FDMComponentPathSegment& InPathSegment) const = 0;
 
 protected:
+	/** Called to ensure that all editor only data is correctly initialized. */
 	virtual void ReinitComponents() = 0;
-
-	virtual void ResetData() = 0;
 };
\ No newline at end of file
diff --git a/Engine/Plugins/Experimental/DynamicMaterial/Source/DynamicMaterial/Public/Utils/DMJsonUtils.h b/Engine/Plugins/Experimental/DynamicMaterial/Source/DynamicMaterial/Public/Utils/DMJsonUtils.h
--- a/Engine/Plugins/Experimental/DynamicMaterial/Source/DynamicMaterial/Public/Utils/DMJsonUtils.h
+++ b/Engine/Plugins/Experimental/DynamicMaterial/Source/DynamicMaterial/Public/Utils/DMJsonUtils.h
@@ -12,7 +12,7 @@ class FJsonValue;
 class UClass;
 class UTexture;
 
-/** Offers custom serialization options for class */
+/** Offers custom serialization options for class. */
 UINTERFACE(MinimalAPI, Blueprintable, BlueprintType)
 class UDMJsonSerializable : public UInterface
 {
@@ -29,18 +29,19 @@ public:
 };
 
 #if WITH_EDITOR
-class DYNAMICMATERIAL_API FDMJsonUtils
+/** Wrapper library to allow quicker (de)serialization to JSON. */
+class FDMJsonUtils
 {
 public:
-	static TSharedPtr<FJsonValue> Serialize(bool bInValue);
-	static TSharedPtr<FJsonValue> Serialize(const FString& InString);
-	static TSharedPtr<FJsonValue> Serialize(const FText& InText);
-	static TSharedPtr<FJsonValue> Serialize(const FName& InName);
-	static TSharedPtr<FJsonValue> Serialize(const UClass* InClass);
-	static TSharedPtr<FJsonValue> Serialize(const UScriptStruct* InScriptStruct, const void* InData);
-	static TSharedPtr<FJsonValue> Serialize(const UObject* InObject);
-	static TSharedPtr<FJsonValue> Serialize(const FObjectPtr& InObject);
-	static TSharedPtr<FJsonValue> Serialize(const TMap<FString, TSharedPtr<FJsonValue>>& InMap);
+	DYNAMICMATERIAL_API static TSharedPtr<FJsonValue> Serialize(bool bInValue);
+	DYNAMICMATERIAL_API static TSharedPtr<FJsonValue> Serialize(const FString& InString);
+	DYNAMICMATERIAL_API static TSharedPtr<FJsonValue> Serialize(const FText& InText);
+	DYNAMICMATERIAL_API static TSharedPtr<FJsonValue> Serialize(const FName& InName);
+	DYNAMICMATERIAL_API static TSharedPtr<FJsonValue> Serialize(const UClass* InClass);
+	DYNAMICMATERIAL_API static TSharedPtr<FJsonValue> Serialize(const UScriptStruct* InScriptStruct, const void* InData);
+	DYNAMICMATERIAL_API static TSharedPtr<FJsonValue> Serialize(const UObject* InObject);
+	DYNAMICMATERIAL_API static TSharedPtr<FJsonValue> Serialize(const FObjectPtr& InObject);
+	DYNAMICMATERIAL_API static TSharedPtr<FJsonValue> Serialize(const TMap<FString, TSharedPtr<FJsonValue>>& InMap);
 
 	template<typename InArithmeticType
 		UE_REQUIRES(TIsArithmetic<InArithmeticType>::Value)>
@@ -70,16 +71,16 @@ public:
 		return Serialize(InClass.Get());
 	}
 
-	static bool Deserialize(const TSharedPtr<FJsonValue>& InJsonValue, bool& bOutValue);
-	static bool Deserialize(const TSharedPtr<FJsonValue>& InJsonValue, FString& OutString);
-	static bool Deserialize(const TSharedPtr<FJsonValue>& InJsonValue, FText& OutText);
-	static bool Deserialize(const TSharedPtr<FJsonValue>& InJsonValue, FName& OutName);
-	static bool Deserialize(const TSharedPtr<FJsonValue>& InJsonValue, UClass*& OutClass);
-	static bool Deserialize(const TSharedPtr<FJsonValue>& InJsonValue, const UScriptStruct* InScriptStruct, void* OutData);
-	static bool Deserialize(const TSharedPtr<FJsonValue>& InJsonValue, const UEnum* InEnum, int64& OutValue);
-	static bool Deserialize(const TSharedPtr<FJsonValue>& InJsonValue, UObject*& OutObject, UObject* InOuter = nullptr);
-	static bool Deserialize(const TSharedPtr<FJsonValue>& InJsonValue, FObjectPtr& OutObject, UObject* InOuter = nullptr);
-	static bool Deserialize(const TSharedPtr<FJsonValue>& InJsonValue, TMap<FString, TSharedPtr<FJsonValue>>& OutMap);
+	DYNAMICMATERIAL_API static bool Deserialize(const TSharedPtr<FJsonValue>& InJsonValue, bool& bOutValue);
+	DYNAMICMATERIAL_API static bool Deserialize(const TSharedPtr<FJsonValue>& InJsonValue, FString& OutString);
+	DYNAMICMATERIAL_API static bool Deserialize(const TSharedPtr<FJsonValue>& InJsonValue, FText& OutText);
+	DYNAMICMATERIAL_API static bool Deserialize(const TSharedPtr<FJsonValue>& InJsonValue, FName& OutName);
+	DYNAMICMATERIAL_API static bool Deserialize(const TSharedPtr<FJsonValue>& InJsonValue, UClass*& OutClass);
+	DYNAMICMATERIAL_API static bool Deserialize(const TSharedPtr<FJsonValue>& InJsonValue, const UScriptStruct* InScriptStruct, void* OutData);
+	DYNAMICMATERIAL_API static bool Deserialize(const TSharedPtr<FJsonValue>& InJsonValue, const UEnum* InEnum, int64& OutValue);
+	DYNAMICMATERIAL_API static bool Deserialize(const TSharedPtr<FJsonValue>& InJsonValue, UObject*& OutObject, UObject* InOuter = nullptr);
+	DYNAMICMATERIAL_API static bool Deserialize(const TSharedPtr<FJsonValue>& InJsonValue, FObjectPtr& OutObject, UObject* InOuter = nullptr);
+	DYNAMICMATERIAL_API static bool Deserialize(const TSharedPtr<FJsonValue>& InJsonValue, TMap<FString, TSharedPtr<FJsonValue>>& OutMap);
 
 	template<typename InArithmeticType
 		UE_REQUIRES(TIsArithmetic<InArithmeticType>::Value)>
@@ -162,7 +163,7 @@ public:
 	}
 
 private:
-	static TSharedPtr<FJsonValue> SerializeNumber(double InNumber);
-	static bool DeserializeNumber(const TSharedPtr<FJsonValue>& InJsonValue, double& OutNumber);
+	DYNAMICMATERIAL_API static TSharedPtr<FJsonValue> SerializeNumber(double InNumber);
+	DYNAMICMATERIAL_API static bool DeserializeNumber(const TSharedPtr<FJsonValue>& InJsonValue, double& OutNumber);
 };
 #endif
diff --git a/Engine/Plugins/Experimental/DynamicMaterial/Source/DynamicMaterial/Public/Utils/DMUtils.h b/Engine/Plugins/Experimental/DynamicMaterial/Source/DynamicMaterial/Public/Utils/DMUtils.h
--- a/Engine/Plugins/Experimental/DynamicMaterial/Source/DynamicMaterial/Public/Utils/DMUtils.h
+++ b/Engine/Plugins/Experimental/DynamicMaterial/Source/DynamicMaterial/Public/Utils/DMUtils.h
@@ -9,20 +9,25 @@ enum class EDMMaterialPropertyType : uint8;
 enum class EDMTextureSetMaterialProperty : uint8;
 enum EMaterialProperty : int;
 
-struct DYNAMICMATERIAL_API FDMUtils
+/** General utilities for use with the Material Designer. */
+struct FDMUtils
 {
 #if WITH_EDITOR
 	/** Designed to be used with preprocessor macros. */
-	static FString CreateNodeComment(const ANSICHAR* InFile, int InLine, const ANSICHAR* InFunction, const FString* InComment = nullptr);
+	DYNAMICMATERIAL_API static FString CreateNodeComment(const ANSICHAR* InFile, int InLine, const ANSICHAR* InFunction, const FString* InComment = nullptr);
 #endif
 
-	static EDMTextureSetMaterialProperty MaterialPropertyTypeToTextureSetMaterialProperty(EDMMaterialPropertyType InPropertyType);
+	/** Converts EDMMaterialPropertyType to EDMTextureSetMaterialProperty. */
+	DYNAMICMATERIAL_API static EDMTextureSetMaterialProperty MaterialPropertyTypeToTextureSetMaterialProperty(EDMMaterialPropertyType InPropertyType);
 
-	static EDMMaterialPropertyType TextureSetMaterialPropertyToMaterialPropertyType(EDMTextureSetMaterialProperty InPropertyType);
+	/** Converts EDMTextureSetMaterialProperty to EDMMaterialPropertyType. */
+	DYNAMICMATERIAL_API static EDMMaterialPropertyType TextureSetMaterialPropertyToMaterialPropertyType(EDMTextureSetMaterialProperty InPropertyType);
 
-	static EMaterialProperty MaterialPropertyTypeToMaterialProperty(EDMMaterialPropertyType InPropertyType);
+	/** Converts EDMMaterialPropertyType to EMaterialProperty. */
+	DYNAMICMATERIAL_API static EMaterialProperty MaterialPropertyTypeToMaterialProperty(EDMMaterialPropertyType InPropertyType);
 
-	static EDMMaterialPropertyType MaterialPropertyToMaterialPropertyType(EMaterialProperty InPropertyType);
+	/** Converts EMaterialProperty to EDMMaterialPropertyType. */
+	DYNAMICMATERIAL_API static EDMMaterialPropertyType MaterialPropertyToMaterialPropertyType(EMaterialProperty InPropertyType);
 };
 
 #if WITH_EDITOR
diff --git a/Engine/Plugins/Experimental/DynamicMaterial/Source/DynamicMaterialEditor/Private/Components/DMMaterialEffectFunction.cpp b/Engine/Plugins/Experimental/DynamicMaterial/Source/DynamicMaterialEditor/Private/Components/DMMaterialEffectFunction.cpp
--- a/Engine/Plugins/Experimental/DynamicMaterial/Source/DynamicMaterialEditor/Private/Components/DMMaterialEffectFunction.cpp
+++ b/Engine/Plugins/Experimental/DynamicMaterial/Source/DynamicMaterialEditor/Private/Components/DMMaterialEffectFunction.cpp
@@ -662,7 +662,14 @@ void UDMMaterialEffectFunction::InitFunction()
 				return;
 		}
 
-		UDMMaterialValue* Value = UDMMaterialValue::CreateMaterialValue(MaterialModel, TEXT(""), ValueType, /* bInLocal */ true);
+		UDMMaterialValue* Value = UDMMaterialValue::CreateMaterialValue(
+			MaterialModel, 
+			TEXT(""), 
+			UDMValueDefinitionLibrary::GetValueDefinition(ValueType).GetValueClass(), 
+			/* bInLocal */ true
+		);
+		check(Value);
+
 		InputValues.Add(Value);
 
 		if (FunctionInput->bUsePreviewValueAsDefault)
diff --git a/Engine/Plugins/Experimental/DynamicMaterial/Source/DynamicMaterialEditor/Private/Components/DMMaterialSlot.cpp b/Engine/Plugins/Experimental/DynamicMaterial/Source/DynamicMaterialEditor/Private/Components/DMMaterialSlot.cpp
--- a/Engine/Plugins/Experimental/DynamicMaterial/Source/DynamicMaterialEditor/Private/Components/DMMaterialSlot.cpp
+++ b/Engine/Plugins/Experimental/DynamicMaterial/Source/DynamicMaterialEditor/Private/Components/DMMaterialSlot.cpp
@@ -11,7 +11,6 @@
 #include "DMComponentPath.h"
 #include "DMDefs.h"
 #include "DynamicMaterialEditorModule.h"
-#include "Factories/MaterialFactoryNew.h"
 #include "Materials/Material.h"
 #include "Materials/MaterialExpression.h"
 #include "Materials/MaterialExpressionMax.h"
@@ -23,7 +22,6 @@
 #include "Model/DynamicMaterialModelEditorOnlyData.h"
 #include "UObject/Package.h"
 #include "UObject/UObjectGlobals.h"
-#include "Utils/DMPrivate.h"
 #include "Utils/DMUtils.h"
 
 #define LOCTEXT_NAMESPACE "DMMaterialStage"
@@ -366,8 +364,6 @@ void UDMMaterialSlot::UpdateMaterialProperties()
 UDMMaterialSlot::UDMMaterialSlot()
 	: Index(INDEX_NONE)
 	, bIsEditingLayers(true)
-	, BasePreviewMaterial(nullptr)
-	, MaskPreviewMaterial(nullptr)
 {
 }
 
@@ -705,6 +701,11 @@ bool UDMMaterialSlot::RemoveLayer(UDMMaterialLayerObject* InLayer)
 	return true;
 }
 
+void UDMMaterialSlot::OnPropertiesUpdated()
+{
+	OnPropertiesUpdateDelegate.Broadcast(this);
+}
+
 void UDMMaterialSlot::GenerateExpressions(const TSharedRef<FDMMaterialBuildState>& InBuildState) const
 {
 	if (!IsComponentValid() || !IsComponentAdded())
@@ -815,50 +816,7 @@ bool UDMMaterialSlot::UnreferencedBySlot(UDMMaterialSlot* InOtherSlot)
 	}
 }
 
-UMaterial* UDMMaterialSlot::GetPreviewMaterial(EDMMaterialLayerStage InLayerStage)
-{
-	check(InLayerStage == EDMMaterialLayerStage::Base || InLayerStage == EDMMaterialLayerStage::Mask);
-	TObjectPtr<UMaterial>& PreviewMaterial = (InLayerStage == EDMMaterialLayerStage::Base ? BasePreviewMaterial : MaskPreviewMaterial);
-
-	if (!PreviewMaterial)
-	{
-		CreatePreviewMaterial(InLayerStage);
-
-		if (PreviewMaterial)
-		{
-			MarkComponentDirty();
-		}
-	}
-
-	return PreviewMaterial;
-}
-
-void UDMMaterialSlot::CreatePreviewMaterial(EDMMaterialLayerStage InLayerStage)
-{
-	if (!IsComponentValid())
-	{
-		return;
-	}
-
-	check(InLayerStage == EDMMaterialLayerStage::Base || InLayerStage == EDMMaterialLayerStage::Mask);
-	TObjectPtr<UMaterial>& PreviewMaterial = (InLayerStage == EDMMaterialLayerStage::Base ? BasePreviewMaterial : MaskPreviewMaterial);
-
-	UMaterialFactoryNew* MaterialFactory = NewObject<UMaterialFactoryNew>();
-	check(MaterialFactory);
-
-	PreviewMaterial = (UMaterial*)MaterialFactory->FactoryCreateNew(
-		UMaterial::StaticClass(),
-		GetTransientPackage(),
-		NAME_None,
-		RF_Transient,
-		nullptr,
-		GWarn
-	);
-
-	PreviewMaterial->bIsPreviewMaterial = true;
-}
-
-void UDMMaterialSlot::UpdatePreviewMaterial(EDMMaterialLayerStage InLayerStage)
+void UDMMaterialSlot::GeneralPreviewMaterial(UMaterial* InPreviewMaterial, EDMMaterialLayerStage InLayerStage)
 {
 	if (!IsComponentValid())
 	{
@@ -887,24 +845,13 @@ void UDMMaterialSlot::UpdatePreviewMaterial(EDMMaterialLayerStage InLayerStage)
 	}
 
 	check(InLayerStage == EDMMaterialLayerStage::Base || InLayerStage == EDMMaterialLayerStage::Mask);
-	TObjectPtr<UMaterial>& PreviewMaterial = (InLayerStage == EDMMaterialLayerStage::Base ? BasePreviewMaterial : MaskPreviewMaterial);
-
-	if (!PreviewMaterial)
-	{
-		CreatePreviewMaterial(InLayerStage);
-
-		if (!PreviewMaterial)
-		{
-			return;
-		}
-	}
 
 	UDynamicMaterialModelEditorOnlyData* ModelEditorOnlyData = GetMaterialModelEditorOnlyData();
 	check(ModelEditorOnlyData);
 
 	if (InLayerStage == EDMMaterialLayerStage::Base)
 	{
-		TSharedRef<FDMMaterialBuildState> BuildState = ModelEditorOnlyData->CreateBuildState(BasePreviewMaterial);
+		TSharedRef<FDMMaterialBuildState> BuildState = ModelEditorOnlyData->CreateBuildState(InPreviewMaterial);
 		BuildState->SetPreviewMaterial();
 
 		GenerateExpressions(BuildState);
@@ -912,7 +859,7 @@ void UDMMaterialSlot::UpdatePreviewMaterial(EDMMaterialLayerStage InLayerStage)
 	}
 	else
 	{
-		TSharedRef<FDMMaterialBuildState> BuildState = ModelEditorOnlyData->CreateBuildState(MaskPreviewMaterial);
+		TSharedRef<FDMMaterialBuildState> BuildState = ModelEditorOnlyData->CreateBuildState(InPreviewMaterial);
 		BuildState->SetPreviewMaterial();
 
 		GenerateExpressions(BuildState);
@@ -929,8 +876,10 @@ void UDMMaterialSlot::UpdateBasePreviewMaterial(const TSharedRef<FDMMaterialBuil
 
 	UE_LOG(LogDynamicMaterialEditor, Display, TEXT("Building Material Designer Slot Base Preview (%s)..."), *GetName());
 
-	BasePreviewMaterial->GetEditorOnlyData()->EmissiveColor.Expression = nullptr;
-	BasePreviewMaterial->GetEditorOnlyData()->EmissiveColor.OutputIndex = 0;
+	UMaterial* PreviewMaterial = InBuildState->GetDynamicMaterial();
+
+	PreviewMaterial->GetEditorOnlyData()->EmissiveColor.Expression = nullptr;
+	PreviewMaterial->GetEditorOnlyData()->EmissiveColor.OutputIndex = 0;
 
 	const TMap<EDMMaterialPropertyType, TArray<UMaterialExpression*>>& SlotPropertyExpressions = InBuildState->GetSlotPropertyExpressions(this);
 
@@ -950,7 +899,7 @@ void UDMMaterialSlot::UpdateBasePreviewMaterial(const TSharedRef<FDMMaterialBuil
 
 void UDMMaterialSlot::UpdateBasePreviewMaterialProperty(const TSharedRef<FDMMaterialBuildState>& InBuildState, EDMMaterialPropertyType InBaseProperty)
 {
-	UpdatePreviewMaterialProperty(InBuildState, BasePreviewMaterial, InBaseProperty);
+	UpdatePreviewMaterialProperty(InBuildState, InBaseProperty);
 }
 
 void UDMMaterialSlot::UpdateBasePreviewMaterialFull(const TSharedRef<FDMMaterialBuildState>& InBuildState)
@@ -960,6 +909,8 @@ void UDMMaterialSlot::UpdateBasePreviewMaterialFull(const TSharedRef<FDMMaterial
 		return;
 	}
 
+	UMaterial* PreviewMaterial = InBuildState->GetDynamicMaterial();
+
 	UDynamicMaterialModelEditorOnlyData* ModelEditorOnlyData = GetMaterialModelEditorOnlyData();
 	check(ModelEditorOnlyData);
 
@@ -1014,8 +965,10 @@ void UDMMaterialSlot::UpdateMaskPreviewMaterial(const TSharedRef<FDMMaterialBuil
 
 	UE_LOG(LogDynamicMaterialEditor, Display, TEXT("Building Material Designer Slot Mask Preview (%s)..."), *GetName());
 
-	MaskPreviewMaterial->GetEditorOnlyData()->EmissiveColor.Expression = nullptr;
-	MaskPreviewMaterial->GetEditorOnlyData()->EmissiveColor.OutputIndex = 0;
+	UMaterial* PreviewMaterial = InBuildState->GetDynamicMaterial();
+
+	PreviewMaterial->GetEditorOnlyData()->EmissiveColor.Expression = nullptr;
+	PreviewMaterial->GetEditorOnlyData()->EmissiveColor.OutputIndex = 0;
 
 	UDynamicMaterialModelEditorOnlyData* ModelEditorOnlyData = GetMaterialModelEditorOnlyData();
 	check(ModelEditorOnlyData);
@@ -1049,17 +1002,18 @@ void UDMMaterialSlot::UpdateMaskPreviewMaterial(const TSharedRef<FDMMaterialBuil
 
 void UDMMaterialSlot::UpdateMaskPreviewMaterialProperty(const TSharedRef<FDMMaterialBuildState>& InBuildState, EDMMaterialPropertyType InMaskProperty)
 {
-	UpdatePreviewMaterialProperty(InBuildState, MaskPreviewMaterial, InMaskProperty);
+	UpdatePreviewMaterialProperty(InBuildState, InMaskProperty);
 }
 
-void UDMMaterialSlot::UpdateMaskPreviewMaterialMaskCombination(const TSharedRef<FDMMaterialBuildState>& InBuildState, 
-	EDMMaterialPropertyType InMaskProperty)
+void UDMMaterialSlot::UpdateMaskPreviewMaterialMaskCombination(const TSharedRef<FDMMaterialBuildState>& InBuildState, EDMMaterialPropertyType InMaskProperty)
 {
 	if (!IsComponentValid())
 	{
 		return;
 	}
 
+	UMaterial* PreviewMaterial = InBuildState->GetDynamicMaterial();
+
 	int32 OutOutputChannel = FDMMaterialStageConnectorChannel::WHOLE_CHANNEL;
 
 	for (const TObjectPtr<UDMMaterialLayerObject>& Layer : LayerObjects)
@@ -1139,12 +1093,12 @@ void UDMMaterialSlot::UpdateMaskPreviewMaterialMaskCombination(const TSharedRef<
 			}
 		}
 
-		if (MaskPreviewMaterial->GetEditorOnlyData()->EmissiveColor.Expression == nullptr)
+		if (PreviewMaterial->GetEditorOnlyData()->EmissiveColor.Expression == nullptr)
 		{
-			MaskPreviewMaterial->GetEditorOnlyData()->EmissiveColor.Expression = MaskOutputExpression;
+			PreviewMaterial->GetEditorOnlyData()->EmissiveColor.Expression = MaskOutputExpression;
 
 			// The first output will use the node's output info.
-			MaskPreviewMaterial->GetEditorOnlyData()->EmissiveColor.OutputIndex = MaskOutputIndex;
+			PreviewMaterial->GetEditorOnlyData()->EmissiveColor.OutputIndex = MaskOutputIndex;
 			OutOutputChannel = MaskOutputChannel;
 			continue;
 		}
@@ -1152,8 +1106,8 @@ void UDMMaterialSlot::UpdateMaskPreviewMaterialMaskCombination(const TSharedRef<
 		UMaterialExpressionMax* Max = InBuildState->GetBuildUtils().CreateExpression<UMaterialExpressionMax>(UE_DM_NodeComment_Default);
 		check(Max);
 
-		Max->A.Expression = MaskPreviewMaterial->GetEditorOnlyData()->EmissiveColor.Expression;
-		Max->A.OutputIndex = MaskPreviewMaterial->GetEditorOnlyData()->EmissiveColor.OutputIndex;
+		Max->A.Expression = PreviewMaterial->GetEditorOnlyData()->EmissiveColor.Expression;
+		Max->A.OutputIndex = PreviewMaterial->GetEditorOnlyData()->EmissiveColor.OutputIndex;
 		Max->A.Mask = 0;
 
 		if (OutOutputChannel != FDMMaterialStageConnectorChannel::WHOLE_CHANNEL)
@@ -1178,33 +1132,34 @@ void UDMMaterialSlot::UpdateMaskPreviewMaterialMaskCombination(const TSharedRef<
 			Max->B.MaskA = !!(MaskOutputChannel & FDMMaterialStageConnectorChannel::FOURTH_CHANNEL);
 		}
 
-		MaskPreviewMaterial->GetEditorOnlyData()->EmissiveColor.Expression = Max;
+		PreviewMaterial->GetEditorOnlyData()->EmissiveColor.Expression = Max;
 
 		// If we have to combine, it will use the Max node's output info
-		MaskPreviewMaterial->GetEditorOnlyData()->EmissiveColor.OutputIndex = 0;
+		PreviewMaterial->GetEditorOnlyData()->EmissiveColor.OutputIndex = 0;
 		OutOutputChannel = FDMMaterialStageConnectorChannel::WHOLE_CHANNEL;
 	}
 
-	MaskPreviewMaterial->GetEditorOnlyData()->EmissiveColor.Mask = 0;
+	PreviewMaterial->GetEditorOnlyData()->EmissiveColor.Mask = 0;
 
 	if (OutOutputChannel != FDMMaterialStageConnectorChannel::WHOLE_CHANNEL)
 	{
-		MaskPreviewMaterial->GetEditorOnlyData()->EmissiveColor.Mask = 1;
-		MaskPreviewMaterial->GetEditorOnlyData()->EmissiveColor.MaskR = !!(OutOutputChannel & FDMMaterialStageConnectorChannel::FIRST_CHANNEL);
-		MaskPreviewMaterial->GetEditorOnlyData()->EmissiveColor.MaskG = !!(OutOutputChannel & FDMMaterialStageConnectorChannel::SECOND_CHANNEL);
-		MaskPreviewMaterial->GetEditorOnlyData()->EmissiveColor.MaskB = !!(OutOutputChannel & FDMMaterialStageConnectorChannel::THIRD_CHANNEL);
-		MaskPreviewMaterial->GetEditorOnlyData()->EmissiveColor.MaskA = !!(OutOutputChannel & FDMMaterialStageConnectorChannel::FOURTH_CHANNEL);
+		PreviewMaterial->GetEditorOnlyData()->EmissiveColor.Mask = 1;
+		PreviewMaterial->GetEditorOnlyData()->EmissiveColor.MaskR = !!(OutOutputChannel & FDMMaterialStageConnectorChannel::FIRST_CHANNEL);
+		PreviewMaterial->GetEditorOnlyData()->EmissiveColor.MaskG = !!(OutOutputChannel & FDMMaterialStageConnectorChannel::SECOND_CHANNEL);
+		PreviewMaterial->GetEditorOnlyData()->EmissiveColor.MaskB = !!(OutOutputChannel & FDMMaterialStageConnectorChannel::THIRD_CHANNEL);
+		PreviewMaterial->GetEditorOnlyData()->EmissiveColor.MaskA = !!(OutOutputChannel & FDMMaterialStageConnectorChannel::FOURTH_CHANNEL);
 	}
 }
 
-void UDMMaterialSlot::UpdatePreviewMaterialProperty(const TSharedRef<FDMMaterialBuildState>& InBuildState, UMaterial* InPreviewMaterial, 
-	EDMMaterialPropertyType InProperty)
+void UDMMaterialSlot::UpdatePreviewMaterialProperty(const TSharedRef<FDMMaterialBuildState>& InBuildState, EDMMaterialPropertyType InProperty)
 {
 	if (!IsComponentValid())
 	{
 		return;
 	}
 
+	UMaterial* PreviewMaterial = InBuildState->GetDynamicMaterial();
+
 	UDynamicMaterialModelEditorOnlyData* ModelEditorOnlyData = GetMaterialModelEditorOnlyData();
 	check(ModelEditorOnlyData);
 	
@@ -1265,8 +1220,8 @@ void UDMMaterialSlot::UpdatePreviewMaterialProperty(const TSharedRef<FDMMaterial
 
 	if (BestMatch != INDEX_NONE)
 	{
-		BasePreviewMaterial->GetEditorOnlyData()->EmissiveColor.Expression = LastExpression;
-		BasePreviewMaterial->GetEditorOnlyData()->EmissiveColor.OutputIndex = BestMatch;
+		PreviewMaterial->GetEditorOnlyData()->EmissiveColor.Expression = LastExpression;
+		PreviewMaterial->GetEditorOnlyData()->EmissiveColor.OutputIndex = BestMatch;
 	}
 }
 
@@ -1410,37 +1365,36 @@ bool UDMMaterialSlot::ChangeMaterialProperty(EDMMaterialPropertyType InPropertyF
 	return true;
 }
 
-void UDMMaterialSlot::DoClean()
-{
-	if (IsComponentValid())
-	{
-		UpdatePreviewMaterial(EDMMaterialLayerStage::Base);
-		UpdatePreviewMaterial(EDMMaterialLayerStage::Mask);
-	}
-
-	Super::DoClean();
-}
-
 FString UDMMaterialSlot::GetComponentPathComponent() const
 {
-	// @TODO These will probably need to change when other slot types are opened up.
-	switch (Index)
+	if (UDynamicMaterialModelEditorOnlyData* EditorOnlyData = GetMaterialModelEditorOnlyData())
 	{
-		case 0:
-			return UDynamicMaterialModelEditorOnlyData::RGBSlotPathToken;
+		TArray<EDMMaterialPropertyType> SlotProperties = EditorOnlyData->GetMaterialPropertiesForSlot(this);
 
-		case 1:
-			return UDynamicMaterialModelEditorOnlyData::OpacitySlotPathToken;
+		if (SlotProperties.Num() == 1)
+		{
+			UEnum* MaterialPropertyEnum = StaticEnum<EDMMaterialPropertyType>();
+			constexpr const TCHAR* ShortNameName = TEXT("ShortName");
+			const FString ShortName = MaterialPropertyEnum->GetMetaData(ShortNameName, MaterialPropertyEnum->GetIndexByValue(static_cast<int64>(SlotProperties[0])));
+			const FString Token = !ShortName.IsEmpty() ? ShortName : MaterialPropertyEnum->GetNameStringByValue(static_cast<int64>(SlotProperties[0]));
 
-		default:
 			return FString::Printf(
-				TEXT("%s%hc%i%hc"), 
+				TEXT("%s%hc%s%hc"),
 				*UDynamicMaterialModelEditorOnlyData::SlotsPathToken,
 				FDMComponentPath::ParameterOpen,
-				Index, 
+				*Token,
 				FDMComponentPath::ParameterClose
 			);
+		}
 	}
+
+	return FString::Printf(
+		TEXT("%s%hc%i%hc"),
+		*UDynamicMaterialModelEditorOnlyData::SlotsPathToken,
+		FDMComponentPath::ParameterOpen,
+		Index,
+		FDMComponentPath::ParameterClose
+	);
 }
 
 UDMMaterialComponent* UDMMaterialSlot::GetSubComponentByPath(FDMComponentPath& InPath, const FDMComponentPathSegment& InPathSegment) const
diff --git a/Engine/Plugins/Experimental/DynamicMaterial/Source/DynamicMaterialEditor/Private/Components/DMMaterialStage.cpp b/Engine/Plugins/Experimental/DynamicMaterial/Source/DynamicMaterialEditor/Private/Components/DMMaterialStage.cpp
--- a/Engine/Plugins/Experimental/DynamicMaterial/Source/DynamicMaterialEditor/Private/Components/DMMaterialStage.cpp
+++ b/Engine/Plugins/Experimental/DynamicMaterial/Source/DynamicMaterialEditor/Private/Components/DMMaterialStage.cpp
@@ -1,20 +1,14 @@
 // Copyright Epic Games, Inc. All Rights Reserved.
 
 #include "Components/DMMaterialStage.h"
-#include "AssetRegistry/AssetRegistryModule.h"
 #include "Components/DMMaterialLayer.h"
 #include "Components/DMMaterialSlot.h"
 #include "Components/DMMaterialStageInput.h"
 #include "Components/DMMaterialStageSource.h"
 #include "Components/DMMaterialStageThroughputLayerBlend.h"
 #include "Components/DMMaterialSubStage.h"
-#include "Components/DMMaterialValue.h"
-#include "Components/DMTextureUV.h"
 #include "DMComponentPath.h"
 #include "DynamicMaterialEditorModule.h"
-#include "DynamicMaterialEditorSettings.h"
-#include "DynamicMaterialModule.h"
-#include "Factories/MaterialFactoryNew.h"
 #include "Materials/Material.h"
 #include "Materials/MaterialExpression.h"
 #include "Materials/MaterialInstanceDynamic.h"
@@ -41,8 +35,6 @@ UDMMaterialStage::UDMMaterialStage()
 	, bEnabled(true)
 	, bCanChangeSource(true)
 	, bIsBeingEdited(false)
-	, PreviewMaterialBase(nullptr)
-	, PreviewMaterialDynamic(nullptr)
 {
 	EditableProperties.Add(GET_MEMBER_NAME_CHECKED(UDMMaterialStage, Source));
 	EditableProperties.Add(GET_MEMBER_NAME_CHECKED(UDMMaterialStage, Inputs));
@@ -73,16 +65,12 @@ void UDMMaterialStage::OnComponentAdded()
 
 		Source->SetComponentState(EDMComponentLifetimeState::Added);
 	}
-
-	AddDelegates();
 }
 
 void UDMMaterialStage::OnComponentRemoved()
 {
 	Super::OnComponentRemoved();
 
-	RemoveDelegates();
-
 	for (UDMMaterialStageInput* Input : Inputs)
 	{
 		Input->SetComponentState(EDMComponentLifetimeState::Removed);
@@ -113,11 +101,6 @@ void UDMMaterialStage::PostEditorDuplicate(UDynamicMaterialModel* InMaterialMode
 		Rename(nullptr, InParent, UE::DynamicMaterial::RenameFlags);
 	}
 
-	PreviewMaterialBase = nullptr;
-	PreviewMaterialDynamic = nullptr;
-
-	AddDelegates();
-
 	if (Source)
 	{
 		if (GUndo)
@@ -244,98 +227,6 @@ UDMMaterialComponent* UDMMaterialStage::GetSubComponentByPath(FDMComponentPath&
 	return Super::GetSubComponentByPath(InPath, InPathSegment);
 }
 
-void UDMMaterialStage::AddDelegates()
-{
-	if (UDMMaterialLayerObject* Layer = GetLayer())
-	{
-		if (UDMMaterialSlot* Slot = Layer->GetSlot())
-		{
-			if (UDynamicMaterialModelEditorOnlyData* ModelEditorOnlyData = Slot->GetMaterialModelEditorOnlyData())
-			{
-				if (UDynamicMaterialModel* MaterialModel = ModelEditorOnlyData->GetMaterialModel())
-				{
-					MaterialModel->GetOnValueUpdateDelegate().AddUObject(this, &UDMMaterialStage::OnValueUpdated);
-					MaterialModel->GetOnTextureUVUpdateDelegate().AddUObject(this, &UDMMaterialStage::OnTextureUVUpdated);
-				}
-			}
-		}
-	}
-}
-
-void UDMMaterialStage::RemoveDelegates()
-{
-	if (UDMMaterialLayerObject* Layer = GetLayer())
-	{
-		if (UDMMaterialSlot* Slot = Layer->GetSlot())
-		{
-			if (UDynamicMaterialModelEditorOnlyData* ModelEditorOnlyData = Slot->GetMaterialModelEditorOnlyData())
-			{
-				if (UDynamicMaterialModel* MaterialModel = ModelEditorOnlyData->GetMaterialModel())
-				{
-					MaterialModel->GetOnValueUpdateDelegate().RemoveAll(this);
-					MaterialModel->GetOnTextureUVUpdateDelegate().RemoveAll(this);
-				}
-			}
-		}
-	}
-}
-
-void UDMMaterialStage::OnValueUpdated(UDynamicMaterialModel* InMaterialModel, UDMMaterialValue* InValue)
-{
-	if (PreviewMaterialDynamic && InMaterialModel && InValue)
-	{
-		InValue->SetMIDParameter(PreviewMaterialDynamic);
-	}
-}
-
-void UDMMaterialStage::OnTextureUVUpdated(UDynamicMaterialModel* InMaterialModel, UDMTextureUV* InTextureUV)
-{
-	if (!UDynamicMaterialEditorSettings::Get()->bPreviewImagesUseTextureUVs)
-	{
-		return;
-	}
-
-	if (!PreviewMaterialDynamic || !InMaterialModel || !InTextureUV)
-	{
-		return;
-	}
-
-	UDMMaterialStage* Stage = InTextureUV->GetTypedParent<UDMMaterialStage>(false);
-
-	if (Stage != this)
-	{
-		return;
-	}
-
-	InTextureUV->SetMIDParameters(PreviewMaterialDynamic);
-
-	UDMMaterialLayerObject* Layer = Stage->GetLayer();
-
-	if (!Layer || !Layer->IsTextureUVLinkEnabled() || Layer->GetStageType(Stage) != EDMMaterialLayerStage::Base)
-	{
-		return;
-	}
-
-	UDMMaterialStageThroughput* Throughput = Cast<UDMMaterialStageThroughput>(Stage->GetSource());
-
-	if (!Throughput || !Throughput->SupportsLayerMaskTextureUVLink())
-	{
-		return;
-	}
-
-	UDMMaterialStage* MaskStage = Layer->GetStage(EDMMaterialLayerStage::Mask);
-
-	if (!MaskStage)
-	{
-		return;
-	}
-
-	if (UMaterialInstanceDynamic* MaskMID = Cast<UMaterialInstanceDynamic>(MaskStage->GetPreviewMaterial()))
-	{
-		InTextureUV->SetMIDParameters(MaskMID);
-	}
-}
-
 UDMMaterialLayerObject* UDMMaterialStage::GetLayer() const
 {
 	return Cast<UDMMaterialLayerObject>(GetOuterSafe());
@@ -869,7 +760,8 @@ UDMMaterialStageSource* UDMMaterialStage::ChangeInput_PreviousStage(int32 InInpu
 	return PreviousSource;
 }
 
-void UDMMaterialStage::UpdateInputMap(int32 InInputIdx, int32 InSourceIndex, int32 InInputChannel, int32 InOutputIdx, int32 InOutputChannel, EDMMaterialPropertyType InStageProperty)
+void UDMMaterialStage::UpdateInputMap(int32 InInputIdx, int32 InSourceIndex, int32 InInputChannel, int32 InOutputIdx, int32 InOutputChannel, 
+	EDMMaterialPropertyType InStageProperty)
 {
 	if (!IsComponentValid())
 	{
@@ -1032,11 +924,6 @@ UDMMaterialStage* UDMMaterialStage::GetNextStage() const
 	return nullptr;
 }
 
-bool UDMMaterialStage::IsRootStage() const
-{
-	return true;
-}
-
 bool UDMMaterialStage::VerifyAllInputMaps()
 {
 	if (!IsComponentValid())
@@ -1244,10 +1131,6 @@ UDMMaterialStageSource* UDMMaterialStage::ChangeSource(TSubclassOf<UDMMaterialSt
 
 bool UDMMaterialStage::IsCompatibleWithPreviousStage(const UDMMaterialStage* InPreviousStage) const
 {
-	/*
-	 * It is now up to the user to sort this particular problem out because it would do more harm than good
-	 * to force correctness in "transition states" while the user is changing settings.
-	 */
 	return true;
 }
 
@@ -1342,21 +1225,6 @@ void UDMMaterialStage::RemoveAllInputs()
 	Update(EDMUpdateType::Structure);
 }
 
-UMaterialInterface* UDMMaterialStage::GetPreviewMaterial()
-{
-	if (!PreviewMaterialBase)
-	{
-		CreatePreviewMaterial();
-
-		if (PreviewMaterialBase)
-		{
-			MarkComponentDirty();
-		}
-	}
-
-	return PreviewMaterialDynamic.Get();
-}
-
 const FDMMaterialStageConnectorChannel* UDMMaterialStage::FindInputChannel(UDMMaterialStageInput* InStageInput)
 {
 	check(InStageInput);
@@ -1384,16 +1252,6 @@ const FDMMaterialStageConnectorChannel* UDMMaterialStage::FindInputChannel(UDMMa
 	return nullptr;
 }
 
-void UDMMaterialStage::DoClean()
-{
-	if (IsComponentValid() && GetLayer())
-	{
-		UpdatePreviewMaterial();
-	}
-
-	Super::DoClean();
-}
-
 void UDMMaterialStage::PostEditUndo()
 {
 	Super::PostEditUndo();
@@ -1416,66 +1274,7 @@ void UDMMaterialStage::PostEditUndo()
 	Update(EDMUpdateType::Structure);
 }
 
-void UDMMaterialStage::PostLoad()
-{
-	Super::PostLoad();
-
-	AddDelegates();
-}
-
-void UDMMaterialStage::PostEditImport()
-{
-	Super::PostEditImport();
-
-	AddDelegates();
-}
-
-void UDMMaterialStage::CreatePreviewMaterial()
-{
-	if (!IsComponentValid())
-	{
-		return;
-	}
-
-	if (FDynamicMaterialModule::IsMaterialExportEnabled() == false)
-	{
-		UMaterialFactoryNew* MaterialFactory = NewObject<UMaterialFactoryNew>();
-		check(MaterialFactory);
-
-		PreviewMaterialBase = Cast<UMaterial>(MaterialFactory->FactoryCreateNew(
-			UMaterial::StaticClass(),
-			GetTransientPackage(),
-			NAME_None,
-			RF_Transient,
-			nullptr,
-			GWarn
-		));
-
-		PreviewMaterialBase->bIsPreviewMaterial = true;
-	}
-	else
-	{
-		FString MaterialBaseName = GetName() + "-" + FGuid::NewGuid().ToString();
-		const FString FullName = "/Game/DynamicMaterials/" + MaterialBaseName;
-		UPackage* Package = CreatePackage(*FullName);
-
-		UMaterialFactoryNew* MaterialFactory = NewObject<UMaterialFactoryNew>();
-		check(MaterialFactory);
-
-		PreviewMaterialBase = Cast<UMaterial>(MaterialFactory->FactoryCreateNew(
-			UMaterial::StaticClass(),
-			Package,
-			*MaterialBaseName,
-			RF_Standalone | RF_Public,
-			nullptr,
-			GWarn
-		));
-
-		FAssetRegistryModule::AssetCreated(PreviewMaterialBase);
-	}
-}
-
-void UDMMaterialStage::UpdatePreviewMaterial()
+void UDMMaterialStage::GeneratePreviewMaterial(UMaterial* InPreviewMaterial)
 {
 	if (!IsComponentValid())
 	{
@@ -1489,16 +1288,6 @@ void UDMMaterialStage::UpdatePreviewMaterial()
 
 	UE_LOG(LogDynamicMaterialEditor, Display, TEXT("Building Material Designer Stage Preview (%s)..."), *GetName());
 
-	if (!PreviewMaterialBase)
-	{
-		CreatePreviewMaterial();
-
-		if (!PreviewMaterialBase)
-		{
-			return;
-		}
-	}
-
 	UDMMaterialLayerObject* Layer = GetLayer();
 	check(Layer);
 
@@ -1511,42 +1300,15 @@ void UDMMaterialStage::UpdatePreviewMaterial()
 	UDynamicMaterialModel* MaterialModel = ModelEditorOnlyData->GetMaterialModel();
 	check(MaterialModel);
 
-	PreviewMaterialBase->GetEditorOnlyData()->EmissiveColor.Expression = nullptr;
-	PreviewMaterialBase->GetEditorOnlyData()->EmissiveColor.OutputIndex = 0;
+	InPreviewMaterial->GetEditorOnlyData()->EmissiveColor.Expression = nullptr;
+	InPreviewMaterial->GetEditorOnlyData()->EmissiveColor.OutputIndex = 0;
 
-	const bool bGenerateSuccess = Source->UpdateStagePreviewMaterial(
+	const bool bGenerateSuccess = Source->GenerateStagePreviewMaterial(
 		this,
-		PreviewMaterialBase,
-		PreviewMaterialBase->GetEditorOnlyData()->EmissiveColor.Expression,
-		PreviewMaterialBase->GetEditorOnlyData()->EmissiveColor.OutputIndex
+		InPreviewMaterial,
+		InPreviewMaterial->GetEditorOnlyData()->EmissiveColor.Expression,
+		InPreviewMaterial->GetEditorOnlyData()->EmissiveColor.OutputIndex
 	);
-
-	if (!bGenerateSuccess)
-	{
-		return;
-	}
-
-	PreviewMaterialDynamic = UMaterialInstanceDynamic::Create(PreviewMaterialBase.Get(), GetTransientPackage());
-
-	const bool bUseTexureUVs = UDynamicMaterialEditorSettings::Get()->bPreviewImagesUseTextureUVs;
-
-	TArray<UObject*> Subobjects;
-	GetObjectsWithOuter(MaterialModel, Subobjects, false);
-
-	for (UObject* Subobject : Subobjects)
-	{
-		if (UDMMaterialValue* Value = Cast<UDMMaterialValue>(Subobject))
-		{
-			OnValueUpdated(MaterialModel, Value);
-		}
-		else if (bUseTexureUVs)
-		{
-			if (UDMTextureUV* TextureUV = Cast<UDMTextureUV>(Subobject))
-			{
-				OnTextureUVUpdated(MaterialModel, TextureUV);
-			}
-		}
-	}
 }
 
 #undef LOCTEXT_NAMESPACE
diff --git a/Engine/Plugins/Experimental/DynamicMaterial/Source/DynamicMaterialEditor/Private/Components/DMMaterialStageBlend.cpp b/Engine/Plugins/Experimental/DynamicMaterial/Source/DynamicMaterialEditor/Private/Components/DMMaterialStageBlend.cpp
--- a/Engine/Plugins/Experimental/DynamicMaterial/Source/DynamicMaterialEditor/Private/Components/DMMaterialStageBlend.cpp
+++ b/Engine/Plugins/Experimental/DynamicMaterial/Source/DynamicMaterialEditor/Private/Components/DMMaterialStageBlend.cpp
@@ -316,7 +316,7 @@ void UDMMaterialStageBlend::GetMaskAlphaBlendNode(const TSharedRef<FDMMaterialBu
 	OutExpression = Expressions.Last();
 }
 
-bool UDMMaterialStageBlend::UpdateStagePreviewMaterial(UDMMaterialStage* InStage, UMaterial* InPreviewMaterial, 
+bool UDMMaterialStageBlend::GenerateStagePreviewMaterial(UDMMaterialStage* InStage, UMaterial* InPreviewMaterial, 
 	UMaterialExpression*& OutMaterialExpression, int32& OutputIndex)
 {
 	check(InStage);
@@ -624,27 +624,12 @@ FDMExpressionInput UDMMaterialStageBlend::GetLayerMaskLinkTextureUVInputExpressi
 	return ExpressionInput;
 }
  
-void UDMMaterialStageBlend::UpdatePreviewMaterial(UMaterial* InPreviewMaterial)
+void UDMMaterialStageBlend::GeneratePreviewMaterial(UMaterial* InPreviewMaterial)
 {
 	if (!IsComponentValid())
 	{
 		return;
 	}
-
-	if (!InPreviewMaterial)
-	{
-		if (!PreviewMaterial)
-		{
-			CreatePreviewMaterial();
-		}
- 
-		InPreviewMaterial = PreviewMaterial;
- 
-		if (!InPreviewMaterial)
-		{
-			return;
-		}
-	}
  
 	UDMMaterialStage* Stage = GetStage();
 	check(Stage);
diff --git a/Engine/Plugins/Experimental/DynamicMaterial/Source/DynamicMaterialEditor/Private/Components/DMMaterialStageFunction.cpp b/Engine/Plugins/Experimental/DynamicMaterial/Source/DynamicMaterialEditor/Private/Components/DMMaterialStageFunction.cpp
--- a/Engine/Plugins/Experimental/DynamicMaterial/Source/DynamicMaterialEditor/Private/Components/DMMaterialStageFunction.cpp
+++ b/Engine/Plugins/Experimental/DynamicMaterial/Source/DynamicMaterialEditor/Private/Components/DMMaterialStageFunction.cpp
@@ -248,9 +248,9 @@ void UDMMaterialStageFunction::PreEditChange(FEditPropertyChain& PropertyAboutTo
 	MaterialFunction_PreEdit = MaterialFunction;
 }
 
-void UDMMaterialStageFunction::PostEditChangeProperty(FPropertyChangedEvent& PropertyChangedEvent)
+void UDMMaterialStageFunction::PostEditChangeProperty(FPropertyChangedEvent& InPropertyChangedEvent)
 {
-	Super::PostEditChangeProperty(PropertyChangedEvent);
+	Super::PostEditChangeProperty(InPropertyChangedEvent);
 
 	if (MaterialFunction != MaterialFunction_PreEdit)
 	{
@@ -577,7 +577,7 @@ void UDMMaterialStageFunction::InitFunction()
 		InputConnectors[InputIndex].Index = InputIndex;
 		InputConnectors[InputIndex].Type = ValueType;
 
-		if (Inputs[InputIndex].ExpressionInput->InputName == NAME_None)
+		if (Inputs[InputIndex].ExpressionInput->InputName.IsNone())
 		{
 			static const FText InputNameFormat = LOCTEXT("InputFormat", "Input {0}");
 			InputConnectors[InputIndex].Name = FText::Format(InputNameFormat, FText::AsNumber(InputIndex + 1));
diff --git a/Engine/Plugins/Experimental/DynamicMaterial/Source/DynamicMaterialEditor/Private/Components/DMMaterialStageInput.cpp b/Engine/Plugins/Experimental/DynamicMaterial/Source/DynamicMaterialEditor/Private/Components/DMMaterialStageInput.cpp
--- a/Engine/Plugins/Experimental/DynamicMaterial/Source/DynamicMaterialEditor/Private/Components/DMMaterialStageInput.cpp
+++ b/Engine/Plugins/Experimental/DynamicMaterial/Source/DynamicMaterialEditor/Private/Components/DMMaterialStageInput.cpp
@@ -50,28 +50,13 @@ void UDMMaterialStageInput::Update(EDMUpdateType InUpdateType)
 	Super::Update(InUpdateType);
 }
 
-void UDMMaterialStageInput::UpdatePreviewMaterial(UMaterial* InPreviewMaterial /*= nullptr*/)
+void UDMMaterialStageInput::GeneratePreviewMaterial(UMaterial* InPreviewMaterial)
 {
 	if (!IsComponentValid())
 	{
 		return;
 	}
 
-	if (!InPreviewMaterial)
-	{
-		if (!PreviewMaterial)
-		{
-			CreatePreviewMaterial();
-		}
-
-		InPreviewMaterial = PreviewMaterial;
-
-		if (!PreviewMaterial)
-		{
-			return;
-		}
-	}
-
 	UDMMaterialStage* Stage = GetStage();
 	check(Stage);
 
diff --git a/Engine/Plugins/Experimental/DynamicMaterial/Source/DynamicMaterialEditor/Private/Components/DMMaterialStageSource.cpp b/Engine/Plugins/Experimental/DynamicMaterial/Source/DynamicMaterialEditor/Private/Components/DMMaterialStageSource.cpp
--- a/Engine/Plugins/Experimental/DynamicMaterial/Source/DynamicMaterialEditor/Private/Components/DMMaterialStageSource.cpp
+++ b/Engine/Plugins/Experimental/DynamicMaterial/Source/DynamicMaterialEditor/Private/Components/DMMaterialStageSource.cpp
@@ -1,31 +1,22 @@
 // Copyright Epic Games, Inc. All Rights Reserved.
 
 #include "Components/DMMaterialStageSource.h"
-#include "AssetRegistry/AssetRegistryModule.h"
 #include "Components/DMMaterialLayer.h"
 #include "Components/DMMaterialProperty.h"
 #include "Components/DMMaterialSlot.h"
 #include "Components/DMMaterialStage.h"
 #include "DynamicMaterialEditorModule.h"
-#include "DynamicMaterialModule.h"
-#include "Factories/MaterialFactoryNew.h"
 #include "Materials/Material.h"
 #include "Materials/MaterialExpression.h"
 #include "MaterialValueType.h"
 #include "Model/DMMaterialBuildState.h"
 #include "Model/DMMaterialBuildUtils.h"
-#include "Model/DynamicMaterialModel.h"
 #include "Model/DynamicMaterialModelEditorOnlyData.h"
 #include "UObject/Package.h"
 #include "UObject/UObjectIterator.h"
 
 TArray<TStrongObjectPtr<UClass>> UDMMaterialStageSource::SourceClasses = TArray<TStrongObjectPtr<UClass>>();
 
-UDMMaterialStageSource::UDMMaterialStageSource()
-{
-	PreviewMaterial = nullptr;
-}
-
 UDMMaterialStage* UDMMaterialStageSource::GetStage() const
 {
 	return Cast<UDMMaterialStage>(GetOuterSafe());
@@ -110,103 +101,13 @@ void UDMMaterialStageSource::GenerateClassList()
 	}
 }
 
-UMaterial* UDMMaterialStageSource::GetPreviewMaterial()
-{
-	if (!PreviewMaterial)
-	{
-		CreatePreviewMaterial();
-
-		if (PreviewMaterial)
-		{
-			MarkComponentDirty();
-		}
-	}
-
-	return PreviewMaterial;
-}
-
-void UDMMaterialStageSource::UpdateOntoPreviewMaterial(UMaterial* ExternalPreviewMaterial)
-{
-	if (!IsComponentValid())
-	{
-		return;
-	}
-
-	if (!ExternalPreviewMaterial)
-	{
-		return;
-	}
-
-	UpdatePreviewMaterial(ExternalPreviewMaterial);
-}
-
-void UDMMaterialStageSource::CreatePreviewMaterial()
+void UDMMaterialStageSource::GeneratePreviewMaterial(UMaterial* InPreviewMaterial)
 {
 	if (!IsComponentValid())
 	{
 		return;
 	}
 
-	if (FDynamicMaterialModule::IsMaterialExportEnabled() == false)
-	{
-		UMaterialFactoryNew* MaterialFactory = NewObject<UMaterialFactoryNew>();
-		check(MaterialFactory);
-
-		PreviewMaterial = Cast<UMaterial>(MaterialFactory->FactoryCreateNew(
-			UMaterial::StaticClass(),
-			GetTransientPackage(),
-			NAME_None,
-			RF_Transient,
-			nullptr,
-			GWarn
-		));
-
-		PreviewMaterial->bIsPreviewMaterial = true;
-	}
-	else
-	{
-		FString MaterialBaseName = GetName() + "-" + FGuid::NewGuid().ToString();
-		const FString FullName = "/Game/DynamicMaterials/" + MaterialBaseName;
-		UPackage* Package = CreatePackage(*FullName);
-
-		UMaterialFactoryNew* MaterialFactory = NewObject<UMaterialFactoryNew>();
-		check(MaterialFactory);
-
-		PreviewMaterial = Cast<UMaterial>(MaterialFactory->FactoryCreateNew(
-			UMaterial::StaticClass(),
-			Package,
-			*MaterialBaseName,
-			RF_Standalone | RF_Public,
-			nullptr,
-			GWarn
-		));
-
-		FAssetRegistryModule::AssetCreated(PreviewMaterial);
-	}
-}
-
-void UDMMaterialStageSource::UpdatePreviewMaterial(UMaterial* InPreviewMaterial)
-{
-	if (!IsComponentValid())
-	{
-		return;
-	}
-
-	if (!InPreviewMaterial)
-	{
-		if (!PreviewMaterial)
-		{
-			CreatePreviewMaterial();
-		}
-
-		InPreviewMaterial = PreviewMaterial;
-
-		if (!PreviewMaterial)
-		{
-			return;
-		}
-	}
-
 	UE_LOG(LogDynamicMaterialEditor, Display, TEXT("Building Material Designer Source Preview (%s)..."), *GetName());
 
 	UDMMaterialStage* Stage = GetStage();
@@ -235,7 +136,7 @@ int32 UDMMaterialStageSource::GetInnateMaskOutput(int32 OutputIndex, int32 Outpu
 	return INDEX_NONE;
 }
 
-bool UDMMaterialStageSource::UpdateStagePreviewMaterial(UDMMaterialStage* InStage, UMaterial* InPreviewMaterial, 
+bool UDMMaterialStageSource::GenerateStagePreviewMaterial(UDMMaterialStage* InStage, UMaterial* InPreviewMaterial, 
 	UMaterialExpression*& OutMaterialExpression, int32& OutputIndex)
 {
 	check(InStage);
@@ -348,25 +249,7 @@ void UDMMaterialStageSource::PostEditUndo()
 	Update(EDMUpdateType::Structure);
 }
 
-void UDMMaterialStageSource::DoClean()
-{
-	if (IsComponentValid())
-	{
-		// Stage Source preview images are currently disabled.
-		//UpdatePreviewMaterial();
-	}
-
-	Super::DoClean();
-}
-
 UDMMaterialComponent* UDMMaterialStageSource::GetParentComponent() const
 {
 	return GetStage();
 }
-
-void UDMMaterialStageSource::PostEditorDuplicate(UDynamicMaterialModel* InMaterialModel, UDMMaterialComponent* InParent)
-{
-	Super::PostEditorDuplicate(InMaterialModel, InParent);
-
-	PreviewMaterial = nullptr;
-}
diff --git a/Engine/Plugins/Experimental/DynamicMaterial/Source/DynamicMaterialEditor/Private/Components/DMMaterialStageThroughput.cpp b/Engine/Plugins/Experimental/DynamicMaterial/Source/DynamicMaterialEditor/Private/Components/DMMaterialStageThroughput.cpp
--- a/Engine/Plugins/Experimental/DynamicMaterial/Source/DynamicMaterialEditor/Private/Components/DMMaterialStageThroughput.cpp
+++ b/Engine/Plugins/Experimental/DynamicMaterial/Source/DynamicMaterialEditor/Private/Components/DMMaterialStageThroughput.cpp
@@ -48,7 +48,8 @@ bool UDMMaterialStageThroughput::CanInputAcceptType(int32 InThroughputInputIndex
 	return InputConnectors[InThroughputInputIndex].IsCompatibleWith(InValueType);
 }
 
-bool UDMMaterialStageThroughput::CanInputConnectTo(int32 InThroughputInputIndex, const FDMMaterialStageConnector& InOutputConnector, int32 InOutputChannel, bool bInCheckSingleFloat)
+bool UDMMaterialStageThroughput::CanInputConnectTo(int32 InThroughputInputIndex, const FDMMaterialStageConnector& InOutputConnector, 
+	int32 InOutputChannel, bool bInCheckSingleFloat)
 {
 	if (InOutputChannel == FDMMaterialStageConnectorChannel::WHOLE_CHANNEL)
 	{
@@ -849,28 +850,13 @@ int32 UDMMaterialStageThroughput::ResolveInputChannel(const TSharedRef<FDMMateri
 	return NodeOutputIndex;
 }
 
-void UDMMaterialStageThroughput::UpdatePreviewMaterial(UMaterial* InPreviewMaterial /*= nullptr*/)
+void UDMMaterialStageThroughput::GeneratePreviewMaterial(UMaterial* InPreviewMaterial)
 {
 	if (!IsComponentValid())
 	{
 		return;
 	}
 
-	if (!InPreviewMaterial)
-	{
-		if (!PreviewMaterial)
-		{
-			CreatePreviewMaterial();
-		}
-
-		InPreviewMaterial = PreviewMaterial;
-
-		if (!InPreviewMaterial)
-		{
-			return;
-		}
-	}
-
 	UDMMaterialStage* Stage = GetStage();
 	check(Stage);
 
diff --git a/Engine/Plugins/Experimental/DynamicMaterial/Source/DynamicMaterialEditor/Private/Components/DMMaterialStageThroughputLayerBlend.cpp b/Engine/Plugins/Experimental/DynamicMaterial/Source/DynamicMaterialEditor/Private/Components/DMMaterialStageThroughputLayerBlend.cpp
--- a/Engine/Plugins/Experimental/DynamicMaterial/Source/DynamicMaterialEditor/Private/Components/DMMaterialStageThroughputLayerBlend.cpp
+++ b/Engine/Plugins/Experimental/DynamicMaterial/Source/DynamicMaterialEditor/Private/Components/DMMaterialStageThroughputLayerBlend.cpp
@@ -555,7 +555,7 @@ FText UDMMaterialStageThroughputLayerBlend::GetStageDescription() const
 	return Super::GetComponentDescription();
 }
 
-bool UDMMaterialStageThroughputLayerBlend::UpdateStagePreviewMaterial(UDMMaterialStage* InStage, UMaterial* InPreviewMaterial, 
+bool UDMMaterialStageThroughputLayerBlend::GenerateStagePreviewMaterial(UDMMaterialStage* InStage, UMaterial* InPreviewMaterial, 
 	UMaterialExpression*& OutMaterialExpression, int32& OutputIndex)
 {
 	check(InStage);
@@ -889,27 +889,12 @@ int32 UDMMaterialStageThroughputLayerBlend::ResolveMaskInput(const TSharedRef<FD
 	return OutChannel.OutputIndex;
 }
  
-void UDMMaterialStageThroughputLayerBlend::UpdatePreviewMaterial(UMaterial* InPreviewMaterial /*= nullptr*/)
+void UDMMaterialStageThroughputLayerBlend::GeneratePreviewMaterial(UMaterial* InPreviewMaterial)
 {
 	if (!IsComponentValid())
 	{
 		return;
 	}
-
-	if (!InPreviewMaterial)
-	{
-		if (!PreviewMaterial)
-		{
-			CreatePreviewMaterial();
-		}
- 
-		InPreviewMaterial = PreviewMaterial;
- 
-		if (!InPreviewMaterial)
-		{
-			return;
-		}
-	}
  
 	UDMMaterialStage* Stage = GetStage();
 	check(Stage);
diff --git a/Engine/Plugins/Experimental/DynamicMaterial/Source/DynamicMaterialEditor/Private/Components/DMMaterialSubStage.cpp b/Engine/Plugins/Experimental/DynamicMaterial/Source/DynamicMaterialEditor/Private/Components/DMMaterialSubStage.cpp
--- a/Engine/Plugins/Experimental/DynamicMaterial/Source/DynamicMaterialEditor/Private/Components/DMMaterialSubStage.cpp
+++ b/Engine/Plugins/Experimental/DynamicMaterial/Source/DynamicMaterialEditor/Private/Components/DMMaterialSubStage.cpp
@@ -78,11 +78,6 @@ bool UDMMaterialSubStage::IsCompatibleWithNextStage(const UDMMaterialStage* Next
 	return false;
 }
 
-bool UDMMaterialSubStage::IsRootStage() const
-{
-	return false;
-}
-
 FString UDMMaterialSubStage::GetComponentPathComponent() const
 {
 	// Skip stage renaming
diff --git a/Engine/Plugins/Experimental/DynamicMaterial/Source/DynamicMaterialEditor/Private/Components/MaterialStageExpressions/DMMSEMathBase.cpp b/Engine/Plugins/Experimental/DynamicMaterial/Source/DynamicMaterialEditor/Private/Components/MaterialStageExpressions/DMMSEMathBase.cpp
--- a/Engine/Plugins/Experimental/DynamicMaterial/Source/DynamicMaterialEditor/Private/Components/MaterialStageExpressions/DMMSEMathBase.cpp
+++ b/Engine/Plugins/Experimental/DynamicMaterial/Source/DynamicMaterialEditor/Private/Components/MaterialStageExpressions/DMMSEMathBase.cpp
@@ -296,10 +296,10 @@ void UDMMaterialStageExpressionMathBase::AddDefaultInput(int32 InInputIndex) con
 	Super::AddDefaultInput(InInputIndex);
 }
 
-bool UDMMaterialStageExpressionMathBase::UpdateStagePreviewMaterial(UDMMaterialStage* InStage, UMaterial* InPreviewMaterial, 
+bool UDMMaterialStageExpressionMathBase::GenerateStagePreviewMaterial(UDMMaterialStage* InStage, UMaterial* InPreviewMaterial, 
 	UMaterialExpression*& OutMaterialExpression, int32& OutputIndex)
 {
-	UpdateOntoPreviewMaterial(InPreviewMaterial);
+	GeneratePreviewMaterial(InPreviewMaterial);
 	return true;
 }
 
@@ -335,23 +335,8 @@ void UDMMaterialStageExpressionMathBase::SetupInputs(int32 InCount)
 	}
 }
 
-void UDMMaterialStageExpressionMathBase::UpdatePreviewMaterial(UMaterial* InPreviewMaterial)
+void UDMMaterialStageExpressionMathBase::GeneratePreviewMaterial(UMaterial* InPreviewMaterial)
 {
-	if (!InPreviewMaterial)
-	{
-		if (!PreviewMaterial)
-		{
-			CreatePreviewMaterial();
-		}
-
-		InPreviewMaterial = PreviewMaterial;
-
-		if (!InPreviewMaterial)
-		{
-			return;
-		}
-	}
-
 	UDMMaterialStage* Stage = GetStage();
 	check(Stage);
 
diff --git a/Engine/Plugins/Experimental/DynamicMaterial/Source/DynamicMaterialEditor/Private/Components/MaterialStageExpressions/DMMSETextureSampleEdgeColor.cpp b/Engine/Plugins/Experimental/DynamicMaterial/Source/DynamicMaterialEditor/Private/Components/MaterialStageExpressions/DMMSETextureSampleEdgeColor.cpp
--- a/Engine/Plugins/Experimental/DynamicMaterial/Source/DynamicMaterialEditor/Private/Components/MaterialStageExpressions/DMMSETextureSampleEdgeColor.cpp
+++ b/Engine/Plugins/Experimental/DynamicMaterial/Source/DynamicMaterialEditor/Private/Components/MaterialStageExpressions/DMMSETextureSampleEdgeColor.cpp
@@ -77,13 +77,13 @@ void UDMMaterialStageExpressionTextureSampleEdgeColor::AddDefaultInput(int32 InI
 	return Super::AddDefaultInput(InInputIndex);
 }
 
-void UDMMaterialStageExpressionTextureSampleEdgeColor::PostEditChangeProperty(struct FPropertyChangedEvent& PropertyChangedEvent)
+void UDMMaterialStageExpressionTextureSampleEdgeColor::PostEditChangeProperty(FPropertyChangedEvent& InPropertyChangedEvent)
 {
-	Super::PostEditChangeProperty(PropertyChangedEvent);
+	Super::PostEditChangeProperty(InPropertyChangedEvent);
 
 	static const FName EdgeLocationName = GET_MEMBER_NAME_CHECKED(UDMMaterialStageExpressionTextureSampleEdgeColor, EdgeLocation);
 
-	if (PropertyChangedEvent.MemberProperty && PropertyChangedEvent.MemberProperty->GetFName() == EdgeLocationName)
+	if (InPropertyChangedEvent.MemberProperty && InPropertyChangedEvent.MemberProperty->GetFName() == EdgeLocationName)
 	{
 		OnEdgeLocationChanged();
 	}
diff --git a/Engine/Plugins/Experimental/DynamicMaterial/Source/DynamicMaterialEditor/Private/Components/MaterialStageInputs/DMMSITextureUV.cpp b/Engine/Plugins/Experimental/DynamicMaterial/Source/DynamicMaterialEditor/Private/Components/MaterialStageInputs/DMMSITextureUV.cpp
--- a/Engine/Plugins/Experimental/DynamicMaterial/Source/DynamicMaterialEditor/Private/Components/MaterialStageInputs/DMMSITextureUV.cpp
+++ b/Engine/Plugins/Experimental/DynamicMaterial/Source/DynamicMaterialEditor/Private/Components/MaterialStageInputs/DMMSITextureUV.cpp
@@ -304,9 +304,15 @@ void UDMMaterialStageInputTextureUV::OnTextureUVUpdated(UDMMaterialComponent* In
 	}
 }
 
-UMaterialExpressionScalarParameter* UDMMaterialStageInputTextureUV::CreateScalarParameter(const TSharedRef<FDMMaterialBuildState>& InBuildState, FName InParamName, float InValue)
+UMaterialExpressionScalarParameter* UDMMaterialStageInputTextureUV::CreateScalarParameter(const TSharedRef<FDMMaterialBuildState>& InBuildState, 
+	FName InParamName, EDMMaterialParameterGroup InParameterGroup, float InValue)
 {
-	UMaterialExpressionScalarParameter* NewExpression = InBuildState->GetBuildUtils().CreateExpressionParameter<UMaterialExpressionScalarParameter>(InParamName, UE_DM_NodeComment_Default);
+	UMaterialExpressionScalarParameter* NewExpression = InBuildState->GetBuildUtils().CreateExpressionParameter<UMaterialExpressionScalarParameter>(
+		InParamName, 
+		InParameterGroup,
+		UE_DM_NodeComment_Default
+	);
+
 	check(NewExpression);
 
 	NewExpression->DefaultValue = InValue;
@@ -383,7 +389,6 @@ TArray<UMaterialExpression*> UDMMaterialStageInputTextureUV::CreateTextureUVExpr
 
 	// Output nodes
 	TArray<UMaterialExpression*> Nodes;
-	InTextureUV->MaterialNodesCreated.Empty();
 
 	// UV Source
 	static const FName UVInputName = TEXT("UV");
@@ -425,33 +430,26 @@ TArray<UMaterialExpression*> UDMMaterialStageInputTextureUV::CreateTextureUVExpr
 	Nodes.Last()->ConnectExpression(TextureUVFunc->GetInput(NameToInputIndex[UVInputName]), 0);
 
 	auto CreateParameterExpression = [InTextureUV, &InBuildState, &NameToInputIndex, TextureUVFunc, &Nodes]
-		(int32 ParamGroupId, int32 ParamId, FName InputName, float DefaultValue)
+		(FName InPropertyName, int32 InComponent, FName InInputName, float InDefaultValue)
 		{
-			check(InTextureUV->MaterialParameters.Contains(ParamId));
-
-			if (!NameToInputIndex.Contains(InputName))
-			{
-				return;
-			}
-
-			check(NameToInputIndex.Contains(InputName));
+			UMaterialExpressionScalarParameter* ParamNode = CreateScalarParameter(
+				InBuildState, 
+				InTextureUV->GetMaterialParameterName(InPropertyName, InComponent),
+				InTextureUV->GetParameterGroup(InPropertyName, InComponent)
+			);
 
-			UMaterialExpressionScalarParameter* ParamNode = CreateScalarParameter(InBuildState, InTextureUV->MaterialParameters[ParamId]->GetParameterName());
-			InTextureUV->MaterialNodesCreated.FindOrAdd(ParamGroupId, true);
-
-			ParamNode->DefaultValue = DefaultValue;
-			ParamNode->ConnectExpression(TextureUVFunc->GetInput(NameToInputIndex[InputName]), 0);
+			ParamNode->DefaultValue = InDefaultValue;
+			ParamNode->ConnectExpression(TextureUVFunc->GetInput(NameToInputIndex[InInputName]), 0);
 
 			Nodes.Add(ParamNode);
 		};
-
-	CreateParameterExpression(ParamID::Offset,   ParamID::OffsetX,  TEXT("OffsetX"),  InTextureUV->GetOffset().X);
-	CreateParameterExpression(ParamID::Offset,   ParamID::OffsetY,  TEXT("OffsetY"),  InTextureUV->GetOffset().Y);
-	CreateParameterExpression(ParamID::Pivot,    ParamID::PivotX,   TEXT("PivotX"),   InTextureUV->GetPivot().X);
-	CreateParameterExpression(ParamID::Pivot,    ParamID::PivotY,   TEXT("PivotY"),   InTextureUV->GetPivot().Y);
-	CreateParameterExpression(ParamID::Rotation, ParamID::Rotation, TEXT("Rotation"), InTextureUV->GetRotation());
-	CreateParameterExpression(ParamID::Tiling,    ParamID::TilingX,   TEXT("TilingX"),   InTextureUV->GetTiling().X);
-	CreateParameterExpression(ParamID::Tiling,    ParamID::TilingY,   TEXT("TilingY"),   InTextureUV->GetTiling().Y);
+	CreateParameterExpression(UDMTextureUV::NAME_Offset,   0, TEXT("OffsetX"),  InTextureUV->GetOffset().X);
+	CreateParameterExpression(UDMTextureUV::NAME_Offset,   1, TEXT("OffsetY"),  InTextureUV->GetOffset().Y);
+	CreateParameterExpression(UDMTextureUV::NAME_Pivot,    0, TEXT("PivotX"),   InTextureUV->GetPivot().X);
+	CreateParameterExpression(UDMTextureUV::NAME_Pivot,    1, TEXT("PivotY"),   InTextureUV->GetPivot().Y);
+	CreateParameterExpression(UDMTextureUV::NAME_Rotation, 0, TEXT("Rotation"), InTextureUV->GetRotation());
+	CreateParameterExpression(UDMTextureUV::NAME_Tiling,   0, TEXT("TilingX"),  InTextureUV->GetTiling().X);
+	CreateParameterExpression(UDMTextureUV::NAME_Tiling,   1, TEXT("TilingY"),  InTextureUV->GetTiling().Y);
 
 	if (UMaterialExpression* GlobalOffset = InBuildState->GetGlobalExpression(UDynamicMaterialModel::GlobalOffsetValueName))
 	{
