Material Designer

diff --git a/Engine/Plugins/Experimental/DynamicMaterial/Source/DynamicMaterialEditor/Private/Components/MaterialStageInputs/DMMSIValue.cpp b/Engine/Plugins/Experimental/DynamicMaterial/Source/DynamicMaterialEditor/Private/Components/MaterialStageInputs/DMMSIValue.cpp
--- a/Engine/Plugins/Experimental/DynamicMaterial/Source/DynamicMaterialEditor/Private/Components/MaterialStageInputs/DMMSIValue.cpp
+++ b/Engine/Plugins/Experimental/DynamicMaterial/Source/DynamicMaterialEditor/Private/Components/MaterialStageInputs/DMMSIValue.cpp
@@ -376,24 +376,10 @@ FText UDMMaterialStageInputValue::GetChannelDescription(const FDMMaterialStageCo
 		return LOCTEXT("LocalValue", "Local Value");
 	}
 
-	if (Value->GetParameter())
-	{
-		static const FText FormatTemplateNamed = LOCTEXT("ChannelDescriptionGlobalNamed", "{0} (Global)");
-
-		return FText::Format(
-			FormatTemplateNamed,
-			FText::FromName(Value->GetParameter()->GetParameterName())
-		);
-	}
-	else
-	{
-		static const FText FormatTemplate = LOCTEXT("ChannelDescriptionGlobal", "Value {0} (Global)");
-
-		return FText::Format(
-			FormatTemplate,
-			FText::AsNumber(Value->FindIndexSafe())
-		);
-	}
+	return FText::Format(
+		LOCTEXT("ChannelDescriptionGlobal", "Value {0} (Global)"),
+		FText::FromName(Value->GetMaterialParameterName())
+	);
 }
 
 void UDMMaterialStageInputValue::SetValue(UDMMaterialValue* InValue)
diff --git a/Engine/Plugins/Experimental/DynamicMaterial/Source/DynamicMaterialEditor/Private/DMContentBrowserIntegration.cpp b/Engine/Plugins/Experimental/DynamicMaterial/Source/DynamicMaterialEditor/Private/DMContentBrowserIntegration.cpp
--- a/Engine/Plugins/Experimental/DynamicMaterial/Source/DynamicMaterialEditor/Private/DMContentBrowserIntegration.cpp
+++ b/Engine/Plugins/Experimental/DynamicMaterial/Source/DynamicMaterialEditor/Private/DMContentBrowserIntegration.cpp
@@ -4,44 +4,66 @@
 
 #include "AssetRegistry/AssetRegistryModule.h"
 #include "AssetToolsModule.h"
+#include "ContentBrowserModule.h"
 #include "CoreGlobals.h"
 #include "DMTextureSet.h"
 #include "DMTextureSetBlueprintFunctionLibrary.h"
 #include "DMTextureSetContentBrowserIntegration.h"
 #include "Engine/Texture.h"
 #include "Framework/MultiBox/MultiBoxBuilder.h"
+#include "IContentBrowserSingleton.h"
 #include "IDynamicMaterialEditorModule.h"
 #include "Material/DynamicMaterialInstance.h"
 #include "Material/DynamicMaterialInstanceFactory.h"
 #include "Materials/Material.h"
+#include "Misc/MessageDialog.h"
+#include "Model/DynamicMaterialModel.h"
+#include "Model/DynamicMaterialModelDynamic.h"
 #include "Model/DynamicMaterialModelEditorOnlyData.h"
+#include "PackageTools.h"
 #include "ScopedTransaction.h"
 #include "Styling/SlateIconFinder.h"
 
 #define LOCTEXT_NAMESPACE "FDMContentBrowserIntegration"
 
-FDelegateHandle FDMContentBrowserIntegration::PopulateHandle;
+FDelegateHandle FDMContentBrowserIntegration::TextureSetPopulateHandle;
+FDelegateHandle FDMContentBrowserIntegration::ContentBrowserAssetHandle;
 
 void FDMContentBrowserIntegration::Integrate()
 {
-	if (PopulateHandle.IsValid())
-	{
-		Disintegrate();
-	}
+	Disintegrate();
 
-	PopulateHandle = FDMTextureSetContentBrowserIntegration::GetPopulateExtenderDelegate().AddStatic(&FDMContentBrowserIntegration::ExtendMenu);
+	TextureSetPopulateHandle = FDMTextureSetContentBrowserIntegration::GetPopulateExtenderDelegate().AddStatic(&FDMContentBrowserIntegration::ExtendMenu);
+
+	FContentBrowserModule& ContentBrowserModule = FModuleManager::LoadModuleChecked<FContentBrowserModule>("ContentBrowser");
+	TArray<FContentBrowserMenuExtender_SelectedAssets>& CBMenuExtenderDelegates = ContentBrowserModule.GetAllAssetViewContextMenuExtenders();
+	CBMenuExtenderDelegates.Add(FContentBrowserMenuExtender_SelectedAssets::CreateStatic(&FDMContentBrowserIntegration::OnExtendContentBrowserAssetSelectionMenu));
+	ContentBrowserAssetHandle = CBMenuExtenderDelegates.Last().GetHandle();
 }
 
 void FDMContentBrowserIntegration::Disintegrate()
 {
-	if (!PopulateHandle.IsValid())
+	if (TextureSetPopulateHandle.IsValid())
 	{
-		return;
+		FDMTextureSetContentBrowserIntegration::GetPopulateExtenderDelegate().Remove(TextureSetPopulateHandle);
+		TextureSetPopulateHandle.Reset();
 	}
 
-	FDMTextureSetContentBrowserIntegration::GetPopulateExtenderDelegate().Remove(PopulateHandle);
-
-	PopulateHandle.Reset();
+	if (ContentBrowserAssetHandle.IsValid())
+	{
+		if (FContentBrowserModule* ContentBrowserModule = FModuleManager::GetModulePtr<FContentBrowserModule>("ContentBrowser"))
+		{
+			TArray<FContentBrowserMenuExtender_SelectedAssets>& CBMenuExtenderDelegates = ContentBrowserModule->GetAllAssetViewContextMenuExtenders();
+
+			CBMenuExtenderDelegates.RemoveAll(
+				[](const FContentBrowserMenuExtender_SelectedAssets& InElement)
+				{
+					return InElement.GetHandle() == ContentBrowserAssetHandle;
+				});
+
+			ContentBrowserAssetHandle.Reset();
+		}
+	}
 }
 
 void FDMContentBrowserIntegration::ExtendMenu(FMenuBuilder& InMenuBuilder, const TArray<FAssetData>& InSelectedAssets)
@@ -167,7 +189,7 @@ void FDMContentBrowserIntegration::OnUpdateMaterialDesignerInstanceFromTextureSe
 
 	IDynamicMaterialEditorModule& DynamicMaterialEditorModule = IDynamicMaterialEditorModule::Get();
 
-	UDynamicMaterialModel* Model = DynamicMaterialEditorModule.GetOpenedMaterialModel(nullptr);
+	UDynamicMaterialModelBase* Model = DynamicMaterialEditorModule.GetOpenedMaterialModel(nullptr);
 
 	if (!Model)
 	{
@@ -192,4 +214,257 @@ void FDMContentBrowserIntegration::OnUpdateMaterialDesignerInstanceFromTextureSe
 	}
 }
 
+TSharedRef<FExtender> FDMContentBrowserIntegration::OnExtendContentBrowserAssetSelectionMenu(const TArray<FAssetData>& InSelectedAssets)
+{
+	TSharedRef<FExtender> Extender = MakeShared<FExtender>();
+	bool bHasMaterialDesignerAsset = false;
+
+	for (const FAssetData& SelectedAsset : InSelectedAssets)
+	{
+		if (UClass* AssetClass = SelectedAsset.GetClass(EResolveClass::Yes))
+		{
+			if (AssetClass->IsChildOf<UDynamicMaterialInstance>() || AssetClass->IsChildOf<UDynamicMaterialModel>())
+			{
+				bHasMaterialDesignerAsset = true;
+				break;
+			}
+		}
+	}
+
+	if (!bHasMaterialDesignerAsset)
+	{
+		return Extender;
+	}
+
+	Extender->AddMenuExtension(
+		"GetAssetActions",
+		EExtensionHook::After,
+		nullptr,
+		FMenuExtensionDelegate::CreateLambda(
+			[InSelectedAssets](FMenuBuilder& InMenuBuilder)
+			{
+				InMenuBuilder.AddMenuEntry(
+					LOCTEXT("CreateDynamic", "Create Material Designer Dynamic"),
+					LOCTEXT("CreateDynamicooltip", "Create a dynamic instance from a Material Designer Instance or Model."),
+					FSlateIconFinder::FindIconForClass(UMaterial::StaticClass()),
+					FUIAction(FExecuteAction::CreateStatic(&FDMContentBrowserIntegration::CreateDynamic, InSelectedAssets))
+				);
+			}
+		)
+	);
+
+	return Extender;
+}
+
+void FDMContentBrowserIntegration::CreateDynamic(TArray<FAssetData> InSelectedAssets)
+{
+	for (const FAssetData& SelectedAsset : InSelectedAssets)
+	{
+		if (UClass* AssetClass = SelectedAsset.GetClass(EResolveClass::Yes))
+		{
+			if (AssetClass->IsChildOf<UDynamicMaterialModel>())
+			{
+				CreateModelDynamic(Cast<UDynamicMaterialModel>(SelectedAsset.GetAsset()));
+				break;
+			}
+
+			if (AssetClass->IsChildOf<UDynamicMaterialInstance>())
+			{
+				CreateInstanceDynamic(Cast<UDynamicMaterialInstance>(SelectedAsset.GetAsset()));
+				break;
+			}
+		}
+	}
+}
+
+void FDMContentBrowserIntegration::CreateModelDynamic(UDynamicMaterialModel* InModel)
+{
+	if (!InModel)
+	{
+		return;
+	}
+
+	UMaterial* ParentMaterial = InModel->GetGeneratedMaterial();
+
+	if (!ParentMaterial)
+	{
+		return;
+	}
+
+	if (!ParentMaterial->HasAnyFlags(RF_Public))
+	{
+		const EAppReturnType::Type Result = FMessageDialog::Open(
+			EAppMsgType::YesNo,
+			LOCTEXT("ExportMaterialFromModel", 
+				"Generating a Material Designer Dynamic requires that the generated material be exported from its package.\n\n"
+				"The package containing the material will be saved. This may be a level.\n\n"
+				"Continue?")
+		);
+
+		switch (Result)
+		{
+			case EAppReturnType::Yes:
+				ParentMaterial->Modify(/* Always Mark Dirty */ true);
+				ParentMaterial->SetFlags(RF_Public);
+				UPackageTools::SavePackagesForObjects({InModel});
+				break;
+
+			default:
+				return;
+		}
+	}
+
+	UDynamicMaterialModelDynamic* ModelDynamic = UDynamicMaterialModelDynamic::Create(GetTransientPackage(), InModel);
+
+	if (!ModelDynamic)
+	{
+		return;
+	}
+
+	IAssetTools& AssetTools = FModuleManager::LoadModuleChecked<FAssetToolsModule>("AssetTools").Get();
+	FString PackageName, AssetName;
+	AssetTools.CreateUniqueAssetName(InModel->GetName(), TEXT(""), PackageName, AssetName);
+
+	IContentBrowserSingleton& ContentBrowser = FModuleManager::LoadModuleChecked<FContentBrowserModule>("ContentBrowser").Get();
+	const FContentBrowserItemPath CurrentPath = ContentBrowser.GetCurrentPath();
+	const FString PathStr = CurrentPath.HasInternalPath() ? CurrentPath.GetInternalPathString() : "/Game";
+
+	FSaveAssetDialogConfig SaveAssetDialogConfig;
+	SaveAssetDialogConfig.DialogTitleOverride = LOCTEXT("SaveAssetDialogTitle", "Save Asset As");
+	SaveAssetDialogConfig.DefaultPath = PathStr;
+	SaveAssetDialogConfig.ExistingAssetPolicy = ESaveAssetDialogExistingAssetPolicy::Disallow;
+	SaveAssetDialogConfig.DefaultAssetName = AssetName;
+
+	FContentBrowserModule& ContentBrowserModule = FModuleManager::LoadModuleChecked<FContentBrowserModule>("ContentBrowser");
+	FString SaveObjectPath = ContentBrowserModule.Get().CreateModalSaveAssetDialog(SaveAssetDialogConfig);
+
+	if (SaveObjectPath.IsEmpty())
+	{
+		return;
+	}
+
+	PackageName = FPaths::GetBaseFilename(*SaveObjectPath, false);
+
+	UPackage* Package = CreatePackage(*PackageName);
+
+	if (!Package)
+	{
+		return;
+	}
+
+	AssetName = FPaths::GetBaseFilename(*SaveObjectPath, true);
+
+	ModelDynamic->SetFlags(RF_Standalone | RF_Public);
+	ModelDynamic->Rename(*AssetName, Package, REN_DontCreateRedirectors);
+
+	FAssetRegistryModule::AssetCreated(ModelDynamic);
+}
+
+void FDMContentBrowserIntegration::CreateInstanceDynamic(UDynamicMaterialInstance* InInstance)
+{
+	if (!InInstance)
+	{
+		return;
+	}
+
+	UDynamicMaterialModel* Model = InInstance->GetMaterialModel();
+
+	if (!Model)
+	{
+		return;
+	}
+
+	UMaterial* ParentMaterial = Model->GetGeneratedMaterial();
+
+	if (!ParentMaterial)
+	{
+		return;
+	}
+
+	if (!ParentMaterial->HasAnyFlags(RF_Public) || !Model->HasAnyFlags(RF_Public))
+	{
+		const EAppReturnType::Type Result = FMessageDialog::Open(
+			EAppMsgType::YesNo,
+			LOCTEXT("ExportMaterialFromInstance",
+				"Generating a Material Designer Dynamic requires that the generated material and material model be exported from this package.\n\n"
+				"The package containing the material will be saved. This may be a level.\n\n"
+				"Continue?")
+		);
+
+		switch (Result)
+		{
+			case EAppReturnType::Yes:
+				Model->Modify(/* Always Mark Dirty */ true);
+				Model->SetFlags(RF_Public);
+				ParentMaterial->Modify(/* Always Mark Dirty */ true);
+				ParentMaterial->SetFlags(RF_Public);
+				UPackageTools::SavePackagesForObjects({InInstance});
+				break;
+
+			default:
+				return;
+		}
+	}
+
+	UDynamicMaterialModelDynamic* ModelDynamic = UDynamicMaterialModelDynamic::Create(GetTransientPackage(), Model);
+
+	if (!ModelDynamic)
+	{
+		return;
+	}
+
+	UDynamicMaterialInstance* Instance = NewObject<UDynamicMaterialInstance>(GetTransientPackage());
+
+	if (!Instance)
+	{
+		return;
+	}
+
+	Instance->SetMaterialModel(ModelDynamic);
+	ModelDynamic->SetDynamicMaterialInstance(Instance);
+	Instance->InitializeMIDPublic();
+
+	FString CurrentName = InInstance->GetName();
+	CurrentName = CurrentName.StartsWith(TEXT("MDI_"))
+		? (TEXT("MDD_") + CurrentName.RightChop(4))
+		: (TEXT("MDD_") + CurrentName);
+
+	IAssetTools& AssetTools = FModuleManager::LoadModuleChecked<FAssetToolsModule>("AssetTools").Get();
+	FString PackageName, AssetName;
+	AssetTools.CreateUniqueAssetName(CurrentName, TEXT(""), PackageName, AssetName);
+
+	IContentBrowserSingleton& ContentBrowser = FModuleManager::LoadModuleChecked<FContentBrowserModule>("ContentBrowser").Get();
+	const FContentBrowserItemPath CurrentPath = ContentBrowser.GetCurrentPath();
+	const FString PathStr = CurrentPath.HasInternalPath() ? CurrentPath.GetInternalPathString() : "/Game";
+
+	FSaveAssetDialogConfig SaveAssetDialogConfig;
+	SaveAssetDialogConfig.DialogTitleOverride = LOCTEXT("SaveAssetDialogTitle", "Save Asset As");
+	SaveAssetDialogConfig.DefaultPath = PathStr;
+	SaveAssetDialogConfig.ExistingAssetPolicy = ESaveAssetDialogExistingAssetPolicy::Disallow;
+	SaveAssetDialogConfig.DefaultAssetName = AssetName;
+
+	FContentBrowserModule& ContentBrowserModule = FModuleManager::LoadModuleChecked<FContentBrowserModule>("ContentBrowser");
+	FString SaveObjectPath = ContentBrowserModule.Get().CreateModalSaveAssetDialog(SaveAssetDialogConfig);
+
+	if (SaveObjectPath.IsEmpty())
+	{
+		return;
+	}
+
+	PackageName = FPaths::GetBaseFilename(*SaveObjectPath, false);
+	UPackage* Package = CreatePackage(*PackageName);
+
+	if (!Package)
+	{
+		return;
+	}
+
+	AssetName = FPaths::GetBaseFilename(*SaveObjectPath, true);
+
+	Instance->SetFlags(RF_Standalone | RF_Public | RF_Transactional);
+	Instance->Rename(*AssetName, Package, REN_DontCreateRedirectors);
+
+	FAssetRegistryModule::AssetCreated(Instance);
+}
+
 #undef LOCTEXT_NAMESPACE
diff --git a/Engine/Plugins/Experimental/DynamicMaterial/Source/DynamicMaterialEditor/Private/DMContentBrowserIntegration.h b/Engine/Plugins/Experimental/DynamicMaterial/Source/DynamicMaterialEditor/Private/DMContentBrowserIntegration.h
--- a/Engine/Plugins/Experimental/DynamicMaterial/Source/DynamicMaterialEditor/Private/DMContentBrowserIntegration.h
+++ b/Engine/Plugins/Experimental/DynamicMaterial/Source/DynamicMaterialEditor/Private/DMContentBrowserIntegration.h
@@ -4,9 +4,13 @@
 
 #include "AssetRegistry/AssetData.h"
 #include "Containers/ContainersFwd.h"
+#include "Material/DynamicMaterialInstance.h"
 
+class FExtender;
 class FMenuBuilder;
 class UDMTextureSet;
+class UDynamicMaterialInstance;
+class UDynamicMaterialModel;
 struct FAssetData;
 
 class FDMContentBrowserIntegration
@@ -19,6 +23,9 @@ public:
 	static void UpdateMaterialDesignerInstanceFromTextureSet(TArray<FAssetData> InSelectedAssets, bool bInReplace);
 
 protected:
+	static FDelegateHandle TextureSetPopulateHandle;
+	static FDelegateHandle ContentBrowserAssetHandle;
+
 	static void ExtendMenu(FMenuBuilder& InMenuBuilder, const TArray<FAssetData>& InSelectedAssets);
 
 	static void CreateMaterialDesignerInstanceFromTextureSet(TArray<FAssetData> InSelectedAssets);
@@ -27,5 +34,11 @@ protected:
 
 	static void OnUpdateMaterialDesignerInstanceFromTextureSetComplete(UDMTextureSet* InTextureSet, bool bInAccepted, bool bInReplace);
 
-	static FDelegateHandle PopulateHandle;
+	static TSharedRef<FExtender> OnExtendContentBrowserAssetSelectionMenu(const TArray<FAssetData>& InSelectedAssets);
+
+	static void CreateDynamic(TArray<FAssetData> InSelectedAssets);
+
+	static void CreateModelDynamic(UDynamicMaterialModel* InModel);
+
+	static void CreateInstanceDynamic(UDynamicMaterialInstance* InInstance);
 };
diff --git a/Engine/Plugins/Experimental/DynamicMaterial/Source/DynamicMaterialEditor/Private/DMObjectMaterialProperty.cpp b/Engine/Plugins/Experimental/DynamicMaterial/Source/DynamicMaterialEditor/Private/DMObjectMaterialProperty.cpp
--- a/Engine/Plugins/Experimental/DynamicMaterial/Source/DynamicMaterialEditor/Private/DMObjectMaterialProperty.cpp
+++ b/Engine/Plugins/Experimental/DynamicMaterial/Source/DynamicMaterialEditor/Private/DMObjectMaterialProperty.cpp
@@ -37,11 +37,11 @@ FDMObjectMaterialProperty::FDMObjectMaterialProperty(UObject* InOuter, FProperty
 {
 }
 
-UDynamicMaterialModel* FDMObjectMaterialProperty::GetMaterialModel() const
+UDynamicMaterialModelBase* FDMObjectMaterialProperty::GetMaterialModelBase() const
 {
 	if (UDynamicMaterialInstance* Instance = GetMaterial())
 	{
-		return Instance->GetMaterialModel();
+		return Instance->GetMaterialModelBase();
 	}
 
 	return nullptr;
@@ -213,7 +213,7 @@ FText FDMObjectMaterialProperty::GetPropertyName(bool bInIgnoreNewStatus) const
 		return FText::GetEmpty();
 	}
 
-	UDynamicMaterialModel* MaterialModel = GetMaterialModel();
+	UDynamicMaterialModelBase* MaterialModelBase = GetMaterialModelBase();
 
 	if (Property != nullptr)
 	{
@@ -228,7 +228,7 @@ FText FDMObjectMaterialProperty::GetPropertyName(bool bInIgnoreNewStatus) const
 			);
 		}
 
-		if (MaterialModel || bInIgnoreNewStatus)
+		if (MaterialModelBase || bInIgnoreNewStatus)
 		{
 			return FText::Format(
 				LOCTEXT("PropertyNameFormat", "{0}"),
@@ -249,7 +249,7 @@ FText FDMObjectMaterialProperty::GetPropertyName(bool bInIgnoreNewStatus) const
 		{
 			if (Index >= 0 && Index <= Component->GetNumMaterials())
 			{
-				if (MaterialModel || bInIgnoreNewStatus)
+				if (MaterialModelBase || bInIgnoreNewStatus)
 				{
 					return FText::Format(
 						LOCTEXT("MaterialListNameFormat", "Element {0}"),
diff --git a/Engine/Plugins/Experimental/DynamicMaterial/Source/DynamicMaterialEditor/Private/DMWorldSubsystem.cpp b/Engine/Plugins/Experimental/DynamicMaterial/Source/DynamicMaterialEditor/Private/DMWorldSubsystem.cpp
--- a/Engine/Plugins/Experimental/DynamicMaterial/Source/DynamicMaterialEditor/Private/DMWorldSubsystem.cpp
+++ b/Engine/Plugins/Experimental/DynamicMaterial/Source/DynamicMaterialEditor/Private/DMWorldSubsystem.cpp
@@ -13,7 +13,7 @@ UDMWorldSubsystem::UDMWorldSubsystem()
 	InvokeTabDelegate.BindWeakLambda(this, [this]() { FDMLevelEditorIntegration::InvokeTabForWorld(GetWorld()); });
 }
 
-UDynamicMaterialModel* UDMWorldSubsystem::ExecuteGetCustomEditorModelDelegate()
+UDynamicMaterialModelBase* UDMWorldSubsystem::ExecuteGetCustomEditorModelDelegate()
 {
 	if (CustomModelEditorGetDelegate.IsBound())
 	{
@@ -24,7 +24,7 @@ UDynamicMaterialModel* UDMWorldSubsystem::ExecuteGetCustomEditorModelDelegate()
 	return nullptr;
 }
 
-void UDMWorldSubsystem::ExecuteSetCustomEditorModelDelegate(UDynamicMaterialModel* InMaterialModel)
+void UDMWorldSubsystem::ExecuteSetCustomEditorModelDelegate(UDynamicMaterialModelBase* InMaterialModel)
 {
 	if (!InMaterialModel || IsValid(InMaterialModel))
 	{
@@ -45,7 +45,7 @@ void UDMWorldSubsystem::ExecuteSetCustomEditorActorDelegate(AActor* InActor)
 	}
 }
 
-bool UDMWorldSubsystem::ExecuteIsValidDelegate(UDynamicMaterialModel* InMaterialModel)
+bool UDMWorldSubsystem::ExecuteIsValidDelegate(UDynamicMaterialModelBase* InMaterialModel)
 {
 	if (!IsValid(InMaterialModel))
 	{
diff --git a/Engine/Plugins/Experimental/DynamicMaterial/Source/DynamicMaterialEditor/Private/DetailsPanel/DMValueDetailsRowExtensions.cpp b/Engine/Plugins/Experimental/DynamicMaterial/Source/DynamicMaterialEditor/Private/DetailsPanel/DMValueDetailsRowExtensions.cpp
--- a/Engine/Plugins/Experimental/DynamicMaterial/Source/DynamicMaterialEditor/Private/DetailsPanel/DMValueDetailsRowExtensions.cpp
+++ b/Engine/Plugins/Experimental/DynamicMaterial/Source/DynamicMaterialEditor/Private/DetailsPanel/DMValueDetailsRowExtensions.cpp
@@ -29,7 +29,11 @@ namespace UE::DynamicMaterialEditor::Private
 	const FText ParameterNameMenuEntry = LOCTEXT("ParameterName", "Parameter Name");
 	const FText ParameterNameMenuEntryX = LOCTEXT("ParameterNameX", "Parameter Name X");
 	const FText ParameterNameMenuEntryY = LOCTEXT("ParameterNameY", "Parameter Name Y");
+	const FText ParameterExposeMenuEntry = LOCTEXT("ParameterExpose", "Expose Parameter");
+	const FText ParameterExposeMenuEntryX = LOCTEXT("ParameterExposeX", "Expose Parameter X");
+	const FText ParameterExposeMenuEntryY = LOCTEXT("ParameterExposeY", "Expose Parameter Y");
 	const FText SetParameterNameToolTip = LOCTEXT("SetParameterNameToolTip", "Set the name of the parameter this property is exposed as within the generated material.");
+	const FText ExposeParameterToolTip = LOCTEXT("ExposeParameterToolTip", "When unchecked, parameter will appear in the \"99 - Uncategorized\" category.");
 }
 
 FDMValueDetailsRowExtensions& FDMValueDetailsRowExtensions::Get()
@@ -157,7 +161,52 @@ void FDMValueDetailsRowExtensions::FillPropertyRightClickMenu_TextureUV(FToolMen
 
 	TWeakObjectPtr<UDMTextureUV> TextureUVWeak = InTextureUV;
 
-	auto CreateMenuEntry = [&InSection, &TextureUVWeak, &InPropertyName](const FText& InMenuText, int32 InComponent)
+	auto CreateParameterExposeMenuEntry = [&InSection, &TextureUVWeak, &InPropertyName](const FText& InMenuText, int32 InComponent)
+		{
+			FUIAction ExposeTextureUVXAction;
+			ExposeTextureUVXAction.ExecuteAction = FExecuteAction::CreateLambda(
+				[TextureUVWeak, InPropertyName, InComponent]()
+				{
+					if (UDMTextureUV* TextureUV = TextureUVWeak.Get())
+					{
+						TextureUV->SetShouldExposeParameter(InPropertyName, InComponent, !TextureUV->GetShouldExposeParameter(InPropertyName, InComponent));
+					}
+				});
+
+			ExposeTextureUVXAction.CanExecuteAction = FCanExecuteAction::CreateLambda(
+				[TextureUVWeak]()
+				{
+					return TextureUVWeak.IsValid();
+				});
+
+			ExposeTextureUVXAction.GetActionCheckState = FGetActionCheckState::CreateLambda(
+				[TextureUVWeak, InPropertyName, InComponent]()
+				{
+					if (UDMTextureUV* TextureUV = TextureUVWeak.Get())
+					{
+						if (TextureUV->GetShouldExposeParameter(InPropertyName, InComponent))
+						{
+							return ECheckBoxState::Checked;
+						}
+					}
+
+					return ECheckBoxState::Unchecked;
+				});
+
+			const FName MenuNameBase = *(InPropertyName.ToString() + TEXT("Expose"));
+			const FName MenuName = FName(MenuNameBase, InComponent);
+
+			InSection.AddEntry(FToolMenuEntry::InitMenuEntry(
+				MenuName,
+				InMenuText,
+				ExposeParameterToolTip,
+				TAttribute<FSlateIcon>(),
+				FToolUIActionChoice(ExposeTextureUVXAction),
+				EUserInterfaceActionType::ToggleButton
+			));
+		};
+
+	auto CreateParameterNameMenuEntry = [&InSection, &TextureUVWeak, &InPropertyName](const FText& InMenuText, int32 InComponent)
 		{
 			FNewToolMenuChoice RenameChoice(FOnGetContent::CreateLambda(
 				[TextureUVWeak, InPropertyName, InComponent]() -> TSharedRef<SWidget>
@@ -201,7 +250,8 @@ void FDMValueDetailsRowExtensions::FillPropertyRightClickMenu_TextureUV(FToolMen
 							});
 				}));
 
-			FName MenuName = FName(InPropertyName, InComponent);
+			const FName MenuNameBase = *(InPropertyName.ToString() + TEXT("Name"));
+			const FName MenuName = FName(MenuNameBase, InComponent);
 
 			InSection.AddSubMenu(
 				MenuName,
@@ -216,15 +266,18 @@ void FDMValueDetailsRowExtensions::FillPropertyRightClickMenu_TextureUV(FToolMen
 
 	if (InPropertyName == UDMTextureUV::NAME_Rotation)
 	{
-		CreateMenuEntry(ParameterNameMenuEntry, ComponentX);
+		CreateParameterExposeMenuEntry(ParameterExposeMenuEntry, ComponentX);
+		CreateParameterNameMenuEntry(ParameterNameMenuEntry, ComponentX);
 	}
 
 	if (InPropertyName == UDMTextureUV::NAME_Offset 
 		|| InPropertyName == UDMTextureUV::NAME_Pivot
 		|| InPropertyName == UDMTextureUV::NAME_Tiling)
 	{
-		CreateMenuEntry(ParameterNameMenuEntryX, ComponentX);
-		CreateMenuEntry(ParameterNameMenuEntryY, ComponentY);
+		CreateParameterExposeMenuEntry(ParameterExposeMenuEntryX, ComponentX);
+		CreateParameterExposeMenuEntry(ParameterExposeMenuEntryY, ComponentY);
+		CreateParameterNameMenuEntry(ParameterNameMenuEntryX, ComponentX);
+		CreateParameterNameMenuEntry(ParameterNameMenuEntryY, ComponentY);
 	}
 }
 
@@ -234,6 +287,45 @@ void FDMValueDetailsRowExtensions::FillPropertyRightClickMenu_Value(FToolMenuSec
 
 	TWeakObjectPtr<UDMMaterialValue> ValueWeak = InValue;
 
+	FUIAction ExposeValueAction;
+	ExposeValueAction.ExecuteAction = FExecuteAction::CreateLambda(
+		[ValueWeak]()
+		{
+			if (UDMMaterialValue* Value = ValueWeak.Get())
+			{
+				Value->SetShouldExposeParameter(!Value->GetShouldExposeParameter());
+			}
+		});
+
+	ExposeValueAction.CanExecuteAction = FCanExecuteAction::CreateLambda(
+		[ValueWeak]()
+		{
+			return ValueWeak.IsValid();
+		});
+
+	ExposeValueAction.GetActionCheckState = FGetActionCheckState::CreateLambda(
+		[ValueWeak]()
+		{
+			if (UDMMaterialValue* Value = ValueWeak.Get())
+			{
+				if (Value->GetShouldExposeParameter())
+				{
+					return ECheckBoxState::Checked;
+				}
+			}
+
+			return ECheckBoxState::Unchecked;
+		});
+
+	InSection.AddEntry(FToolMenuEntry::InitMenuEntry(
+		"ExposeParameter",
+		LOCTEXT("ExposeParameter", "Expose Parameter"),
+		ExposeParameterToolTip,
+		TAttribute<FSlateIcon>(),
+		FToolUIActionChoice(ExposeValueAction),
+		EUserInterfaceActionType::ToggleButton
+	));
+
 	FNewToolMenuChoice RenameChoice(FOnGetContent::CreateLambda(
 		[ValueWeak]() -> TSharedRef<SWidget>
 		{
@@ -277,7 +369,7 @@ void FDMValueDetailsRowExtensions::FillPropertyRightClickMenu_Value(FToolMenuSec
 		}));
 
 	InSection.AddSubMenu(
-		UDMMaterialValue::ValueName,
+		"RenameParameter",
 		ParameterNameMenuEntry,
 		SetParameterNameToolTip,
 		RenameChoice,
@@ -298,7 +390,7 @@ void FDMValueDetailsRowExtensions::SetValueParameterName(TWeakObjectPtr<UDMMater
 		return;
 	}
 
-	if (InName == NAME_None)
+	if (InName.IsNone())
 	{
 		return;
 	}
@@ -329,7 +421,7 @@ void FDMValueDetailsRowExtensions::SetTextureUVParameterName(TWeakObjectPtr<UDMT
 		return;
 	}
 
-	if (InName == NAME_None)
+	if (InName.IsNone())
 	{
 		return;
 	}
diff --git a/Engine/Plugins/Experimental/DynamicMaterial/Source/DynamicMaterialEditor/Private/DetailsPanel/Slate/SDMDetailsPanelMaterialInterfaceWidget.cpp b/Engine/Plugins/Experimental/DynamicMaterial/Source/DynamicMaterialEditor/Private/DetailsPanel/Slate/SDMDetailsPanelMaterialInterfaceWidget.cpp
--- a/Engine/Plugins/Experimental/DynamicMaterial/Source/DynamicMaterialEditor/Private/DetailsPanel/Slate/SDMDetailsPanelMaterialInterfaceWidget.cpp
+++ b/Engine/Plugins/Experimental/DynamicMaterial/Source/DynamicMaterialEditor/Private/DetailsPanel/Slate/SDMDetailsPanelMaterialInterfaceWidget.cpp
@@ -7,6 +7,7 @@
 #include "IDynamicMaterialEditorModule.h"
 #include "Material/DynamicMaterialInstance.h"
 #include "Material/DynamicMaterialInstanceFactory.h"
+#include "Model/DynamicMaterialModel.h"
 #include "PropertyCustomizationHelpers.h"
 #include "PropertyHandle.h"
 #include "Widgets/Input/SButton.h"
diff --git a/Engine/Plugins/Experimental/DynamicMaterial/Source/DynamicMaterialEditor/Private/DetailsPanel/Slate/SDMDetailsPanelTabSpawner.cpp b/Engine/Plugins/Experimental/DynamicMaterial/Source/DynamicMaterialEditor/Private/DetailsPanel/Slate/SDMDetailsPanelTabSpawner.cpp
--- a/Engine/Plugins/Experimental/DynamicMaterial/Source/DynamicMaterialEditor/Private/DetailsPanel/Slate/SDMDetailsPanelTabSpawner.cpp
+++ b/Engine/Plugins/Experimental/DynamicMaterial/Source/DynamicMaterialEditor/Private/DetailsPanel/Slate/SDMDetailsPanelTabSpawner.cpp
@@ -6,6 +6,7 @@
 #include "DynamicMaterialEditorStyle.h"
 #include "IAssetTools.h"
 #include "Model/DynamicMaterialModel.h"
+#include "Model/DynamicMaterialModelBase.h"
 #include "Model/DynamicMaterialModelFactory.h"
 #include "PropertyCustomizationHelpers.h"
 #include "PropertyHandle.h"
@@ -29,8 +30,6 @@ void SDMDetailsPanelTabSpawner::Construct(const FArguments& InArgs, const TShare
 	UObject* Value = nullptr;
 	InPropertyHandle->GetValue(Value);
 
-	UDynamicMaterialModel* MaterialModel = GetMaterialModel();
-
 	// @formatter:off
 	ChildSlot
 	[
@@ -43,7 +42,7 @@ void SDMDetailsPanelTabSpawner::Construct(const FArguments& InArgs, const TShare
 		[
 			SNew(SObjectPropertyEntryBox)
 			.AllowClear(true)
-			.AllowedClass(UDynamicMaterialModel::StaticClass())
+			.AllowedClass(UDynamicMaterialModelBase::StaticClass())
 			.DisplayBrowse(true)
 			.DisplayThumbnail(false)
 			.DisplayCompactSize(true)
@@ -71,7 +70,7 @@ void SDMDetailsPanelTabSpawner::Construct(const FArguments& InArgs, const TShare
 	// @formatter:on
 }
 
-UDynamicMaterialModel* SDMDetailsPanelTabSpawner::GetMaterialModel() const
+UDynamicMaterialModelBase* SDMDetailsPanelTabSpawner::GetMaterialModelBase() const
 {
 	TArray<UObject*> OuterObjects;
 	PropertyHandle->GetOuterObjects(OuterObjects);
@@ -84,10 +83,10 @@ UDynamicMaterialModel* SDMDetailsPanelTabSpawner::GetMaterialModel() const
 	UObject* Value = nullptr;
 	PropertyHandle->GetValue(Value);
 
-	return Cast<UDynamicMaterialModel>(Value);
+	return Cast<UDynamicMaterialModelBase>(Value);
 }
 
-void SDMDetailsPanelTabSpawner::SetMaterialModel(UDynamicMaterialModel* InNewModel)
+void SDMDetailsPanelTabSpawner::SetMaterialModelBase(UDynamicMaterialModelBase* InNewModel)
 {
 	TArray<UObject*> OuterObjects;
 	PropertyHandle->GetOuterObjects(OuterObjects);
@@ -102,7 +101,7 @@ void SDMDetailsPanelTabSpawner::SetMaterialModel(UDynamicMaterialModel* InNewMod
 
 FText SDMDetailsPanelTabSpawner::GetButtonText() const
 {
-	if (GetMaterialModel())
+	if (GetMaterialModelBase())
 	{
 		return LOCTEXT("OpenMaterialDesignerModel", "Edit with Material Designer");
 	}
@@ -112,7 +111,7 @@ FText SDMDetailsPanelTabSpawner::GetButtonText() const
 
 FReply SDMDetailsPanelTabSpawner::OnButtonClicked()
 {
-	if (GetMaterialModel())
+	if (GetMaterialModelBase())
 	{
 		return OpenDynamicMaterialModelTab();
 	}
@@ -122,10 +121,10 @@ FReply SDMDetailsPanelTabSpawner::OnButtonClicked()
 
 FReply SDMDetailsPanelTabSpawner::CreateDynamicMaterialModel()
 {
-	UDynamicMaterialModel* MaterialModel = GetMaterialModel();
+	UDynamicMaterialModelBase* MaterialModelBase = GetMaterialModelBase();
 
 	// We already have a builder, so we don't need to create one
-	if (MaterialModel)
+	if (MaterialModelBase)
 	{
 		return FReply::Unhandled();
 	}
@@ -154,7 +153,7 @@ FReply SDMDetailsPanelTabSpawner::CreateDynamicMaterialModel()
 	check(DynamicMaterialModelFactory);
 
 	UDynamicMaterialModel* NewModel = Cast<UDynamicMaterialModel>(DynamicMaterialModelFactory->FactoryCreateNew(
-		UDynamicMaterialModel::StaticClass(),
+		UDynamicMaterialModelBase::StaticClass(),
 		Package,
 		*AssetName,
 		RF_Standalone | RF_Public,
@@ -171,46 +170,45 @@ FReply SDMDetailsPanelTabSpawner::CreateDynamicMaterialModel()
 
 FReply SDMDetailsPanelTabSpawner::ClearDynamicMaterialModel()
 {
-	UDynamicMaterialModel* MaterialModel = GetMaterialModel();
+	UDynamicMaterialModelBase* MaterialModelBase = GetMaterialModelBase();
 
 	// We don't have a builder, so we don't need to clear it
-	if (!MaterialModel)
+	if (!MaterialModelBase)
 	{
 		return FReply::Unhandled();
 	}
 
-	MaterialModel->ResetData();
-	SetMaterialModel(nullptr);
+	SetMaterialModelBase(nullptr);
 
 	return FReply::Handled();
 }
 
 FReply SDMDetailsPanelTabSpawner::OpenDynamicMaterialModelTab()
 {
-	UDynamicMaterialModel* MaterialModel = GetMaterialModel();
+	UDynamicMaterialModelBase* MaterialModelBase = GetMaterialModelBase();
 
 	// We don't have a builder, so we can't open it
-	if (!MaterialModel)
+	if (!MaterialModelBase)
 	{
 		return FReply::Unhandled();
 	}
 
 	IAssetTools& AssetTools = FModuleManager::LoadModuleChecked<FAssetToolsModule>("AssetTools").Get();
-	AssetTools.OpenEditorForAssets({MaterialModel});
+	AssetTools.OpenEditorForAssets({MaterialModelBase});
 
 	return FReply::Handled();
 }
 
 FString SDMDetailsPanelTabSpawner::GetEditorPath() const
 {
-	UDynamicMaterialModel* MaterialModel = GetMaterialModel();
+	UDynamicMaterialModelBase* MaterialModelBase = GetMaterialModelBase();
 
-	return MaterialModel ? MaterialModel->GetPathName() : "";
+	return MaterialModelBase ? MaterialModelBase->GetPathName() : "";
 }
 
 void SDMDetailsPanelTabSpawner::OnEditorChanged(const FAssetData& InAssetData)
 {
-	SetMaterialModel(Cast<UDynamicMaterialModel>(InAssetData.GetAsset()));
+	SetMaterialModelBase(Cast<UDynamicMaterialModelBase>(InAssetData.GetAsset()));
 }
 
 #undef LOCTEXT_NAMESPACE
diff --git a/Engine/Plugins/Experimental/DynamicMaterial/Source/DynamicMaterialEditor/Private/DetailsPanel/Slate/SDMDetailsPanelTabSpawner.h b/Engine/Plugins/Experimental/DynamicMaterial/Source/DynamicMaterialEditor/Private/DetailsPanel/Slate/SDMDetailsPanelTabSpawner.h
--- a/Engine/Plugins/Experimental/DynamicMaterial/Source/DynamicMaterialEditor/Private/DetailsPanel/Slate/SDMDetailsPanelTabSpawner.h
+++ b/Engine/Plugins/Experimental/DynamicMaterial/Source/DynamicMaterialEditor/Private/DetailsPanel/Slate/SDMDetailsPanelTabSpawner.h
@@ -6,7 +6,7 @@
 #include "Templates/SharedPointer.h"
 
 class IPropertyHandle;
-class UDynamicMaterialModel;
+class UDynamicMaterialModelBase;
 
 class SDMDetailsPanelTabSpawner : public SCompoundWidget
 {
@@ -19,8 +19,8 @@ public:
 protected:
 	TSharedPtr<IPropertyHandle> PropertyHandle;
 
-	UDynamicMaterialModel* GetMaterialModel() const;
-	void SetMaterialModel(UDynamicMaterialModel* InNewModel);
+	UDynamicMaterialModelBase* GetMaterialModelBase() const;
+	void SetMaterialModelBase(UDynamicMaterialModelBase* InNewModel);
 
 	FReply OnButtonClicked();
 	FReply CreateDynamicMaterialModel();
diff --git a/Engine/Plugins/Experimental/DynamicMaterial/Source/DynamicMaterialEditor/Private/DragDrop/DMStageDragDropOperation.cpp b/Engine/Plugins/Experimental/DynamicMaterial/Source/DynamicMaterialEditor/Private/DragDrop/DMStageDragDropOperation.cpp
--- a/Engine/Plugins/Experimental/DynamicMaterial/Source/DynamicMaterialEditor/Private/DragDrop/DMStageDragDropOperation.cpp
+++ b/Engine/Plugins/Experimental/DynamicMaterial/Source/DynamicMaterialEditor/Private/DragDrop/DMStageDragDropOperation.cpp
@@ -20,6 +20,9 @@ TSharedPtr<SWidget> FDMStageDragDropOperation::GetDefaultDecorator() const
 	TSharedPtr<SDMStage> StageWidget = StageWidgetWeak.Pin();
 	check(StageWidget.IsValid());
 
+	TSharedPtr<SDMSlot> SlotWidget = StageWidget->GetSlotWidget();
+	check(SlotWidget.IsValid());
+
 	UDMMaterialStage* const Stage = StageWidget->GetStage();
 	check(Stage);
 
@@ -29,7 +32,7 @@ TSharedPtr<SWidget> FDMStageDragDropOperation::GetDefaultDecorator() const
 		.HAlign(HAlign_Fill)
 		.VAlign(VAlign_Fill)
 		[
-			SNew(SDMStage, Stage)
+			SNew(SDMStage, SlotWidget.ToSharedRef(), Stage)
 			.DesiredSize(StageWidget->GetPreviewSize())
 		]
 		+ SOverlay::Slot()
diff --git a/Engine/Plugins/Experimental/DynamicMaterial/Source/DynamicMaterialEditor/Private/DynamicMaterialEditorModule.cpp b/Engine/Plugins/Experimental/DynamicMaterial/Source/DynamicMaterialEditor/Private/DynamicMaterialEditorModule.cpp
--- a/Engine/Plugins/Experimental/DynamicMaterial/Source/DynamicMaterialEditor/Private/DynamicMaterialEditorModule.cpp
+++ b/Engine/Plugins/Experimental/DynamicMaterial/Source/DynamicMaterialEditor/Private/DynamicMaterialEditorModule.cpp
@@ -6,7 +6,9 @@
 #include "Components/DMMaterialComponent.h"
 #include "Components/DMMaterialEffectFunction.h"
 #include "Components/DMMaterialStageThroughput.h"
+#include "Components/DMMaterialValueDynamic.h"
 #include "Components/DMTextureUV.h"
+#include "Components/DMTextureUVDynamic.h"
 #include "Components/MaterialStageInputs/DMMSIFunction.h"
 #include "Components/MaterialStageInputs/DMMSIThroughput.h"
 #include "Components/MaterialValues/DMMaterialValueBool.h"
@@ -51,12 +53,13 @@
 #include "Slate/Properties/Generators/DMInputThroughputPropertyRowGenerator.h"
 #include "Slate/Properties/Generators/DMMaterialEffectFunctionPropertyRowGenerator.h"
 #include "Slate/Properties/Generators/DMMaterialStageFunctionPropertyRowGenerator.h"
+#include "Slate/Properties/Generators/DMMaterialValueDynamicPropertyRowGenerator.h"
 #include "Slate/Properties/Generators/DMMaterialValuePropertyRowGenerator.h"
 #include "Slate/Properties/Generators/DMStagePropertyRowGenerator.h"
 #include "Slate/Properties/Generators/DMTextureUVPropertyRowGenerator.h"
+#include "Slate/Properties/Generators/DMTextureUVDynamicPropertyRowGenerator.h"
 #include "Slate/Properties/Generators/DMThroughputPropertyRowGenerator.h"
 #include "Slate/SDMEditor.h"
-#include "Utils/DMMaterialFunctionLibrary.h"
 
 DEFINE_LOG_CATEGORY(LogDynamicMaterialEditor);
 
@@ -291,7 +294,9 @@ void FDynamicMaterialEditorModule::StartupModule()
 	RegisterComponentPropertyRowGeneratorDelegate<UDMMaterialComponent,            FDMComponentPropertyRowGenerator>();
 	RegisterComponentPropertyRowGeneratorDelegate<UDMMaterialStage,                FDMStagePropertyRowGenerator>();
 	RegisterComponentPropertyRowGeneratorDelegate<UDMMaterialValue,                FDMMaterialValuePropertyRowGenerator>();
+	RegisterComponentPropertyRowGeneratorDelegate<UDMMaterialValueDynamic,         FDMMaterialValueDynamicPropertyRowGenerator>();
 	RegisterComponentPropertyRowGeneratorDelegate<UDMTextureUV,                    FDMTextureUVPropertyRowGenerator>();
+	RegisterComponentPropertyRowGeneratorDelegate<UDMTextureUVDynamic,             FDMTextureUVDynamicPropertyRowGenerator>();
 	RegisterComponentPropertyRowGeneratorDelegate<UDMMaterialStageThroughput,      FDMThroughputPropertyRowGenerator>();
 	RegisterComponentPropertyRowGeneratorDelegate<UDMMaterialStageInputThroughput, FDMInputThroughputPropertyRowGenerator>();
 	RegisterComponentPropertyRowGeneratorDelegate<UDMMaterialEffectFunction,       FDMMaterialEffectFunctionPropertyRowGenerator>();
@@ -345,7 +350,7 @@ void FDynamicMaterialEditorModule::ShutdownModule()
 	FDMValueDetailsRowExtensions::Get().UnregisterRowExtensions();
 }
 
-void FDynamicMaterialEditorModule::OpenMaterialModel(UDynamicMaterialModel* InMaterialModel, UWorld* InWorld, bool bInInvokeTab) const
+void FDynamicMaterialEditorModule::OpenMaterialModel(UDynamicMaterialModelBase* InMaterialModel, UWorld* InWorld, bool bInInvokeTab) const
 {
 	using namespace UE::DynamicMaterialEditor::Private;
 
@@ -356,7 +361,7 @@ void FDynamicMaterialEditorModule::OpenMaterialModel(UDynamicMaterialModel* InMa
 			FDMLevelEditorIntegration::InvokeTabForWorld(InWorld);
 		}
 
-		Editor->SetMaterialModel(InMaterialModel);
+		Editor->SetMaterialModelBase(InMaterialModel);
 	}
 	else if (IsValid(InWorld))
 	{
@@ -444,9 +449,10 @@ void FDynamicMaterialEditorModule::ClearDynamicMaterialModel(UWorld* InWorld) co
 	OpenMaterialModel(nullptr, InWorld, /* Invoke tab */ false);
 }
 
-TSharedRef<SWidget> FDynamicMaterialEditorModule::CreateEditor(UDynamicMaterialModel* InMaterialModel, UWorld* InAssetEditorWorld)
+TSharedRef<SWidget> FDynamicMaterialEditorModule::CreateEditor(UDynamicMaterialModelBase* InMaterialModelBase, UWorld* InAssetEditorWorld)
 {
-	TSharedRef<SDMEditor> NewEditor = SNew(SDMEditor, InMaterialModel);
+	TSharedRef<SDMEditor> NewEditor = SNew(SDMEditor);
+	NewEditor->SetMaterialModelBase(InMaterialModelBase);
 
 	if (IsValid(InAssetEditorWorld))
 	{
@@ -454,8 +460,8 @@ TSharedRef<SWidget> FDynamicMaterialEditorModule::CreateEditor(UDynamicMaterialM
 
 		if (IsValid(WorldSubsystem))
 		{
-			WorldSubsystem->GetGetCustomEditorModelDelegate().BindSP(NewEditor, &SDMEditor::GetMaterialModel);
-			WorldSubsystem->GetSetCustomEditorModelDelegate().BindSP(NewEditor, &SDMEditor::SetMaterialModel);
+			WorldSubsystem->GetGetCustomEditorModelDelegate().BindSP(NewEditor, &SDMEditor::GetMaterialModelBase);
+			WorldSubsystem->GetSetCustomEditorModelDelegate().BindSP(NewEditor, &SDMEditor::SetMaterialModelBase);
 			WorldSubsystem->GetCustomObjectPropertyEditorDelegate().BindSP(NewEditor, &SDMEditor::SetMaterialObjectProperty);
 			WorldSubsystem->GetSetCustomEditorActorDelegate().BindSP(NewEditor, &SDMEditor::OnActorSelected);
 		}
@@ -556,11 +562,11 @@ void FDynamicMaterialEditorModule::OpenEditor(UWorld* InWorld) const
 	}
 }
 
-UDynamicMaterialModel* FDynamicMaterialEditorModule::GetOpenedMaterialModel(UWorld* InWorld) const
+UDynamicMaterialModelBase* FDynamicMaterialEditorModule::GetOpenedMaterialModel(UWorld* InWorld) const
 {
 	if (TSharedPtr<SDMEditor> Editor = FDMLevelEditorIntegration::GetEditorForWorld(InWorld))
 	{
-		return Editor->GetMaterialModel();
+		return Editor->GetMaterialModelBase();
 	}
 
 	if (IsValid(InWorld))
diff --git a/Engine/Plugins/Experimental/DynamicMaterial/Source/DynamicMaterialEditor/Private/DynamicMaterialEditorModule.h b/Engine/Plugins/Experimental/DynamicMaterial/Source/DynamicMaterialEditor/Private/DynamicMaterialEditorModule.h
--- a/Engine/Plugins/Experimental/DynamicMaterial/Source/DynamicMaterialEditor/Private/DynamicMaterialEditorModule.h
+++ b/Engine/Plugins/Experimental/DynamicMaterial/Source/DynamicMaterialEditor/Private/DynamicMaterialEditorModule.h
@@ -21,7 +21,7 @@ class SWidget;
 class UDMMaterialComponent;
 class UDMMaterialStageSource;
 class UDynamicMaterialInstance;
-class UDynamicMaterialModel;
+class UDynamicMaterialModelBase;
 struct FDMBuildRequestList;
 
 DECLARE_LOG_CATEGORY_EXTERN(LogDynamicMaterialEditor, Log, All);
@@ -65,7 +65,7 @@ class FDynamicMaterialEditorModule : public IDynamicMaterialEditorModule, public
 {
 public:
 	static const FName TabId;
-	static FDMOnUIValueUpdate& GetOnUIValueUpdate() { return OnUIValueUpdate; }
+	static FDMOnUIValueUpdate::RegistrationType& GetOnUIValueUpdate() { return OnUIValueUpdate; }
 
 	static FDynamicMaterialEditorModule& Get();
 
@@ -83,7 +83,7 @@ public:
 	static FDMGetObjectMaterialPropertiesDelegate GetCustomMaterialPropertyGenerator(UClass* InClass);
 
 	/** With a provided world, the editor will bind to the MD world subsystem to receive model changes. */
-	static TSharedRef<SWidget> CreateEditor(UDynamicMaterialModel* InMaterialModel, UWorld* InAssetEditorWorld);
+	static TSharedRef<SWidget> CreateEditor(UDynamicMaterialModelBase* InMaterialModelBase, UWorld* InAssetEditorWorld);
 	static TSharedRef<SWidget> CreateEmptyTabContent();
 
 	FDynamicMaterialEditorModule();
@@ -93,8 +93,8 @@ public:
 	virtual void RegisterMaterialModelCreatedCallback(const TSharedRef<IDMOnWizardCompleteCallback> InCallback)  override;
 	virtual void UnregisterMaterialModelCreatedCallback(const TSharedRef<IDMOnWizardCompleteCallback> InCallback) override;
 	virtual void OpenEditor(UWorld* InWorld) const override;
-	virtual UDynamicMaterialModel* GetOpenedMaterialModel(UWorld* InWorld) const override;
-	virtual void OpenMaterialModel(UDynamicMaterialModel* InMaterialModel, UWorld* InWorld, bool bInInvokeTab) const override;
+	virtual UDynamicMaterialModelBase* GetOpenedMaterialModel(UWorld* InWorld) const override;
+	virtual void OpenMaterialModel(UDynamicMaterialModelBase* InMaterialModel, UWorld* InWorld, bool bInInvokeTab) const override;
 	virtual void OpenMaterialObjectProperty(const FDMObjectMaterialProperty& InObjectProperty, UWorld* InWorld, bool bInInvokeTab) const override;
 	virtual void OpenMaterialInstance(UDynamicMaterialInstance* InInstance, UWorld* InWorld, bool bInInvokeTab) const override;
 	virtual void OnActorSelected(AActor* InActor, UWorld* InWorld, bool bInInvokeTab) const override;
diff --git a/Engine/Plugins/Experimental/DynamicMaterial/Source/DynamicMaterialEditor/Private/LevelEditor/DMLevelEditorIntegration.cpp b/Engine/Plugins/Experimental/DynamicMaterial/Source/DynamicMaterialEditor/Private/LevelEditor/DMLevelEditorIntegration.cpp
--- a/Engine/Plugins/Experimental/DynamicMaterial/Source/DynamicMaterialEditor/Private/LevelEditor/DMLevelEditorIntegration.cpp
+++ b/Engine/Plugins/Experimental/DynamicMaterial/Source/DynamicMaterialEditor/Private/LevelEditor/DMLevelEditorIntegration.cpp
@@ -49,9 +49,9 @@ namespace UE::DynamicMaterialEditor::Private
 			return;
 		}
 
-		UDynamicMaterialModel* MaterialModel = Editor->GetMaterialModel();
+		UDynamicMaterialModelBase* MaterialModelBase = Editor->GetMaterialModelBase();
 
-		if (!MaterialModel)
+		if (!MaterialModelBase)
 		{
 			return;
 		}
@@ -60,7 +60,7 @@ namespace UE::DynamicMaterialEditor::Private
 
 		const FString WorldPath = InWorld->GetPathName();
 		const int WorldPathLength = WorldPath.Len();
-		const FString ModelPath = MaterialModel->GetPathName();
+		const FString ModelPath = MaterialModelBase->GetPathName();
 
 		if (ModelPath.Len() > WorldPathLength && ModelPath.StartsWith(WorldPath))
 		{
@@ -105,7 +105,7 @@ namespace UE::DynamicMaterialEditor::Private
 		{
 			if (UDynamicMaterialModel* MaterialModel = Cast<UDynamicMaterialModel>(Object))
 			{
-				Editor->SetMaterialModel(MaterialModel);
+				Editor->SetMaterialModelBase(MaterialModel);
 
 				if (UDynamicMaterialModelEditorOnlyData* EditorOnlyData = UDynamicMaterialModelEditorOnlyData::Get(MaterialModel))
 				{
diff --git a/Engine/Plugins/Experimental/DynamicMaterial/Source/DynamicMaterialEditor/Private/LevelEditor/DMLevelEditorIntegrationInstance.h b/Engine/Plugins/Experimental/DynamicMaterial/Source/DynamicMaterialEditor/Private/LevelEditor/DMLevelEditorIntegrationInstance.h
--- a/Engine/Plugins/Experimental/DynamicMaterial/Source/DynamicMaterialEditor/Private/LevelEditor/DMLevelEditorIntegrationInstance.h
+++ b/Engine/Plugins/Experimental/DynamicMaterial/Source/DynamicMaterialEditor/Private/LevelEditor/DMLevelEditorIntegrationInstance.h
@@ -12,7 +12,6 @@ class AActor;
 class ILevelEditor;
 class SDMEditor;
 class SDockTab;
-class UDynamicMaterialModel;
 class UTypedElementSelectionSet;
 class UWorld;
 
diff --git a/Engine/Plugins/Experimental/DynamicMaterial/Source/DynamicMaterialEditor/Private/Material/AssetDefinition_DynamicMaterialInstance.cpp b/Engine/Plugins/Experimental/DynamicMaterial/Source/DynamicMaterialEditor/Private/Material/AssetDefinition_DynamicMaterialInstance.cpp
--- a/Engine/Plugins/Experimental/DynamicMaterial/Source/DynamicMaterialEditor/Private/Material/AssetDefinition_DynamicMaterialInstance.cpp
+++ b/Engine/Plugins/Experimental/DynamicMaterial/Source/DynamicMaterialEditor/Private/Material/AssetDefinition_DynamicMaterialInstance.cpp
@@ -47,14 +47,14 @@ EAssetCommandResult UAssetDefinition_DynamicMaterialInstance::OpenAssets(const F
 			continue;
 		}
 
-		UDynamicMaterialModel* MaterialModel = Instance->GetMaterialModel();
+		UDynamicMaterialModelBase* MaterialModelBase = Instance->GetMaterialModelBase();
 
-		if (!MaterialModel)
+		if (!MaterialModelBase)
 		{
 			continue;
 		}
 
-		MaterialModels.Add(MaterialModel);
+		MaterialModels.Add(MaterialModelBase);
 	}
 
 	if (MaterialModels.IsEmpty())
diff --git a/Engine/Plugins/Experimental/DynamicMaterial/Source/DynamicMaterialEditor/Private/Material/DynamicMaterialInstanceFactory.cpp b/Engine/Plugins/Experimental/DynamicMaterial/Source/DynamicMaterialEditor/Private/Material/DynamicMaterialInstanceFactory.cpp
--- a/Engine/Plugins/Experimental/DynamicMaterial/Source/DynamicMaterialEditor/Private/Material/DynamicMaterialInstanceFactory.cpp
+++ b/Engine/Plugins/Experimental/DynamicMaterial/Source/DynamicMaterialEditor/Private/Material/DynamicMaterialInstanceFactory.cpp
@@ -35,19 +35,27 @@ UObject* UDynamicMaterialInstanceFactory::FactoryCreateNew(UClass* Class, UObjec
 	UDynamicMaterialModelFactory* EditorFactory = NewObject<UDynamicMaterialModelFactory>();
 	check(EditorFactory);
 
-	UDynamicMaterialModel* NewModel = Cast<UDynamicMaterialModel>(EditorFactory->FactoryCreateNew(
-		UDynamicMaterialModel::StaticClass(), NewInstance, NAME_None, RF_Transactional, nullptr, GWarn));
-	check(NewModel);
+	UDynamicMaterialModelBase* ModelBase = Cast<UDynamicMaterialModelBase>(Context);
+
+	if (!ModelBase)
+	{
+		ModelBase = Cast<UDynamicMaterialModel>(EditorFactory->FactoryCreateNew(
+			UDynamicMaterialModel::StaticClass(), NewInstance, NAME_None, RF_Transactional | RF_Public, nullptr, GWarn));
+		check(ModelBase);
+	}
 
 	const FDMInitializationGuard InitGuard;
 
-	NewInstance->SetMaterialModel(NewModel);
+	NewInstance->SetMaterialModel(ModelBase);
 
-	NewModel->SetDynamicMaterialInstance(NewInstance);
+	ModelBase->SetDynamicMaterialInstance(NewInstance);
 
-	if (IDynamicMaterialModelEditorOnlyDataInterface* ModelEditorOnlyData = NewModel->GetEditorOnlyData())
+	if (UDynamicMaterialModel* MaterialModel = Cast<UDynamicMaterialModel>(ModelBase))
 	{
-		ModelEditorOnlyData->RequestMaterialBuild();
+		if (IDynamicMaterialModelEditorOnlyDataInterface* ModelEditorOnlyData = MaterialModel->GetEditorOnlyData())
+		{
+			ModelEditorOnlyData->RequestMaterialBuild();
+		}
 	}
 
 	NewInstance->InitializeMIDPublic();
diff --git a/Engine/Plugins/Experimental/DynamicMaterial/Source/DynamicMaterialEditor/Private/Menus/DMMaterialSlotLayerAddEffectContext.h b/Engine/Plugins/Experimental/DynamicMaterial/Source/DynamicMaterialEditor/Private/Menus/DMMaterialSlotLayerAddEffectContext.h
--- a/Engine/Plugins/Experimental/DynamicMaterial/Source/DynamicMaterialEditor/Private/Menus/DMMaterialSlotLayerAddEffectContext.h
+++ b/Engine/Plugins/Experimental/DynamicMaterial/Source/DynamicMaterialEditor/Private/Menus/DMMaterialSlotLayerAddEffectContext.h
@@ -18,8 +18,6 @@ class UDMMaterialSlotLayerAddEffectContext : public UObject
 	GENERATED_BODY()
 
 public:
-	UDMMaterialSlotLayerAddEffectContext() = default;
-
 	void SetEditorWidget(const TSharedPtr<SDMEditor>& InEditor) { EditorWidgetWeak = InEditor; }
 
 	TSharedPtr<SDMEditor> GetEditorWidget() const { return EditorWidgetWeak.Pin(); }
diff --git a/Engine/Plugins/Experimental/DynamicMaterial/Source/DynamicMaterialEditor/Private/Menus/DMMaterialStageSourceMenus.cpp b/Engine/Plugins/Experimental/DynamicMaterial/Source/DynamicMaterialEditor/Private/Menus/DMMaterialStageSourceMenus.cpp
--- a/Engine/Plugins/Experimental/DynamicMaterial/Source/DynamicMaterialEditor/Private/Menus/DMMaterialStageSourceMenus.cpp
+++ b/Engine/Plugins/Experimental/DynamicMaterial/Source/DynamicMaterialEditor/Private/Menus/DMMaterialStageSourceMenus.cpp
@@ -372,7 +372,7 @@ namespace UE::DynamicMaterialEditor::Private
 			return;
 		}
 
-		UDynamicMaterialModel* MaterialModel = MenuContext->GetModel();
+		UDynamicMaterialModelBase* MaterialModel = MenuContext->GetModelBase();
 		if (!MaterialModel)
 		{
 			return;
@@ -499,14 +499,14 @@ namespace UE::DynamicMaterialEditor::Private
 			return;
 		}
 
-		UDynamicMaterialModel* MaterialModel = MenuContext->GetModel();
-		if (!MaterialModel)
+		UDynamicMaterialModelBase* MaterialModelBase = MenuContext->GetModelBase();
+		if (!MaterialModelBase)
 		{
 			return;
 		}
 
-		UDynamicMaterialModelEditorOnlyData* ModelEditorOnlyData = UDynamicMaterialModelEditorOnlyData::Get(MaterialModel);
-		if (!MaterialModel)
+		UDynamicMaterialModelEditorOnlyData* ModelEditorOnlyData = UDynamicMaterialModelEditorOnlyData::Get(MaterialModelBase);
+		if (!MaterialModelBase)
 		{
 			return;
 		}
diff --git a/Engine/Plugins/Experimental/DynamicMaterial/Source/DynamicMaterialEditor/Private/Menus/DMMenuContext.cpp b/Engine/Plugins/Experimental/DynamicMaterial/Source/DynamicMaterialEditor/Private/Menus/DMMenuContext.cpp
--- a/Engine/Plugins/Experimental/DynamicMaterial/Source/DynamicMaterialEditor/Private/Menus/DMMenuContext.cpp
+++ b/Engine/Plugins/Experimental/DynamicMaterial/Source/DynamicMaterialEditor/Private/Menus/DMMenuContext.cpp
@@ -110,6 +110,16 @@ UDMMaterialSlot* UDMMenuContext::GetSlot() const
 	return nullptr;;
 }
 
+UDynamicMaterialModelBase* UDMMenuContext::GetModelBase() const
+{
+	if (TSharedPtr<SDMEditor> EditorWidget = EditorWidgetWeak.Pin())
+	{
+		return EditorWidget->GetMaterialModelBase();
+	}
+
+	return nullptr;
+}
+
 UDynamicMaterialModel* UDMMenuContext::GetModel() const
 {
 	if (TSharedPtr<SDMEditor> EditorWidget = EditorWidgetWeak.Pin())
diff --git a/Engine/Plugins/Experimental/DynamicMaterial/Source/DynamicMaterialEditor/Private/Menus/DMMenuContext.h b/Engine/Plugins/Experimental/DynamicMaterial/Source/DynamicMaterialEditor/Private/Menus/DMMenuContext.h
--- a/Engine/Plugins/Experimental/DynamicMaterial/Source/DynamicMaterialEditor/Private/Menus/DMMenuContext.h
+++ b/Engine/Plugins/Experimental/DynamicMaterial/Source/DynamicMaterialEditor/Private/Menus/DMMenuContext.h
@@ -13,6 +13,7 @@ class UDMMaterialStage;
 class UDMMaterialStageBlend;
 class UDMMaterialStageSource;
 class UDynamicMaterialModel;
+class UDynamicMaterialModelBase;
 class UToolMenu;
 
 UCLASS()
@@ -36,6 +37,8 @@ public:
 
 	UDMMaterialSlot* GetSlot() const;
 
+	UDynamicMaterialModelBase* GetModelBase() const;
+
 	UDynamicMaterialModel* GetModel() const;
 
 	UDMMaterialStage* GetStage() const;
diff --git a/Engine/Plugins/Experimental/DynamicMaterial/Source/DynamicMaterialEditor/Private/Menus/DMToolBarMenus.cpp b/Engine/Plugins/Experimental/DynamicMaterial/Source/DynamicMaterialEditor/Private/Menus/DMToolBarMenus.cpp
--- a/Engine/Plugins/Experimental/DynamicMaterial/Source/DynamicMaterialEditor/Private/Menus/DMToolBarMenus.cpp
+++ b/Engine/Plugins/Experimental/DynamicMaterial/Source/DynamicMaterialEditor/Private/Menus/DMToolBarMenus.cpp
@@ -1,6 +1,8 @@
 // Copyright Epic Games, Inc. All Rights Reserved.
 
 #include "DMToolBarMenus.h"
+
+#include "AssetToolsModule.h"
 #include "ContentBrowserModule.h"
 #include "DesktopPlatformModule.h"
 #include "DynamicMaterialEditorModule.h"
@@ -55,18 +57,26 @@ namespace UE::DynamicMaterialEditor::Private
 	{
 		if (UDynamicMaterialInstance* MaterialInstance = InMaterialInstanceWeak.Get())
 		{
+			IAssetTools& AssetTools = FModuleManager::LoadModuleChecked<FAssetToolsModule>("AssetTools").Get();
+			FString PackageName, AssetName;
+			AssetTools.CreateUniqueAssetName(MaterialInstance->GetName(), TEXT(""), PackageName, AssetName);
+
+			IContentBrowserSingleton& ContentBrowser = FModuleManager::LoadModuleChecked<FContentBrowserModule>("ContentBrowser").Get();
+			const FContentBrowserItemPath CurrentPath = ContentBrowser.GetCurrentPath();
+			const FString PathStr = CurrentPath.HasInternalPath() ? CurrentPath.GetInternalPathString() : "/Game";
+
 			FSaveAssetDialogConfig SaveAssetDialogConfig;
 			SaveAssetDialogConfig.DialogTitleOverride = LOCTEXT("SaveAssetDialogTitle", "Save Asset As");
-			SaveAssetDialogConfig.DefaultPath = "/Game";
-			SaveAssetDialogConfig.DefaultAssetName = MaterialInstance->GetName();
-			SaveAssetDialogConfig.ExistingAssetPolicy = ESaveAssetDialogExistingAssetPolicy::AllowButWarn;
+			SaveAssetDialogConfig.DefaultPath = PathStr;
+			SaveAssetDialogConfig.ExistingAssetPolicy = ESaveAssetDialogExistingAssetPolicy::Disallow;
+			SaveAssetDialogConfig.DefaultAssetName = AssetName;
 
 			FContentBrowserModule& ContentBrowserModule = FModuleManager::LoadModuleChecked<FContentBrowserModule>("ContentBrowser");
 			FString SaveObjectPath = ContentBrowserModule.Get().CreateModalSaveAssetDialog(SaveAssetDialogConfig);
 
 			if (!SaveObjectPath.IsEmpty())
 			{
-				UDMBlueprintFunctionLibrary::ExportMaterialInstance(MaterialInstance->GetMaterialModel(), SaveObjectPath);
+				UDMBlueprintFunctionLibrary::ExportMaterialInstance(MaterialInstance->GetMaterialModelBase(), SaveObjectPath);
 
 				if (FEngineAnalytics::IsAvailable())
 				{
@@ -76,53 +86,61 @@ namespace UE::DynamicMaterialEditor::Private
 		}
 	}
 
-	void ExportMaterialModelFromModel(TWeakObjectPtr<UDynamicMaterialModel> InMaterialModelWeak)
+	void ExportMaterialModelFromModel(TWeakObjectPtr<UDynamicMaterialModelBase> InMaterialModelBaseWeak)
 	{
-		if (UDynamicMaterialModel* MaterialModel = InMaterialModelWeak.Get())
+		UDynamicMaterialModelBase* MaterialModelBase = InMaterialModelBaseWeak.Get();
+
+		if (!MaterialModelBase)
 		{
-			UMaterial* GeneratedMaterial = MaterialModel->GetGeneratedMaterial();
+			return;
+		}
 
-			if (!GeneratedMaterial)
-			{
-				UE::DynamicMaterialEditor::Private::LogError(TEXT("Failed to find a generated material to export."), true, MaterialModel);
-				return;
-			}
+		UMaterial* GeneratedMaterial = MaterialModelBase->GetGeneratedMaterial();
 
-			FSaveAssetDialogConfig SaveAssetDialogConfig;
-			SaveAssetDialogConfig.DialogTitleOverride = LOCTEXT("SaveAssetDialogTitle", "Save Asset As");
-			SaveAssetDialogConfig.DefaultPath = "/Game";
-			SaveAssetDialogConfig.DefaultAssetName = GeneratedMaterial->GetName();
-			SaveAssetDialogConfig.ExistingAssetPolicy = ESaveAssetDialogExistingAssetPolicy::AllowButWarn;
+		if (!GeneratedMaterial)
+		{
+			UE::DynamicMaterialEditor::Private::LogError(TEXT("Failed to find a generated material to export."), true, MaterialModelBase);
+			return;
+		}
 
-			FContentBrowserModule& ContentBrowserModule = FModuleManager::LoadModuleChecked<FContentBrowserModule>("ContentBrowser");
-			const FString SaveObjectPath = ContentBrowserModule.Get().CreateModalSaveAssetDialog(SaveAssetDialogConfig);
+		IContentBrowserSingleton& ContentBrowser = FModuleManager::LoadModuleChecked<FContentBrowserModule>("ContentBrowser").Get();
+		const FContentBrowserItemPath CurrentPath = ContentBrowser.GetCurrentPath();
+		const FString PathStr = CurrentPath.HasInternalPath() ? CurrentPath.GetInternalPathString() : "/Game";
 
-			if (SaveObjectPath.Len() == 0)
-			{
-				return;
-			}
+		FSaveAssetDialogConfig SaveAssetDialogConfig;
+		SaveAssetDialogConfig.DialogTitleOverride = LOCTEXT("SaveAssetDialogTitle", "Save Asset As");
+		SaveAssetDialogConfig.DefaultPath = PathStr;
+		SaveAssetDialogConfig.DefaultAssetName = GeneratedMaterial->GetName();
+		SaveAssetDialogConfig.ExistingAssetPolicy = ESaveAssetDialogExistingAssetPolicy::Disallow;
 
-			UDMBlueprintFunctionLibrary::ExportGeneratedMaterial(MaterialModel, SaveObjectPath);
+		FContentBrowserModule& ContentBrowserModule = FModuleManager::LoadModuleChecked<FContentBrowserModule>("ContentBrowser");
+		const FString SaveObjectPath = ContentBrowserModule.Get().CreateModalSaveAssetDialog(SaveAssetDialogConfig);
 
-			if (FEngineAnalytics::IsAvailable())
-			{
-				FEngineAnalytics::GetProvider().RecordEvent(TEXT("Editor.Usage.MaterialDesigner.ExportedGeneratedMaterial"));
-			}
+		if (SaveObjectPath.Len() == 0)
+		{
+			return;
+		}
+
+		UDMBlueprintFunctionLibrary::ExportGeneratedMaterial(MaterialModelBase, SaveObjectPath);
+
+		if (FEngineAnalytics::IsAvailable())
+		{
+			FEngineAnalytics::GetProvider().RecordEvent(TEXT("Editor.Usage.MaterialDesigner.ExportedGeneratedMaterial"));
 		}
 	}
 
-	void SnapshotMaterial(TWeakObjectPtr<UDynamicMaterialModel> InMaterialModelWeak, FIntPoint InTextureSize)
+	void SnapshotMaterial(TWeakObjectPtr<UDynamicMaterialModelBase> InMaterialModelBaseWeak, FIntPoint InTextureSize)
 	{
-		UDynamicMaterialModel* MaterialModel = InMaterialModelWeak.Get();
+		UDynamicMaterialModelBase* MaterialModelBase = InMaterialModelBaseWeak.Get();
 
-		if (!IsValid(MaterialModel))
+		if (!IsValid(MaterialModelBase))
 		{
 			return;
 		}
 
-		UMaterialInterface* Material = MaterialModel->GetGeneratedMaterial();
+		UMaterialInterface* Material = MaterialModelBase->GetGeneratedMaterial();
 
-		if (UDynamicMaterialInstance* MaterialInstance = MaterialModel->GetDynamicMaterialInstance())
+		if (UDynamicMaterialInstance* MaterialInstance = MaterialModelBase->GetDynamicMaterialInstance())
 		{
 			if (!IsValid(MaterialInstance->Parent.Get()))
 			{
@@ -234,14 +252,14 @@ namespace UE::DynamicMaterialEditor::Private
 			return;
 		}
 
-		UDynamicMaterialModel* const MaterialModel = MenuContext->GetModel();
+		UDynamicMaterialModelBase* const MaterialModelBase = MenuContext->GetModelBase();
 
-		if (!MaterialModel)
+		if (!MaterialModelBase)
 		{
 			return;
 		}
 
-		TWeakObjectPtr<UDynamicMaterialModel> MaterialModelWeak = MaterialModel;
+		TWeakObjectPtr<UDynamicMaterialModelBase> MaterialModelWeak = MaterialModelBase;
 		FToolMenuSection& NewSection = InMenu->AddSection("SnapshotMaterial", LOCTEXT("SnapshotMaterial", "Snapshop Material"));
 
 		NewSection.AddMenuEntry(
@@ -307,14 +325,14 @@ namespace UE::DynamicMaterialEditor::Private
 			return;
 		}
 
-		UDynamicMaterialModel* const MaterialModel = MenuContext->GetModel();
+		UDynamicMaterialModelBase* const MaterialModelBase = MenuContext->GetModelBase();
 
-		if (!MaterialModel)
+		if (!MaterialModelBase)
 		{
 			return;
 		}
 
-		UDynamicMaterialInstance* Instance = MaterialModel->GetDynamicMaterialInstance();
+		UDynamicMaterialInstance* Instance = MaterialModelBase->GetDynamicMaterialInstance();
 
 		if (!Instance)
 		{
@@ -322,7 +340,7 @@ namespace UE::DynamicMaterialEditor::Private
 		}
 
 		const bool bAllowInstanceExport = IsValid(Instance->GetOuter()) && !Instance->GetOuter()->IsA<UPackage>();
-		const bool bAllowMaterialExport = IsValid(MaterialModel->GetGeneratedMaterial());
+		const bool bAllowMaterialExport = IsValid(MaterialModelBase->GetGeneratedMaterial());
 
 		if (!bAllowInstanceExport && !bAllowMaterialExport)
 		{
@@ -362,7 +380,7 @@ namespace UE::DynamicMaterialEditor::Private
 				FSlateIcon(),
 				FUIAction(FExecuteAction::CreateStatic(
 					&UE::DynamicMaterialEditor::Private::ExportMaterialModelFromModel,
-					TWeakObjectPtr<UDynamicMaterialModel>(MaterialModel)
+					TWeakObjectPtr<UDynamicMaterialModelBase>(MaterialModelBase)
 				))
 			);
 		}
diff --git a/Engine/Plugins/Experimental/DynamicMaterial/Source/DynamicMaterialEditor/Private/Model/AssetDefinition_DynamicMaterialModel.cpp b/Engine/Plugins/Experimental/DynamicMaterial/Source/DynamicMaterialEditor/Private/Model/AssetDefinition_DynamicMaterialModel.cpp
--- a/Engine/Plugins/Experimental/DynamicMaterial/Source/DynamicMaterialEditor/Private/Model/AssetDefinition_DynamicMaterialModel.cpp
+++ b/Engine/Plugins/Experimental/DynamicMaterial/Source/DynamicMaterialEditor/Private/Model/AssetDefinition_DynamicMaterialModel.cpp
@@ -1,6 +1,6 @@
 // Copyright Epic Games, Inc. All Rights Reserved.
 
-#include "AssetDefinition_DynamicMaterialModel.h"
+#include "Model/AssetDefinition_DynamicMaterialModel.h"
 #include "DynamicMaterialEditorModule.h"
 #include "Engine/World.h"
 #include "Model/DynamicMaterialModel.h"
@@ -15,7 +15,7 @@ FText UAssetDefinition_DynamicMaterialModel::GetAssetDisplayName() const
 
 FText UAssetDefinition_DynamicMaterialModel::GetAssetDisplayName(const FAssetData& InAssetData) const
 {
-	return LOCTEXT("MaterialDesignerModel", "Material Designer Model");
+	return GetAssetDisplayName();
 }
 
 TSoftClassPtr<> UAssetDefinition_DynamicMaterialModel::GetAssetClass() const
diff --git a/Engine/Plugins/Experimental/DynamicMaterial/Source/DynamicMaterialEditor/Private/Model/AssetDefinition_DynamicMaterialModel.h b/Engine/Plugins/Experimental/DynamicMaterial/Source/DynamicMaterialEditor/Private/Model/AssetDefinition_DynamicMaterialModel.h
--- a/Engine/Plugins/Experimental/DynamicMaterial/Source/DynamicMaterialEditor/Private/Model/AssetDefinition_DynamicMaterialModel.h
+++ b/Engine/Plugins/Experimental/DynamicMaterial/Source/DynamicMaterialEditor/Private/Model/AssetDefinition_DynamicMaterialModel.h
@@ -11,12 +11,12 @@ class UAssetDefinition_DynamicMaterialModel : public UAssetDefinitionDefault
 	GENERATED_BODY()
 
 public:
-	//~ Begin UAssetDefinition Begin
+	//~ Begin UAssetDefinition
 	virtual FText GetAssetDisplayName() const override;
 	virtual FText GetAssetDisplayName(const FAssetData& InAssetData) const override;
 	virtual TSoftClassPtr<UObject> GetAssetClass() const override;
 	virtual FLinearColor GetAssetColor() const override;
 	virtual TConstArrayView<FAssetCategoryPath> GetAssetCategories() const override;
 	virtual EAssetCommandResult OpenAssets(const FAssetOpenArgs& InOpenArgs) const override;
-	//~ End UAssetDefinition End
+	//~ End UAssetDefinition
 };
diff --git a/Engine/Plugins/Experimental/DynamicMaterial/Source/DynamicMaterialEditor/Private/Model/DMMaterialBuildState.cpp b/Engine/Plugins/Experimental/DynamicMaterial/Source/DynamicMaterialEditor/Private/Model/DMMaterialBuildState.cpp
--- a/Engine/Plugins/Experimental/DynamicMaterial/Source/DynamicMaterialEditor/Private/Model/DMMaterialBuildState.cpp
+++ b/Engine/Plugins/Experimental/DynamicMaterial/Source/DynamicMaterialEditor/Private/Model/DMMaterialBuildState.cpp
@@ -473,46 +473,6 @@ const TMap<const UDMMaterialValue*, TArray<UMaterialExpression*>>& FDMMaterialBu
 	return Values;
 }
 
-///////////////////////////////////////
-/// Callbacks
-
-bool FDMMaterialBuildState::HasCallback(const UDMMaterialStageSource* InCallback) const
-{
-	return Callbacks.Contains(InCallback);
-}
-
-const TArray<UDMMaterialStageSource*>& FDMMaterialBuildState::GetCallbackExpressions(const UDMMaterialStageSource* InCallback) const
-{
-	const TArray<UDMMaterialStageSource*>* CallbackSources = Callbacks.Find(InCallback);
-
-	if (ensure(CallbackSources))
-	{
-		return *CallbackSources;
-	}
-
-	return UE::DynamicMaterialEditor::Private::EmptySourceSet;
-}
-
-void FDMMaterialBuildState::AddCallbackExpressions(const UDMMaterialStageSource* InCallback, const TArray<UDMMaterialStageSource*>& InCallbackExpressions)
-{
-	if (ensure(!HasCallback(InCallback)))
-	{
-		Callbacks.Emplace(InCallback, InCallbackExpressions);
-	}
-}
-
-TArray<const UDMMaterialStageSource*> FDMMaterialBuildState::GetCallbacks() const
-{
-	TArray<const UDMMaterialStageSource*> Keys;
-	Callbacks.GetKeys(Keys);
-	return Keys;
-}
-
-const TMap<const UDMMaterialStageSource*, TArray<UDMMaterialStageSource*>>& FDMMaterialBuildState::GetCallbackMap() const
-{
-	return Callbacks;
-}
-
 ///////////////////////////////////////
 /// Other expression
 
diff --git a/Engine/Plugins/Experimental/DynamicMaterial/Source/DynamicMaterialEditor/Private/Model/DMMaterialBuildUtils.cpp b/Engine/Plugins/Experimental/DynamicMaterial/Source/DynamicMaterialEditor/Private/Model/DMMaterialBuildUtils.cpp
--- a/Engine/Plugins/Experimental/DynamicMaterial/Source/DynamicMaterialEditor/Private/Model/DMMaterialBuildUtils.cpp
+++ b/Engine/Plugins/Experimental/DynamicMaterial/Source/DynamicMaterialEditor/Private/Model/DMMaterialBuildUtils.cpp
@@ -1,7 +1,10 @@
 // Copyright Epic Games, Inc. All Rights Reserved.
 
 #include "Model/DMMaterialBuildUtils.h"
+
+#include "Components/DMMaterialProperty.h"
 #include "Components/DMMaterialStageInput.h"
+#include "DynamicMaterialEditorModule.h"
 #include "MaterialEditingLibrary.h"
 #include "Materials/Material.h"
 #include "Materials/MaterialExpression.h"
@@ -54,7 +57,7 @@ UMaterialExpression* FDMMaterialBuildUtils::CreateExpression(TSubclassOf<UMateri
 }
 
 UMaterialExpression* FDMMaterialBuildUtils::CreateExpressionParameter(TSubclassOf<UMaterialExpression> InExpressionClass, 
-	FName InParameterName, const FString& InComment, UObject* InAsset /*= nullptr*/) const
+	FName InParameterName, EDMMaterialParameterGroup InParameterGroup, const FString& InComment, UObject* InAsset /*= nullptr*/) const
 {
 	check(BuildState.GetDynamicMaterial());
 	check(InExpressionClass.Get());
@@ -71,8 +74,54 @@ UMaterialExpression* FDMMaterialBuildUtils::CreateExpressionParameter(TSubclassO
 	);
 
 	NewExpression->Desc = InComment;
-
 	NewExpression->SetParameterName(InParameterName);
+
+	if (FNameProperty* GroupProperty = CastField<FNameProperty>(NewExpression->GetClass()->FindPropertyByName("Group")))
+	{
+		int32 PropertyIndex;
+		FString PropertyName;
+
+		switch (InParameterGroup)
+		{
+			case EDMMaterialParameterGroup::Property:
+				if (const UDMMaterialProperty* Property = BuildState.GetCurrentMaterialProperty())
+				{
+					const EDMMaterialPropertyType PropertyType = Property->GetMaterialProperty();
+					PropertyIndex = static_cast<int32>(PropertyType);
+					PropertyName = StaticEnum<EDMMaterialPropertyType>()->GetDisplayNameTextByValue(static_cast<int64>(PropertyType)).ToString();
+					break;
+				}
+
+				// Error - fall back to hidden group
+				if (!BuildState.IsPreviewMaterial())
+				{
+					UE_LOG(LogDynamicMaterialEditor, Error, TEXT("Missing property for group type."));
+				}
+
+				// Fall through
+
+			case EDMMaterialParameterGroup::NotExposed:
+			default:
+				PropertyIndex = 99;
+				PropertyName = TEXT("Uncategorized");
+				break;
+
+			case EDMMaterialParameterGroup::Global:
+				PropertyIndex = 0;
+				PropertyName = TEXT("Global");
+				break;
+		}
+
+		GroupProperty->SetValue_InContainer(
+			NewExpression, 
+			*FString::Printf(
+				TEXT("%02d - %s"),
+				PropertyIndex,
+				*PropertyName
+			)
+		);
+	}
+
 	BuildState.GetDynamicMaterial()->GetEditorOnlyData()->ExpressionCollection.AddExpression(NewExpression);
 
 	TArray<UMaterialExpression*>& ExpressionList = BuildState.GetDynamicMaterial()->EditorParameters.FindOrAdd(InParameterName);
diff --git a/Engine/Plugins/Experimental/DynamicMaterial/Source/DynamicMaterialEditor/Private/Model/DynamicMaterialModelEditorOnlyData.cpp b/Engine/Plugins/Experimental/DynamicMaterial/Source/DynamicMaterialEditor/Private/Model/DynamicMaterialModelEditorOnlyData.cpp
--- a/Engine/Plugins/Experimental/DynamicMaterial/Source/DynamicMaterialEditor/Private/Model/DynamicMaterialModelEditorOnlyData.cpp
+++ b/Engine/Plugins/Experimental/DynamicMaterial/Source/DynamicMaterialEditor/Private/Model/DynamicMaterialModelEditorOnlyData.cpp
@@ -52,10 +52,31 @@
 
 #define LOCTEXT_NAMESPACE "MaterialDesignerModel"
 
+const FString UDynamicMaterialModelEditorOnlyData::SlotsPathToken               = FString(TEXT("Slots"));
+const FString UDynamicMaterialModelEditorOnlyData::BaseColorSlotPathToken       = FString(TEXT("BaseColor"));
+const FString UDynamicMaterialModelEditorOnlyData::EmissiveSlotPathToken        = FString(TEXT("Emissive"));
+const FString UDynamicMaterialModelEditorOnlyData::OpacitySlotPathToken         = FString(TEXT("Opacity"));
+const FString UDynamicMaterialModelEditorOnlyData::RoughnessPathToken           = FString(TEXT("Roughness"));
+const FString UDynamicMaterialModelEditorOnlyData::SpecularPathToken            = FString(TEXT("Specular"));
+const FString UDynamicMaterialModelEditorOnlyData::MetallicPathToken            = FString(TEXT("Metallic"));
+const FString UDynamicMaterialModelEditorOnlyData::NormalPathToken              = FString(TEXT("Normal"));
+const FString UDynamicMaterialModelEditorOnlyData::PixelDepthOffsetPathToken    = FString(TEXT("PDO"));
+const FString UDynamicMaterialModelEditorOnlyData::WorldPositionOffsetPathToken = FString(TEXT("WPO"));
+const FString UDynamicMaterialModelEditorOnlyData::AmbientOcclusionPathToken    = FString(TEXT("AO"));
+const FString UDynamicMaterialModelEditorOnlyData::AnisotropyPathToken          = FString(TEXT("Anisotropy"));
+const FString UDynamicMaterialModelEditorOnlyData::RefractionPathToken          = FString(TEXT("Refraction"));
+const FString UDynamicMaterialModelEditorOnlyData::TangentPathToken             = FString(TEXT("Tangent"));
+const FString UDynamicMaterialModelEditorOnlyData::Custom1PathToken             = FString(TEXT("Custom1"));
+const FString UDynamicMaterialModelEditorOnlyData::Custom2PathToken             = FString(TEXT("Custom2"));
+const FString UDynamicMaterialModelEditorOnlyData::Custom3PathToken             = FString(TEXT("Custom3"));
+const FString UDynamicMaterialModelEditorOnlyData::Custom4PathToken             = FString(TEXT("Custom4"));
+const FString UDynamicMaterialModelEditorOnlyData::PropertiesPathToken          = FString(TEXT("Properties"));
+
 namespace UE::DynamicMaterialEditor::Private
 {
 	const TMap<FString, EDMMaterialPropertyType> TokenToPropertyMap = {
-		{UDynamicMaterialModelEditorOnlyData::RGBSlotPathToken,             EDMMaterialPropertyType::BaseColor},
+		{UDynamicMaterialModelEditorOnlyData::BaseColorSlotPathToken,       EDMMaterialPropertyType::BaseColor},
+		{UDynamicMaterialModelEditorOnlyData::EmissiveSlotPathToken,        EDMMaterialPropertyType::EmissiveColor},
 		{UDynamicMaterialModelEditorOnlyData::OpacitySlotPathToken,         EDMMaterialPropertyType::Opacity},
 		{UDynamicMaterialModelEditorOnlyData::RoughnessPathToken,           EDMMaterialPropertyType::Roughness},
 		{UDynamicMaterialModelEditorOnlyData::SpecularPathToken,            EDMMaterialPropertyType::Specular},
@@ -74,25 +95,6 @@ namespace UE::DynamicMaterialEditor::Private
 	};
 }
 
-const FString UDynamicMaterialModelEditorOnlyData::SlotsPathToken               = FString(TEXT("Slots"));
-const FString UDynamicMaterialModelEditorOnlyData::RGBSlotPathToken             = FString(TEXT("RGBSlot"));
-const FString UDynamicMaterialModelEditorOnlyData::OpacitySlotPathToken         = FString(TEXT("OpacitySlot"));
-const FString UDynamicMaterialModelEditorOnlyData::RoughnessPathToken           = FString(TEXT("RoughnessSlot"));
-const FString UDynamicMaterialModelEditorOnlyData::SpecularPathToken            = FString(TEXT("SpecularSlot"));
-const FString UDynamicMaterialModelEditorOnlyData::MetallicPathToken            = FString(TEXT("MetallicSlot"));
-const FString UDynamicMaterialModelEditorOnlyData::NormalPathToken              = FString(TEXT("NormalSlot"));
-const FString UDynamicMaterialModelEditorOnlyData::PixelDepthOffsetPathToken    = FString(TEXT("PixelDepthOffsetSlot"));
-const FString UDynamicMaterialModelEditorOnlyData::WorldPositionOffsetPathToken = FString(TEXT("WorldPositionOffsetSlot"));
-const FString UDynamicMaterialModelEditorOnlyData::AmbientOcclusionPathToken    = FString(TEXT("AmbientOcclusionSlot"));
-const FString UDynamicMaterialModelEditorOnlyData::AnisotropyPathToken          = FString(TEXT("AnisotropySlot"));
-const FString UDynamicMaterialModelEditorOnlyData::RefractionPathToken          = FString(TEXT("RefractionSlot"));
-const FString UDynamicMaterialModelEditorOnlyData::TangentPathToken             = FString(TEXT("TangentSlot"));
-const FString UDynamicMaterialModelEditorOnlyData::Custom1PathToken             = FString(TEXT("Custom1Slot"));
-const FString UDynamicMaterialModelEditorOnlyData::Custom2PathToken             = FString(TEXT("Custom2Slot"));
-const FString UDynamicMaterialModelEditorOnlyData::Custom3PathToken             = FString(TEXT("Custom3Slot"));
-const FString UDynamicMaterialModelEditorOnlyData::Custom4PathToken             = FString(TEXT("Custom4Slot"));
-const FString UDynamicMaterialModelEditorOnlyData::PropertiesPathToken          = FString(TEXT("Properties"));
-
 const FName UDynamicMaterialModelEditorOnlyData::AlphaValueName = TEXT("AlphaValue");
 
 const TArray<EMaterialDomain> UDynamicMaterialModelEditorOnlyData::SupportedDomains =
@@ -111,6 +113,21 @@ const TArray<EBlendMode> UDynamicMaterialModelEditorOnlyData::SupportedBlendMode
 	EBlendMode::BLEND_Modulate
 };
 
+UDynamicMaterialModelEditorOnlyData* UDynamicMaterialModelEditorOnlyData::Get(UDynamicMaterialModelBase* InModelBase)
+{
+	if (InModelBase)
+	{
+		return Get(InModelBase->ResolveMaterialModel());
+	}
+
+	return nullptr;
+}
+
+UDynamicMaterialModelEditorOnlyData* UDynamicMaterialModelEditorOnlyData::Get(const TWeakObjectPtr<UDynamicMaterialModelBase>& InModelBaseWeak)
+{
+	return Get(InModelBaseWeak.Get());
+}
+
 UDynamicMaterialModelEditorOnlyData* UDynamicMaterialModelEditorOnlyData::Get(UDynamicMaterialModel* InModel)
 {
 	if (InModel)
@@ -121,7 +138,7 @@ UDynamicMaterialModelEditorOnlyData* UDynamicMaterialModelEditorOnlyData::Get(UD
 	return nullptr;
 }
 
-UDynamicMaterialModelEditorOnlyData* UDynamicMaterialModelEditorOnlyData::Get(TWeakObjectPtr<UDynamicMaterialModel> InModelWeak)
+UDynamicMaterialModelEditorOnlyData* UDynamicMaterialModelEditorOnlyData::Get(const TWeakObjectPtr<UDynamicMaterialModel>& InModelWeak)
 {
 	return Get(InModelWeak.Get());
 }
@@ -360,20 +377,6 @@ UMaterial* UDynamicMaterialModelEditorOnlyData::GetGeneratedMaterial() const
 	return IsValid(MaterialModel) ? MaterialModel->DynamicMaterial : nullptr;
 }
 
-void UDynamicMaterialModelEditorOnlyData::ResetData()
-{
-	if (GUndo)
-	{
-		Modify();
-	}
-
-	BlendMode = EBlendMode::BLEND_Translucent;
-	ShadingModel = EDMMaterialShadingModel::Unlit;
-	Slots.Empty();
-	PropertySlotMap.Empty();
-	Expressions.Empty();
-}
-
 void UDynamicMaterialModelEditorOnlyData::CreateMaterial()
 {
 	if (!IsValid(MaterialModel))
@@ -390,7 +393,7 @@ void UDynamicMaterialModelEditorOnlyData::CreateMaterial()
 			UMaterial::StaticClass(),
 			MaterialModel,
 			NAME_None,
-			RF_DuplicateTransient | RF_TextExportTransient,
+			RF_DuplicateTransient | RF_TextExportTransient | RF_Public,
 			nullptr,
 			GWarn
 		));
@@ -904,7 +907,7 @@ bool UDynamicMaterialModelEditorOnlyData::AddTextureSet(UDMTextureSet* InTexture
 
 bool UDynamicMaterialModelEditorOnlyData::NeedsWizard() const
 {
-	return ChannelListPreset == NAME_None;
+	return ChannelListPreset.IsNone();
 }
 
 void UDynamicMaterialModelEditorOnlyData::OnWizardComplete()
@@ -1363,7 +1366,7 @@ UDMMaterialSlot* UDynamicMaterialModelEditorOnlyData::RemoveSlotForMaterialPrope
 	return nullptr;
 }
 
-TArray<EDMMaterialPropertyType> UDynamicMaterialModelEditorOnlyData::GetMaterialPropertiesForSlot(UDMMaterialSlot* Slot) const
+TArray<EDMMaterialPropertyType> UDynamicMaterialModelEditorOnlyData::GetMaterialPropertiesForSlot(const UDMMaterialSlot* Slot) const
 {
 	TArray<EDMMaterialPropertyType> OutProperties;
 
@@ -1390,7 +1393,7 @@ void UDynamicMaterialModelEditorOnlyData::AssignMaterialPropertyToSlot(EDMMateri
 
 	PropertySlotMap.FindOrAdd(Property) = Slot;
 	Properties[Property]->ResetInputConnectionMap();
-	Slot->GetOnPropertiesUpdateDelegate().Broadcast(Slot);
+	Slot->OnPropertiesUpdated();
 
 	RequestMaterialBuild();
 }
@@ -1405,12 +1408,12 @@ void UDynamicMaterialModelEditorOnlyData::UnassignMaterialProperty(EDMMaterialPr
 	}
 
 	PropertySlotMap.Remove(Property);
-	(*SlotPtr)->GetOnPropertiesUpdateDelegate().Broadcast(*SlotPtr);
+	(*SlotPtr)->OnPropertiesUpdated();
 
 	RequestMaterialBuild();
 }
 
-void UDynamicMaterialModelEditorOnlyData::NotifyPostChange(const FPropertyChangedEvent& PropertyChangedEvent, class FEditPropertyChain* PropertyThatChanged)
+void UDynamicMaterialModelEditorOnlyData::NotifyPostChange(const FPropertyChangedEvent& InPropertyChangedEvent, class FEditPropertyChain* PropertyThatChanged)
 {
 	RequestMaterialBuild();
 }
@@ -1507,7 +1510,7 @@ void UDynamicMaterialModelEditorOnlyData::Serialize(FArchive& Ar)
 	// transparent channel, let's set it back to translucent.
 	if (Version < FDynamicMaterialModelEditorOnlyDataVersion::GlobalValueRename)
 	{
-		if (ChannelListPreset == NAME_None)
+		if (ChannelListPreset.IsNone())
 		{
 			// Try to guess from the available slots.
 			if (GetSlotForMaterialProperty(EDMMaterialPropertyType::Opacity))
diff --git a/Engine/Plugins/Experimental/DynamicMaterial/Source/DynamicMaterialEditor/Private/Slate/Layers/SDMSlotLayerItem.cpp b/Engine/Plugins/Experimental/DynamicMaterial/Source/DynamicMaterialEditor/Private/Slate/Layers/SDMSlotLayerItem.cpp
--- a/Engine/Plugins/Experimental/DynamicMaterial/Source/DynamicMaterialEditor/Private/Slate/Layers/SDMSlotLayerItem.cpp
+++ b/Engine/Plugins/Experimental/DynamicMaterial/Source/DynamicMaterialEditor/Private/Slate/Layers/SDMSlotLayerItem.cpp
@@ -23,6 +23,7 @@
 #include "Slate/SDMSlot.h"
 #include "Slate/SDMStage.h"
 #include "SlateOptMacros.h"
+#include "Model/DynamicMaterialModel.h"
 #include "Styling/StyleColors.h"
 #include "Utils/DMPrivate.h"
 #include "Widgets/Images/SImage.h"
@@ -51,6 +52,16 @@ void SDMSlotLayerItem::Construct(const FArguments& InArgs, const TSharedPtr<SDMS
 	OnStageSelected = InArgs._OnStageSelected;
 	OnLayerLinkToggled = InArgs._OnLayerLinkToggled;
 
+	bIsDynamic = false;
+
+	if (TSharedPtr<SDMSlot> SlotWidget = SlotWidgetWeak.Pin())
+	{
+		if (TSharedPtr<SDMEditor> EditorWidget = SlotWidget->GetEditorWidget())
+		{
+			bIsDynamic = !Cast<UDynamicMaterialModel>(EditorWidget->GetMaterialModelBase());
+		}
+	}
+
 	STableRow<TSharedPtr<FDMMaterialLayerReference>>::Construct(
 		STableRow<TSharedPtr<FDMMaterialLayerReference>>::FArguments()
 		.Padding(2.0f)
@@ -269,8 +280,15 @@ TSharedRef<SWidget> SDMSlotLayerItem::CreateStageBaseWidget(const bool bInteract
 		return SNullWidget::NullWidget;
 	}
 
+	TSharedPtr<SDMSlot> SlotWidget = SlotWidgetWeak.Pin();
+
+	if (!SlotWidget.IsValid())
+	{
+		return SNullWidget::NullWidget;
+	}
+
 	return 
-		SAssignNew(BaseStageWidget, SDMStage, Layer->GetStage(EDMMaterialLayerStage::Base))
+		SAssignNew(BaseStageWidget, SDMStage, SlotWidget.ToSharedRef(), Layer->GetStage(EDMMaterialLayerStage::Base))
 		.Interactable(bInteractable)
 		.StageEnabled(StageBaseEnabled)
 		.StageSelected(StageBaseSelected)
@@ -299,8 +317,15 @@ TSharedRef<SWidget> SDMSlotLayerItem::CreateStageMaskWidget(const bool bInteract
 		return SNullWidget::NullWidget;
 	}
 
+	TSharedPtr<SDMSlot> SlotWidget = SlotWidgetWeak.Pin();
+
+	if (!SlotWidget.IsValid())
+	{
+		return SNullWidget::NullWidget;
+	}
+
 	return
-		SAssignNew(MaskStageWidget, SDMStage, Layer->GetStage(EDMMaterialLayerStage::Mask))
+		SAssignNew(MaskStageWidget, SDMStage, SlotWidget.ToSharedRef(), Layer->GetStage(EDMMaterialLayerStage::Mask))
 		.Interactable(bInteractable)
 		.StageEnabled(StageMaskEnabled)
 		.StageSelected(StageMaskSelected)
@@ -321,7 +346,7 @@ TSharedRef<SWidget> SDMSlotLayerItem::CreateHandleWidget()
 		SNew(SBox)
 		.HAlign(HAlign_Fill)
 		.VAlign(VAlign_Fill)
-		.Cursor(EMouseCursor::GrabHand)
+		.Cursor(bIsDynamic ? EMouseCursor::Default : EMouseCursor::GrabHand)
 		.ToolTipText(this, &SDMSlotLayerItem::GetToolTipText)
 		[
 			SNew(SBorder)
@@ -353,6 +378,7 @@ TSharedRef<SWidget> SDMSlotLayerItem::CreateLayerBypassButton()
 {
 	return 
 		SNew(SButton)
+		.IsEnabled(!bIsDynamic)
 		.ContentPadding(4.0f)
 		.ButtonStyle(FDynamicMaterialEditorStyle::Get(), "HoverHintOnly")
 		.ToolTipText(LOCTEXT("LayerBypassTooltip", "Toggle the bypassing of this layer."))
@@ -396,6 +422,7 @@ TSharedRef<SWidget> SDMSlotLayerItem::CreateLayerBaseToggleButton()
 {
 	return
 		SNew(SButton)
+		.IsEnabled(!bIsDynamic)
 		.ContentPadding(0.0f)
 		.ButtonStyle(FDynamicMaterialEditorStyle::Get(), "HoverHintOnly")
 		.ToolTipText(LOCTEXT("MaterialLayerBaseToggleTooltip", "Toggle the layer base on and off.\n\n"
@@ -412,6 +439,7 @@ TSharedRef<SWidget> SDMSlotLayerItem::CreateLayerMaskToggleButton()
 {
 	return
 		SNew(SButton)
+		.IsEnabled(!bIsDynamic)
 		.ContentPadding(0.0f)
 		.ButtonStyle(FDynamicMaterialEditorStyle::Get(), "HoverHintOnly")
 		.ToolTipText(LOCTEXT("MaterialLayerMaskToggleTooltip", "Toggle the layer mask on and off."))
@@ -426,6 +454,7 @@ TSharedRef<SWidget> SDMSlotLayerItem::CreateLayerLinkToggleButton()
 {
 	return 
 		SNew(SButton)
+		.IsEnabled(!bIsDynamic)
 		.ContentPadding(0.0f)
 		.ButtonStyle(FDynamicMaterialEditorStyle::Get(), "HoverHintOnly")
 		.ToolTipText(LOCTEXT("MaterialStageLinkTooltip", "Click to toggle UV Link on and off."))
@@ -632,6 +661,11 @@ void SDMSlotLayerItem::OnLayerItemDragLeave(const FDragDropEvent& InDragDropEven
 
 FReply SDMSlotLayerItem::OnLayerItemDragDetected(const FGeometry& InMyGeometry, const FPointerEvent& InMouseEvent)
 {
+	if (bIsDynamic)
+	{
+		return FReply::Handled();
+	}
+
 	const bool bShouldDuplicate = InMouseEvent.IsAltDown();
 
 	TSharedRef<FDMSlotLayerDragDropOperation> SlotLayerDragDropOp = MakeShared<FDMSlotLayerDragDropOperation>(SharedThis(this), bShouldDuplicate);
@@ -1230,20 +1264,23 @@ TSharedRef<SWidget> SDMSlotLayerItem::CreateLayerHeaderText() const
 		.TextStyle(FDynamicMaterialEditorStyle::Get(), "SmallFont")
 		.Text(this, &SDMSlotLayerItem::GetLayerHeaderText);
 
-	TextBlock->SetOnMouseButtonDown(FPointerEventHandler::CreateSPLambda(
-		this,
-		[this](const FGeometry& InGeometry, const FPointerEvent& InPointerEvent)
-		{
-			if (LayerHeaderTextContainer.IsValid())
+	if (!bIsDynamic)
+	{
+		TextBlock->SetOnMouseButtonDown(FPointerEventHandler::CreateSPLambda(
+			this,
+			[this](const FGeometry& InGeometry, const FPointerEvent& InPointerEvent)
 			{
-				TSharedRef<SWidget> EditableContent = CreateLayerHeaderEditableText();
-				LayerHeaderTextContainer->SetContent(EditableContent);
-				FSlateApplication::Get().SetKeyboardFocus(EditableContent);
-			}
+				if (LayerHeaderTextContainer.IsValid())
+				{
+					TSharedRef<SWidget> EditableContent = CreateLayerHeaderEditableText();
+					LayerHeaderTextContainer->SetContent(EditableContent);
+					FSlateApplication::Get().SetKeyboardFocus(EditableContent);
+				}
 
-			return FReply::Handled();
-		}
-	));
+				return FReply::Handled();
+			}
+		));
+	}
 
 	return TextBlock;
 }
diff --git a/Engine/Plugins/Experimental/DynamicMaterial/Source/DynamicMaterialEditor/Private/Slate/Layers/SDMSlotLayerItem.h b/Engine/Plugins/Experimental/DynamicMaterial/Source/DynamicMaterialEditor/Private/Slate/Layers/SDMSlotLayerItem.h
--- a/Engine/Plugins/Experimental/DynamicMaterial/Source/DynamicMaterialEditor/Private/Slate/Layers/SDMSlotLayerItem.h
+++ b/Engine/Plugins/Experimental/DynamicMaterial/Source/DynamicMaterialEditor/Private/Slate/Layers/SDMSlotLayerItem.h
@@ -73,6 +73,7 @@ protected:
 	TWeakPtr<SDMSlot> SlotWidgetWeak;
 	TWeakPtr<STableViewBase> LayerViewWeak;
 	TSharedPtr<FDMMaterialLayerReference> LayerItem;
+	bool bIsDynamic;
 
 	TAttribute<bool> StageBaseEnabled;
 	TAttribute<bool> StageBaseSelected;
diff --git a/Engine/Plugins/Experimental/DynamicMaterial/Source/DynamicMaterialEditor/Private/Slate/Previews/SDMComponentPreview.cpp b/Engine/Plugins/Experimental/DynamicMaterial/Source/DynamicMaterialEditor/Private/Slate/Previews/SDMComponentPreview.cpp
new file mode 100644
--- /dev/null
+++ b/Engine/Plugins/Experimental/DynamicMaterial/Source/DynamicMaterialEditor/Private/Slate/Previews/SDMComponentPreview.cpp
@@ -0,0 +1,183 @@
+// Copyright Epic Games, Inc. All Rights Reserved.
+
+#include "SDMComponentPreview.h"
+#include "Components/DMMaterialStage.h"
+#include "Components/DMMaterialValue.h"
+#include "Components/DMTextureUV.h"
+#include "Components/DMTextureUVDynamic.h"
+#include "DynamicMaterialModule.h"
+#include "Materials/Material.h"
+#include "Materials/MaterialInstanceDynamic.h"
+#include "Model/DMMaterialBuildUtils.h"
+#include "Model/DynamicMaterialModel.h"
+#include "Model/DynamicMaterialModelDynamic.h"
+#include "Slate/SDMEditor.h"
+#include "Widgets/Images/SImage.h"
+
+SDMComponentPreview::SDMComponentPreview()
+	: Brush(FSlateMaterialBrush(FVector2D(1.f, 1.f)))
+{
+	Brush.SetUVRegion(FBox2f(FVector2f::ZeroVector, FVector2f::UnitVector));
+}
+
+SDMComponentPreview::~SDMComponentPreview()
+{
+	if (!FDynamicMaterialModule::AreUObjectsSafe())
+	{
+		return;
+	}
+
+	if (UDMMaterialComponent* Component = ComponentWeak.Get())
+	{
+		Component->GetOnUpdate().RemoveAll(this);
+
+		if (TSharedPtr<SDMEditor> EditorWidget = EditorWidgetWeak.Pin())
+		{
+			EditorWidget->FreePreviewMaterial(Component);
+		}
+	}
+
+	if (UDynamicMaterialModel* MaterialModel = Cast<UDynamicMaterialModel>(MaterialModelBaseWeak.Get()))
+	{
+		MaterialModel->GetOnValueUpdateDelegate().RemoveAll(this);
+		MaterialModel->GetOnTextureUVUpdateDelegate().RemoveAll(this);
+	}
+	else if (UDynamicMaterialModelDynamic* MaterialModelDynamic = Cast<UDynamicMaterialModelDynamic>(MaterialModelBaseWeak.Get()))
+	{
+		MaterialModelDynamic->GetOnValueDynamicUpdateDelegate().RemoveAll(this);
+		MaterialModelDynamic->GetOnTextureUVDynamicUpdateDelegate().RemoveAll(this);
+	}
+}
+
+void SDMComponentPreview::Construct(const FArguments& InArgs, const TSharedRef<SDMEditor>& InEditorWidget, UDMMaterialComponent* InComponent)
+{
+	EditorWidgetWeak = InEditorWidget;
+	ComponentWeak = InComponent;
+
+	PreviewMaterialBaseWeak = InEditorWidget->CreatePreviewMaterial(InComponent);
+	PreviewMaterialDynamicWeak = InEditorWidget->CreateMID(PreviewMaterialBaseWeak.Get());
+	MaterialModelBaseWeak = InEditorWidget->GetMaterialModelBase();
+
+	if (UDynamicMaterialModel* MaterialModel = Cast<UDynamicMaterialModel>(MaterialModelBaseWeak.Get()))
+	{
+		MaterialModel->GetOnValueUpdateDelegate().AddSP(this, &SDMComponentPreview::OnValueUpdated);
+		MaterialModel->GetOnTextureUVUpdateDelegate().AddSP(this, &SDMComponentPreview::OnTextureUVUpdated);
+		MaterialModel->ApplyComponents(PreviewMaterialDynamicWeak.Get());
+	}
+	else if (UDynamicMaterialModelDynamic* MaterialModelDynamic = Cast<UDynamicMaterialModelDynamic>(MaterialModelBaseWeak.Get()))
+	{
+		MaterialModelDynamic->GetOnValueDynamicUpdateDelegate().AddSP(this, &SDMComponentPreview::OnValueDynamicUpdated);
+		MaterialModelDynamic->GetOnTextureUVDynamicUpdateDelegate().AddSP(this, &SDMComponentPreview::OnTextureUVDynamicUpdated);
+		MaterialModelDynamic->ApplyComponents(PreviewMaterialDynamicWeak.Get());
+	}
+	else
+	{
+		return;
+	}
+
+	PreviewSize = InArgs._PreviewSize;
+
+	SetCanTick(true);
+
+	Brush.SetImageSize(PreviewSize.Get());
+
+	if (ensure(IsValid(InComponent)))
+	{
+		InComponent->GetOnUpdate().AddSP(this, &SDMComponentPreview::OnComponentUpdated);
+		OnComponentUpdated(InComponent, EDMUpdateType::Structure);
+	}
+
+	ChildSlot
+	[
+		SNew(SImage)
+		.Image(&Brush)
+		.DesiredSizeOverride(this, &SDMComponentPreview::GetPreviewSize)
+	];
+}
+
+void SDMComponentPreview::Tick(const FGeometry& AllottedGeometry, const double InCurrentTime, const float InDeltaTime)
+{
+	SCompoundWidget::Tick(AllottedGeometry, InCurrentTime, InDeltaTime);
+
+	if (!PreviewMaterialBaseWeak.IsValid() || !PreviewMaterialDynamicWeak.IsValid())
+	{
+		Brush.SetMaterial(nullptr);
+	}	
+}
+
+void SDMComponentPreview::OnComponentUpdated(UDMMaterialComponent* InComponent, EDMUpdateType InUpdateType)
+{
+	if (UDMMaterialStage* Stage = Cast<UDMMaterialStage>(InComponent))
+	{
+		if (Stage == ComponentWeak.Get() && IsValid(Stage) && Stage->IsComponentValid())
+		{
+			if (TSharedPtr<SDMEditor> EditorWidget = EditorWidgetWeak.Pin())
+			{
+				UMaterial* PreviewMaterialBase = PreviewMaterialBaseWeak.Get();
+
+				if (!PreviewMaterialBase)
+				{
+					PreviewMaterialBase = EditorWidget->CreatePreviewMaterial(Stage);
+				}
+
+				if (InUpdateType == EDMUpdateType::Structure)
+				{
+					Stage->GeneratePreviewMaterial(PreviewMaterialBase);
+
+					EditorWidget->FreeMID(PreviewMaterialBase);
+					PreviewMaterialDynamicWeak = EditorWidget->CreateMID(PreviewMaterialBase);
+
+					UDynamicMaterialModelBase* MaterialModelBase = EditorWidget->GetMaterialModelBase();
+
+					if (UDynamicMaterialModel* MaterialModel = Cast<UDynamicMaterialModel>(MaterialModelBase))
+					{
+						MaterialModel->ApplyComponents(PreviewMaterialDynamicWeak.Get());
+					}
+					else if (UDynamicMaterialModelDynamic* MaterialModelDynamic = Cast<UDynamicMaterialModelDynamic>(MaterialModelBase))
+					{
+						MaterialModelDynamic->ApplyComponents(PreviewMaterialDynamicWeak.Get());
+					}
+
+					Brush.SetMaterial(PreviewMaterialDynamicWeak.Get());
+				}
+			}
+		}
+	}
+}
+
+TOptional<FVector2D> SDMComponentPreview::GetPreviewSize() const
+{
+	return PreviewSize.Get();
+}
+
+void SDMComponentPreview::OnValueUpdated(UDynamicMaterialModel* InMaterialModel, UDMMaterialValue* InValue)
+{
+	if (UMaterialInstanceDynamic* PreviewMaterialDynamic = PreviewMaterialDynamicWeak.Get())
+	{
+		InValue->SetMIDParameter(PreviewMaterialDynamic);
+	}
+}
+
+void SDMComponentPreview::OnTextureUVUpdated(UDynamicMaterialModel* InMaterialModel, UDMTextureUV* InTextureUV)
+{
+	if (UMaterialInstanceDynamic* PreviewMaterialDynamic = PreviewMaterialDynamicWeak.Get())
+	{
+		InTextureUV->SetMIDParameters(PreviewMaterialDynamic);
+	}
+}
+
+void SDMComponentPreview::OnValueDynamicUpdated(UDynamicMaterialModelDynamic* InMaterialModel, UDMMaterialValueDynamic* InValueDynamic)
+{
+	if (UMaterialInstanceDynamic* PreviewMaterialDynamic = PreviewMaterialDynamicWeak.Get())
+	{
+		InValueDynamic->SetMIDParameter(PreviewMaterialDynamic);
+	}
+}
+
+void SDMComponentPreview::OnTextureUVDynamicUpdated(UDynamicMaterialModelDynamic* InMaterialModel, UDMTextureUVDynamic* InTextureUVDynamic)
+{
+	if (UMaterialInstanceDynamic* PreviewMaterialDynamic = PreviewMaterialDynamicWeak.Get())
+	{
+		InTextureUVDynamic->SetMIDParameters(PreviewMaterialDynamic);
+	}
+}
diff --git a/Engine/Plugins/Experimental/DynamicMaterial/Source/DynamicMaterialEditor/Private/Slate/Previews/SDMComponentPreview.h b/Engine/Plugins/Experimental/DynamicMaterial/Source/DynamicMaterialEditor/Private/Slate/Previews/SDMComponentPreview.h
new file mode 100644
--- /dev/null
+++ b/Engine/Plugins/Experimental/DynamicMaterial/Source/DynamicMaterialEditor/Private/Slate/Previews/SDMComponentPreview.h
@@ -0,0 +1,60 @@
+// Copyright Epic Games, Inc. All Rights Reserved.
+
+#pragma once
+
+#include "DMEDefs.h"
+#include "SlateMaterialBrush.h"
+#include "UObject/WeakObjectPtr.h"
+#include "UObject/WeakObjectPtrTemplates.h"
+#include "Widgets/SCompoundWidget.h"
+
+class SDMEditor;
+class UDMMaterialComponent;
+class UDMMaterialValueDynamic;
+class UDMTextureUV;
+class UDMTextureUVDynamic;
+class UDynamicMaterialModelBase;
+class UDynamicMaterialModelDynamic;
+
+class SDMComponentPreview : public SCompoundWidget
+{
+public:
+	SLATE_BEGIN_ARGS(SDMComponentPreview)
+		: _PreviewSize(FVector2D(48.f, 48.f))
+		{}
+		SLATE_ATTRIBUTE(FVector2D, PreviewSize)
+	SLATE_END_ARGS()
+
+	SDMComponentPreview();
+
+	virtual ~SDMComponentPreview() override;
+
+	void Construct(const FArguments& InArgs, const TSharedRef<SDMEditor>& InEditorWidget, UDMMaterialComponent* InComponent);
+
+	FSlateMaterialBrush& GetBrush() { return Brush; }
+
+	//~ Begin SWidget
+	virtual void Tick(const FGeometry& InAllottedGeometry, const double InCurrentTime, const float InDeltaTime) override;
+	//~ End SWidget
+
+protected:
+	TWeakPtr<SDMEditor> EditorWidgetWeak;
+	TWeakObjectPtr<UDMMaterialComponent> ComponentWeak;
+	TWeakObjectPtr<UDynamicMaterialModelBase> MaterialModelBaseWeak;
+	TWeakObjectPtr<UMaterial> PreviewMaterialBaseWeak;
+	TWeakObjectPtr<UMaterialInstanceDynamic> PreviewMaterialDynamicWeak;
+	FSlateMaterialBrush Brush;
+	TAttribute<FVector2D> PreviewSize;
+
+	void OnComponentUpdated(UDMMaterialComponent* InComponent, EDMUpdateType InUpdateType);
+
+	TOptional<FVector2D> GetPreviewSize() const;
+
+	void OnValueUpdated(UDynamicMaterialModel* InMaterialModel, UDMMaterialValue* InValue);
+
+	void OnTextureUVUpdated(UDynamicMaterialModel* InMaterialModel, UDMTextureUV* InTextureUV);
+
+	void OnValueDynamicUpdated(UDynamicMaterialModelDynamic* InMaterialModel, UDMMaterialValueDynamic* InValueDynamic);
+
+	void OnTextureUVDynamicUpdated(UDynamicMaterialModelDynamic* InMaterialModel, UDMTextureUVDynamic* InTextureUVDynamic);
+};
diff --git a/Engine/Plugins/Experimental/DynamicMaterial/Source/DynamicMaterialEditor/Private/Slate/Properties/Editors/SDMPropertyEditBool.h b/Engine/Plugins/Experimental/DynamicMaterial/Source/DynamicMaterialEditor/Private/Slate/Properties/Editors/SDMPropertyEditBool.h
--- a/Engine/Plugins/Experimental/DynamicMaterial/Source/DynamicMaterialEditor/Private/Slate/Properties/Editors/SDMPropertyEditBool.h
+++ b/Engine/Plugins/Experimental/DynamicMaterial/Source/DynamicMaterialEditor/Private/Slate/Properties/Editors/SDMPropertyEditBool.h
@@ -17,7 +17,6 @@ class SDMPropertyEditBool : public SDMPropertyEdit
 	SLATE_END_ARGS()
 
 public:
-	SDMPropertyEditBool() = default;
 	virtual ~SDMPropertyEditBool() override = default;
 
 	void Construct(const FArguments& InArgs, const TSharedPtr<IPropertyHandle>& InPropertyHandle);
diff --git a/Engine/Plugins/Experimental/DynamicMaterial/Source/DynamicMaterialEditor/Private/Slate/Properties/Editors/SDMPropertyEditBoolValue.h b/Engine/Plugins/Experimental/DynamicMaterial/Source/DynamicMaterialEditor/Private/Slate/Properties/Editors/SDMPropertyEditBoolValue.h
--- a/Engine/Plugins/Experimental/DynamicMaterial/Source/DynamicMaterialEditor/Private/Slate/Properties/Editors/SDMPropertyEditBoolValue.h
+++ b/Engine/Plugins/Experimental/DynamicMaterial/Source/DynamicMaterialEditor/Private/Slate/Properties/Editors/SDMPropertyEditBoolValue.h
@@ -20,7 +20,6 @@ class SDMPropertyEditBoolValue : public SDMPropertyEdit
 public:
 	static TSharedPtr<SWidget> CreateEditWidget(const TSharedPtr<SDMComponentEdit>& InComponentEditWidget, UDMMaterialValue* InBoolValue);
 
-	SDMPropertyEditBoolValue() = default;
 	virtual ~SDMPropertyEditBoolValue() override = default;
 
 	void Construct(const FArguments& InArgs, UDMMaterialValueBool* InBoolValue);
diff --git a/Engine/Plugins/Experimental/DynamicMaterial/Source/DynamicMaterialEditor/Private/Slate/Properties/Editors/SDMPropertyEditColor.h b/Engine/Plugins/Experimental/DynamicMaterial/Source/DynamicMaterialEditor/Private/Slate/Properties/Editors/SDMPropertyEditColor.h
--- a/Engine/Plugins/Experimental/DynamicMaterial/Source/DynamicMaterialEditor/Private/Slate/Properties/Editors/SDMPropertyEditColor.h
+++ b/Engine/Plugins/Experimental/DynamicMaterial/Source/DynamicMaterialEditor/Private/Slate/Properties/Editors/SDMPropertyEditColor.h
@@ -17,7 +17,6 @@ class SDMPropertyEditColor : public SDMPropertyEdit
 	SLATE_END_ARGS()
 
 public:
-	SDMPropertyEditColor() = default;
 	virtual ~SDMPropertyEditColor() override = default;
 
 	void Construct(const FArguments& InArgs, const TSharedPtr<IPropertyHandle>& InPropertyHandle);
diff --git a/Engine/Plugins/Experimental/DynamicMaterial/Source/DynamicMaterialEditor/Private/Slate/Properties/Editors/SDMPropertyEditEnum.h b/Engine/Plugins/Experimental/DynamicMaterial/Source/DynamicMaterialEditor/Private/Slate/Properties/Editors/SDMPropertyEditEnum.h
--- a/Engine/Plugins/Experimental/DynamicMaterial/Source/DynamicMaterialEditor/Private/Slate/Properties/Editors/SDMPropertyEditEnum.h
+++ b/Engine/Plugins/Experimental/DynamicMaterial/Source/DynamicMaterialEditor/Private/Slate/Properties/Editors/SDMPropertyEditEnum.h
@@ -17,7 +17,6 @@ class SDMPropertyEditEnum : public SDMPropertyEdit
 	SLATE_END_ARGS()
 
 public:
-	SDMPropertyEditEnum() = default;
 	virtual ~SDMPropertyEditEnum() override = default;
 
 	void Construct(const FArguments& InArgs, const TSharedPtr<IPropertyHandle>& InPropertyHandle);
diff --git a/Engine/Plugins/Experimental/DynamicMaterial/Source/DynamicMaterialEditor/Private/Slate/Properties/Editors/SDMPropertyEditFloat.h b/Engine/Plugins/Experimental/DynamicMaterial/Source/DynamicMaterialEditor/Private/Slate/Properties/Editors/SDMPropertyEditFloat.h
--- a/Engine/Plugins/Experimental/DynamicMaterial/Source/DynamicMaterialEditor/Private/Slate/Properties/Editors/SDMPropertyEditFloat.h
+++ b/Engine/Plugins/Experimental/DynamicMaterial/Source/DynamicMaterialEditor/Private/Slate/Properties/Editors/SDMPropertyEditFloat.h
@@ -18,7 +18,6 @@ class SDMPropertyEditFloat : public SDMPropertyEdit
 	SLATE_END_ARGS()
 
 public:
-	SDMPropertyEditFloat() = default;
 	virtual ~SDMPropertyEditFloat() override = default;
 
 	void Construct(const FArguments& InArgs, const TSharedPtr<IPropertyHandle>& InPropertyHandle);
diff --git a/Engine/Plugins/Experimental/DynamicMaterial/Source/DynamicMaterialEditor/Private/Slate/Properties/Editors/SDMPropertyEditFloat1Value.h b/Engine/Plugins/Experimental/DynamicMaterial/Source/DynamicMaterialEditor/Private/Slate/Properties/Editors/SDMPropertyEditFloat1Value.h
--- a/Engine/Plugins/Experimental/DynamicMaterial/Source/DynamicMaterialEditor/Private/Slate/Properties/Editors/SDMPropertyEditFloat1Value.h
+++ b/Engine/Plugins/Experimental/DynamicMaterial/Source/DynamicMaterialEditor/Private/Slate/Properties/Editors/SDMPropertyEditFloat1Value.h
@@ -17,7 +17,6 @@ class SDMPropertyEditFloat1Value : public SDMPropertyEdit
 public:
 	static TSharedPtr<SWidget> CreateEditWidget(const TSharedPtr<SDMComponentEdit>& InComponentEditWidget, UDMMaterialValue* InFloat1Value);
 
-	SDMPropertyEditFloat1Value() = default;
 	virtual ~SDMPropertyEditFloat1Value() override = default;
 
 	void Construct(const FArguments& InArgs, UDMMaterialValueFloat1* InFloat1Value);
diff --git a/Engine/Plugins/Experimental/DynamicMaterial/Source/DynamicMaterialEditor/Private/Slate/Properties/Editors/SDMPropertyEditFloat2Value.h b/Engine/Plugins/Experimental/DynamicMaterial/Source/DynamicMaterialEditor/Private/Slate/Properties/Editors/SDMPropertyEditFloat2Value.h
--- a/Engine/Plugins/Experimental/DynamicMaterial/Source/DynamicMaterialEditor/Private/Slate/Properties/Editors/SDMPropertyEditFloat2Value.h
+++ b/Engine/Plugins/Experimental/DynamicMaterial/Source/DynamicMaterialEditor/Private/Slate/Properties/Editors/SDMPropertyEditFloat2Value.h
@@ -17,7 +17,6 @@ class SDMPropertyEditFloat2Value : public SDMPropertyEdit
 public:
 	static TSharedPtr<SWidget> CreateEditWidget(const TSharedPtr<SDMComponentEdit>& InComponentEditWidget, UDMMaterialValue* InFloat2Value);
 
-	SDMPropertyEditFloat2Value() = default;
 	virtual ~SDMPropertyEditFloat2Value() override = default;
 
 	void Construct(const FArguments& InArgs, UDMMaterialValueFloat2* InFloat2Value);
