Material Designer

diff --git a/Engine/Plugins/Experimental/DynamicMaterial/Source/DynamicMaterialEditor/Private/Slate/Properties/Editors/SDMPropertyEditFloat3RGBValue.h b/Engine/Plugins/Experimental/DynamicMaterial/Source/DynamicMaterialEditor/Private/Slate/Properties/Editors/SDMPropertyEditFloat3RGBValue.h
--- a/Engine/Plugins/Experimental/DynamicMaterial/Source/DynamicMaterialEditor/Private/Slate/Properties/Editors/SDMPropertyEditFloat3RGBValue.h
+++ b/Engine/Plugins/Experimental/DynamicMaterial/Source/DynamicMaterialEditor/Private/Slate/Properties/Editors/SDMPropertyEditFloat3RGBValue.h
@@ -17,7 +17,6 @@ class SDMPropertyEditFloat3RGBValue : public SDMPropertyEdit
 public:
 	static TSharedPtr<SWidget> CreateEditWidget(const TSharedPtr<SDMComponentEdit>& InComponentEditWidget, UDMMaterialValue* InFloat3RGBValue);
 
-	SDMPropertyEditFloat3RGBValue() = default;
 	virtual ~SDMPropertyEditFloat3RGBValue() override = default;
 
 	void Construct(const FArguments& InArgs, UDMMaterialValueFloat3RGB* InFloat3RGBValue);
diff --git a/Engine/Plugins/Experimental/DynamicMaterial/Source/DynamicMaterialEditor/Private/Slate/Properties/Editors/SDMPropertyEditFloat3RPYValue.h b/Engine/Plugins/Experimental/DynamicMaterial/Source/DynamicMaterialEditor/Private/Slate/Properties/Editors/SDMPropertyEditFloat3RPYValue.h
--- a/Engine/Plugins/Experimental/DynamicMaterial/Source/DynamicMaterialEditor/Private/Slate/Properties/Editors/SDMPropertyEditFloat3RPYValue.h
+++ b/Engine/Plugins/Experimental/DynamicMaterial/Source/DynamicMaterialEditor/Private/Slate/Properties/Editors/SDMPropertyEditFloat3RPYValue.h
@@ -17,7 +17,6 @@ class SDMPropertyEditFloat3RPYValue : public SDMPropertyEdit
 public:
 	static TSharedPtr<SWidget> CreateEditWidget(const TSharedPtr<SDMComponentEdit>& InComponentEditWidget, UDMMaterialValue* InFloat3RPYValue);
 
-	SDMPropertyEditFloat3RPYValue() = default;
 	virtual ~SDMPropertyEditFloat3RPYValue() override = default;
 
 	void Construct(const FArguments& InArgs, UDMMaterialValueFloat3RPY* InFloat3RPYValue);
diff --git a/Engine/Plugins/Experimental/DynamicMaterial/Source/DynamicMaterialEditor/Private/Slate/Properties/Editors/SDMPropertyEditFloat3XYZValue.h b/Engine/Plugins/Experimental/DynamicMaterial/Source/DynamicMaterialEditor/Private/Slate/Properties/Editors/SDMPropertyEditFloat3XYZValue.h
--- a/Engine/Plugins/Experimental/DynamicMaterial/Source/DynamicMaterialEditor/Private/Slate/Properties/Editors/SDMPropertyEditFloat3XYZValue.h
+++ b/Engine/Plugins/Experimental/DynamicMaterial/Source/DynamicMaterialEditor/Private/Slate/Properties/Editors/SDMPropertyEditFloat3XYZValue.h
@@ -17,7 +17,6 @@ class SDMPropertyEditFloat3XYZValue : public SDMPropertyEdit
 public:
 	static TSharedPtr<SWidget> CreateEditWidget(const TSharedPtr<SDMComponentEdit>& InComponentEditWidget, UDMMaterialValue* InFloat3XYZValue);
 
-	SDMPropertyEditFloat3XYZValue() = default;
 	virtual ~SDMPropertyEditFloat3XYZValue() override = default;
 
 	void Construct(const FArguments& InArgs, UDMMaterialValueFloat3XYZ* InFloat3XYZValue);
diff --git a/Engine/Plugins/Experimental/DynamicMaterial/Source/DynamicMaterialEditor/Private/Slate/Properties/Editors/SDMPropertyEditFloat4Value.h b/Engine/Plugins/Experimental/DynamicMaterial/Source/DynamicMaterialEditor/Private/Slate/Properties/Editors/SDMPropertyEditFloat4Value.h
--- a/Engine/Plugins/Experimental/DynamicMaterial/Source/DynamicMaterialEditor/Private/Slate/Properties/Editors/SDMPropertyEditFloat4Value.h
+++ b/Engine/Plugins/Experimental/DynamicMaterial/Source/DynamicMaterialEditor/Private/Slate/Properties/Editors/SDMPropertyEditFloat4Value.h
@@ -17,7 +17,6 @@ class SDMPropertyEditFloat4Value : public SDMPropertyEdit
 public:
 	static TSharedPtr<SWidget> CreateEditWidget(const TSharedPtr<SDMComponentEdit>& InComponentEditWidget, UDMMaterialValue* InFloat4Value);
 
-	SDMPropertyEditFloat4Value() = default;
 	virtual ~SDMPropertyEditFloat4Value() override = default;
 
 	void Construct(const FArguments& InArgs, UDMMaterialValueFloat4* InFloat4Value);
diff --git a/Engine/Plugins/Experimental/DynamicMaterial/Source/DynamicMaterialEditor/Private/Slate/Properties/Editors/SDMPropertyEditObject.h b/Engine/Plugins/Experimental/DynamicMaterial/Source/DynamicMaterialEditor/Private/Slate/Properties/Editors/SDMPropertyEditObject.h
--- a/Engine/Plugins/Experimental/DynamicMaterial/Source/DynamicMaterialEditor/Private/Slate/Properties/Editors/SDMPropertyEditObject.h
+++ b/Engine/Plugins/Experimental/DynamicMaterial/Source/DynamicMaterialEditor/Private/Slate/Properties/Editors/SDMPropertyEditObject.h
@@ -21,7 +21,6 @@ class SDMPropertyEditObject : public SDMPropertyEdit
 	SLATE_END_ARGS()
 
 public:
-	SDMPropertyEditObject() = default;
 	virtual ~SDMPropertyEditObject() override = default;
 
 	void Construct(const FArguments& InArgs, const TSharedPtr<IPropertyHandle>& InPropertyHandle, UClass* InAllowedClass);
diff --git a/Engine/Plugins/Experimental/DynamicMaterial/Source/DynamicMaterialEditor/Private/Slate/Properties/Editors/SDMPropertyEditTextureValue.h b/Engine/Plugins/Experimental/DynamicMaterial/Source/DynamicMaterialEditor/Private/Slate/Properties/Editors/SDMPropertyEditTextureValue.h
--- a/Engine/Plugins/Experimental/DynamicMaterial/Source/DynamicMaterialEditor/Private/Slate/Properties/Editors/SDMPropertyEditTextureValue.h
+++ b/Engine/Plugins/Experimental/DynamicMaterial/Source/DynamicMaterialEditor/Private/Slate/Properties/Editors/SDMPropertyEditTextureValue.h
@@ -19,7 +19,6 @@ class SDMPropertyEditTextureValue : public SDMPropertyEdit
 public:
 	static TSharedPtr<SWidget> CreateEditWidget(const TSharedPtr<SDMComponentEdit>& InComponentEditWidget, UDMMaterialValue* InTextureValue);
 
-	SDMPropertyEditTextureValue() = default;
 	virtual ~SDMPropertyEditTextureValue() override = default;
 
 	void Construct(const FArguments& InArgs, UDMMaterialValueTexture* InTextureValue);
diff --git a/Engine/Plugins/Experimental/DynamicMaterial/Source/DynamicMaterialEditor/Private/Slate/Properties/Editors/SDMPropertyEditVector.h b/Engine/Plugins/Experimental/DynamicMaterial/Source/DynamicMaterialEditor/Private/Slate/Properties/Editors/SDMPropertyEditVector.h
--- a/Engine/Plugins/Experimental/DynamicMaterial/Source/DynamicMaterialEditor/Private/Slate/Properties/Editors/SDMPropertyEditVector.h
+++ b/Engine/Plugins/Experimental/DynamicMaterial/Source/DynamicMaterialEditor/Private/Slate/Properties/Editors/SDMPropertyEditVector.h
@@ -16,7 +16,6 @@ class SDMPropertyEditVector : public SDMPropertyEdit
 	SLATE_END_ARGS()
 
 public:
-	SDMPropertyEditVector() = default;
 	virtual ~SDMPropertyEditVector() override = default;
 
 	void Construct(const FArguments& InArgs, const TSharedPtr<IPropertyHandle>& InPropertyHandle, int32 InComponentCount);
diff --git a/Engine/Plugins/Experimental/DynamicMaterial/Source/DynamicMaterialEditor/Private/Slate/Properties/Generators/DMComponentPropertyRowGenerator.cpp b/Engine/Plugins/Experimental/DynamicMaterial/Source/DynamicMaterialEditor/Private/Slate/Properties/Generators/DMComponentPropertyRowGenerator.cpp
--- a/Engine/Plugins/Experimental/DynamicMaterial/Source/DynamicMaterialEditor/Private/Slate/Properties/Generators/DMComponentPropertyRowGenerator.cpp
+++ b/Engine/Plugins/Experimental/DynamicMaterial/Source/DynamicMaterialEditor/Private/Slate/Properties/Generators/DMComponentPropertyRowGenerator.cpp
@@ -1,11 +1,12 @@
 // Copyright Epic Games, Inc. All Rights Reserved.
 
 #include "Slate/Properties/Generators/DMComponentPropertyRowGenerator.h"
-#include "DynamicMaterialEditorModule.h"
 #include "Components/DMMaterialComponent.h"
+#include "DynamicMaterialEditorModule.h"
+#include "Model/DynamicMaterialModelDynamic.h"
 #include "Runtime/Launch/Resources/Version.h"
-#include "Slate/SDMEditor.h"
 #include "Slate/SDMComponentEdit.h"
+#include "Slate/SDMEditor.h"
 
 const TSharedRef<FDMComponentPropertyRowGenerator>& FDMComponentPropertyRowGenerator::Get()
 {
@@ -96,10 +97,24 @@ void FDMComponentPropertyRowGenerator::AddPropertyEditRows(const TSharedRef<SDMC
 		}
 	}
 
-	InOutPropertyRows.Add(SDMEditor::GetPropertyHandle(&*InComponentEditWidget, InComponent, InProperty->GetFName()));
+	FDMPropertyHandle& Handle = InOutPropertyRows.Add_GetRef(SDMEditor::GetPropertyHandle(&*InComponentEditWidget, InComponent, InProperty->GetFName()));
+	Handle.bEnabled = !IsDynamic(InComponentEditWidget);
 }
 
 bool FDMComponentPropertyRowGenerator::AllowKeyframeButton(UDMMaterialComponent* InComponent, FProperty* InProperty)
 {
 	return false;
 }
+
+bool FDMComponentPropertyRowGenerator::IsDynamic(const TSharedRef<SDMComponentEdit>& InComponentEditWidget)
+{
+	if (TSharedPtr<SDMEditor> EditorWidget = InComponentEditWidget->GetEditorWidget())
+	{
+		if (UDynamicMaterialModelBase* MaterialModelBase = EditorWidget->GetMaterialModelBase())
+		{
+			return !!Cast<UDynamicMaterialModelDynamic>(MaterialModelBase);
+		}
+	}
+
+	return false;
+}
diff --git a/Engine/Plugins/Experimental/DynamicMaterial/Source/DynamicMaterialEditor/Private/Slate/Properties/Generators/DMInputThroughputPropertyRowGenerator.h b/Engine/Plugins/Experimental/DynamicMaterial/Source/DynamicMaterialEditor/Private/Slate/Properties/Generators/DMInputThroughputPropertyRowGenerator.h
--- a/Engine/Plugins/Experimental/DynamicMaterial/Source/DynamicMaterialEditor/Private/Slate/Properties/Generators/DMInputThroughputPropertyRowGenerator.h
+++ b/Engine/Plugins/Experimental/DynamicMaterial/Source/DynamicMaterialEditor/Private/Slate/Properties/Generators/DMInputThroughputPropertyRowGenerator.h
@@ -9,7 +9,6 @@ class FDMInputThroughputPropertyRowGenerator : public FDMThroughputPropertyRowGe
 public:
 	static const TSharedRef<FDMInputThroughputPropertyRowGenerator>& Get();
 
-	FDMInputThroughputPropertyRowGenerator() = default;
 	virtual ~FDMInputThroughputPropertyRowGenerator() override = default;
 
 	virtual void AddComponentProperties(const TSharedRef<SDMComponentEdit>& InComponentEditWidget, UDMMaterialComponent* InComponent,
diff --git a/Engine/Plugins/Experimental/DynamicMaterial/Source/DynamicMaterialEditor/Private/Slate/Properties/Generators/DMMaterialEffectFunctionPropertyRowGenerator.h b/Engine/Plugins/Experimental/DynamicMaterial/Source/DynamicMaterialEditor/Private/Slate/Properties/Generators/DMMaterialEffectFunctionPropertyRowGenerator.h
--- a/Engine/Plugins/Experimental/DynamicMaterial/Source/DynamicMaterialEditor/Private/Slate/Properties/Generators/DMMaterialEffectFunctionPropertyRowGenerator.h
+++ b/Engine/Plugins/Experimental/DynamicMaterial/Source/DynamicMaterialEditor/Private/Slate/Properties/Generators/DMMaterialEffectFunctionPropertyRowGenerator.h
@@ -9,7 +9,6 @@ class FDMMaterialEffectFunctionPropertyRowGenerator : public FDMComponentPropert
 public:
 	static const TSharedRef<FDMMaterialEffectFunctionPropertyRowGenerator>& Get();
 
-	FDMMaterialEffectFunctionPropertyRowGenerator() = default;
 	virtual ~FDMMaterialEffectFunctionPropertyRowGenerator() override = default;
 
 	virtual void AddComponentProperties(const TSharedRef<SDMComponentEdit>& InComponentEditWidget, UDMMaterialComponent* InComponent,
diff --git a/Engine/Plugins/Experimental/DynamicMaterial/Source/DynamicMaterialEditor/Private/Slate/Properties/Generators/DMMaterialStageFunctionPropertyRowGenerator.h b/Engine/Plugins/Experimental/DynamicMaterial/Source/DynamicMaterialEditor/Private/Slate/Properties/Generators/DMMaterialStageFunctionPropertyRowGenerator.h
--- a/Engine/Plugins/Experimental/DynamicMaterial/Source/DynamicMaterialEditor/Private/Slate/Properties/Generators/DMMaterialStageFunctionPropertyRowGenerator.h
+++ b/Engine/Plugins/Experimental/DynamicMaterial/Source/DynamicMaterialEditor/Private/Slate/Properties/Generators/DMMaterialStageFunctionPropertyRowGenerator.h
@@ -9,7 +9,6 @@ class FDMMaterialStageFunctionPropertyRowGenerator : public FDMComponentProperty
 public:
 	static const TSharedRef<FDMMaterialStageFunctionPropertyRowGenerator>& Get();
 
-	FDMMaterialStageFunctionPropertyRowGenerator() = default;
 	virtual ~FDMMaterialStageFunctionPropertyRowGenerator() override = default;
 
 	virtual void AddComponentProperties(const TSharedRef<SDMComponentEdit>& InComponentEditWidget, UDMMaterialComponent* InComponent,
diff --git a/Engine/Plugins/Experimental/DynamicMaterial/Source/DynamicMaterialEditor/Private/Slate/Properties/Generators/DMMaterialValueDynamicPropertyRowGenerator.cpp b/Engine/Plugins/Experimental/DynamicMaterial/Source/DynamicMaterialEditor/Private/Slate/Properties/Generators/DMMaterialValueDynamicPropertyRowGenerator.cpp
new file mode 100644
--- /dev/null
+++ b/Engine/Plugins/Experimental/DynamicMaterial/Source/DynamicMaterialEditor/Private/Slate/Properties/Generators/DMMaterialValueDynamicPropertyRowGenerator.cpp
@@ -0,0 +1,88 @@
+// Copyright Epic Games, Inc. All Rights Reserved.
+
+#include "Slate/Properties/Generators/DMMaterialValueDynamicPropertyRowGenerator.h"
+#include "IDetailPropertyRow.h"
+#include "Components/DMMaterialComponent.h"
+#include "Components/DMMaterialValue.h"
+#include "Components/DMMaterialValueDynamic.h"
+#include "Slate/SDMEditor.h"
+#include "Slate/SDMComponentEdit.h"
+
+const TSharedRef<FDMMaterialValueDynamicPropertyRowGenerator>& FDMMaterialValueDynamicPropertyRowGenerator::Get()
+{
+	static TSharedRef<FDMMaterialValueDynamicPropertyRowGenerator> Generator = MakeShared<FDMMaterialValueDynamicPropertyRowGenerator>();
+	return Generator;
+}
+
+void FDMMaterialValueDynamicPropertyRowGenerator::AddComponentProperties(const TSharedRef<SDMComponentEdit>& InComponentEditWidget, UDMMaterialComponent* InComponent,
+	TArray<FDMPropertyHandle>& InOutPropertyRows, TSet<UDMMaterialComponent*>& InOutProcessedObjects)
+{
+	if (!IsValid(InComponent))
+	{
+		return;
+	}
+
+	if (InOutProcessedObjects.Contains(InComponent))
+	{
+		return;
+	}
+
+	UDMMaterialValueDynamic* ValueDynamic = Cast<UDMMaterialValueDynamic>(InComponent);
+
+	if (!ValueDynamic)
+	{
+		return;
+	}
+
+	// The base material value class is abstract and not allowed.
+	if (ValueDynamic->GetClass() == UDMMaterialValueDynamic::StaticClass())
+	{
+		return;
+	}
+
+	UDMMaterialValue* ParentValue = ValueDynamic->GetParentValue();
+
+	if (!ParentValue)
+	{
+		return;
+	}
+
+	InOutProcessedObjects.Add(InComponent);
+
+	if (ParentValue->AllowEditValue())
+	{
+		FDMPropertyHandle Handle = SDMEditor::GetPropertyHandle(&*InComponentEditWidget, ValueDynamic, UDMMaterialValue::ValueName);
+
+		Handle.ResetToDefaultOverride = FResetToDefaultOverride::Create(
+			FIsResetToDefaultVisible::CreateUObject(ValueDynamic, &UDMMaterialValueDynamic::CanResetToDefault),
+			FResetToDefaultHandler::CreateUObject(ValueDynamic, &UDMMaterialValueDynamic::ResetToDefault)
+		);
+
+		Handle.bEnabled = true;
+
+		InOutPropertyRows.Add(Handle);
+	}
+
+	const TArray<FName>& Properties = ParentValue->GetEditableProperties();
+	const int32 StartRow = InOutPropertyRows.Num();
+
+	for (const FName& Property : Properties)
+	{
+		if (Property == UDMMaterialValue::ValueName)
+		{
+			continue;
+		}
+			
+		if (InComponent->IsPropertyVisible(Property))
+		{
+			AddPropertyEditRows(InComponentEditWidget, InComponent, Property, InOutPropertyRows, InOutProcessedObjects);
+		}
+	}
+
+	const int32 EndRow = InOutPropertyRows.Num();
+
+	for (int32 RowIndex = StartRow; RowIndex < EndRow; ++RowIndex)
+	{
+		InOutPropertyRows[RowIndex].bEnabled = false;
+	}
+}
diff --git a/Engine/Plugins/Experimental/DynamicMaterial/Source/DynamicMaterialEditor/Private/Slate/Properties/Generators/DMMaterialValueDynamicPropertyRowGenerator.h b/Engine/Plugins/Experimental/DynamicMaterial/Source/DynamicMaterialEditor/Private/Slate/Properties/Generators/DMMaterialValueDynamicPropertyRowGenerator.h
new file mode 100644
--- /dev/null
+++ b/Engine/Plugins/Experimental/DynamicMaterial/Source/DynamicMaterialEditor/Private/Slate/Properties/Generators/DMMaterialValueDynamicPropertyRowGenerator.h
@@ -0,0 +1,16 @@
+// Copyright Epic Games, Inc. All Rights Reserved.
+
+#pragma once
+
+#include "Slate/Properties/Generators/DMComponentPropertyRowGenerator.h"
+
+class FDMMaterialValueDynamicPropertyRowGenerator : public FDMComponentPropertyRowGenerator
+{
+public:
+	static const TSharedRef<FDMMaterialValueDynamicPropertyRowGenerator>& Get();
+
+	virtual ~FDMMaterialValueDynamicPropertyRowGenerator() override = default;
+
+	virtual void AddComponentProperties(const TSharedRef<SDMComponentEdit>& InComponentEditWidget, UDMMaterialComponent* InComponent,
+		TArray<FDMPropertyHandle>& InOutPropertyRows, TSet<UDMMaterialComponent*>& InOutProcessedObjects) override;
+};
diff --git a/Engine/Plugins/Experimental/DynamicMaterial/Source/DynamicMaterialEditor/Private/Slate/Properties/Generators/DMMaterialValuePropertyRowGenerator.cpp b/Engine/Plugins/Experimental/DynamicMaterial/Source/DynamicMaterialEditor/Private/Slate/Properties/Generators/DMMaterialValuePropertyRowGenerator.cpp
--- a/Engine/Plugins/Experimental/DynamicMaterial/Source/DynamicMaterialEditor/Private/Slate/Properties/Generators/DMMaterialValuePropertyRowGenerator.cpp
+++ b/Engine/Plugins/Experimental/DynamicMaterial/Source/DynamicMaterialEditor/Private/Slate/Properties/Generators/DMMaterialValuePropertyRowGenerator.cpp
@@ -1,13 +1,17 @@
 // Copyright Epic Games, Inc. All Rights Reserved.
 
 #include "Slate/Properties/Generators/DMMaterialValuePropertyRowGenerator.h"
+
 #include "DynamicMaterialEditorModule.h"
-#include "IDetailPropertyRow.h"
 #include "Components/DMMaterialComponent.h"
 #include "Components/DMMaterialValue.h"
-#include "Components/MaterialValues/DMMaterialValueFloat1.h"
-#include "Slate/SDMEditor.h"
+#include "Components/DMMaterialValueDynamic.h"
+#include "IDetailPropertyRow.h"
+#include "Model/DynamicMaterialModel.h"
+#include "Model/DynamicMaterialModelBase.h"
+#include "Model/DynamicMaterialModelDynamic.h"
 #include "Slate/SDMComponentEdit.h"
+#include "Slate/SDMEditor.h"
 
 const TSharedRef<FDMMaterialValuePropertyRowGenerator>& FDMMaterialValuePropertyRowGenerator::Get()
 {
@@ -43,6 +47,22 @@ void FDMMaterialValuePropertyRowGenerator::AddComponentProperties(const TSharedR
 
 	InOutProcessedObjects.Add(InComponent);
 
+	if (TSharedPtr<SDMEditor> EditorWidget = InComponentEditWidget->GetEditorWidget())
+	{
+		if (UDynamicMaterialModelBase* MaterialModelBase = EditorWidget->GetMaterialModelBase())
+		{
+			if (UDynamicMaterialModelDynamic* MaterialModelDynamic = Cast<UDynamicMaterialModelDynamic>(MaterialModelBase))
+			{
+				if (UDMMaterialComponentDynamic* ComponentDynamic = MaterialModelDynamic->GetComponentDynamic(Value->GetFName()))
+				{
+					FDynamicMaterialEditorModule::Get().GeneratorComponentPropertyRows(InComponentEditWidget, ComponentDynamic, InOutPropertyRows, InOutProcessedObjects);
+				}
+
+				return;
+			}
+		}
+	}
+
 	if (Value->AllowEditValue())
 	{
 		FDMPropertyHandle Handle = SDMEditor::GetPropertyHandle(&*InComponentEditWidget, Value, UDMMaterialValue::ValueName);
@@ -52,6 +72,8 @@ void FDMMaterialValuePropertyRowGenerator::AddComponentProperties(const TSharedR
 			FResetToDefaultHandler::CreateUObject(Value, &UDMMaterialValue::ResetToDefault)
 		);
 
+		Handle.bEnabled = true;
+
 		InOutPropertyRows.Add(Handle);
 	}
 
diff --git a/Engine/Plugins/Experimental/DynamicMaterial/Source/DynamicMaterialEditor/Private/Slate/Properties/Generators/DMMaterialValuePropertyRowGenerator.h b/Engine/Plugins/Experimental/DynamicMaterial/Source/DynamicMaterialEditor/Private/Slate/Properties/Generators/DMMaterialValuePropertyRowGenerator.h
--- a/Engine/Plugins/Experimental/DynamicMaterial/Source/DynamicMaterialEditor/Private/Slate/Properties/Generators/DMMaterialValuePropertyRowGenerator.h
+++ b/Engine/Plugins/Experimental/DynamicMaterial/Source/DynamicMaterialEditor/Private/Slate/Properties/Generators/DMMaterialValuePropertyRowGenerator.h
@@ -9,7 +9,6 @@ class FDMMaterialValuePropertyRowGenerator : public FDMComponentPropertyRowGener
 public:
 	static const TSharedRef<FDMMaterialValuePropertyRowGenerator>& Get();
 
-	FDMMaterialValuePropertyRowGenerator() = default;
 	virtual ~FDMMaterialValuePropertyRowGenerator() override = default;
 
 	virtual void AddComponentProperties(const TSharedRef<SDMComponentEdit>& InComponentEditWidget, UDMMaterialComponent* InComponent,
diff --git a/Engine/Plugins/Experimental/DynamicMaterial/Source/DynamicMaterialEditor/Private/Slate/Properties/Generators/DMStagePropertyRowGenerator.h b/Engine/Plugins/Experimental/DynamicMaterial/Source/DynamicMaterialEditor/Private/Slate/Properties/Generators/DMStagePropertyRowGenerator.h
--- a/Engine/Plugins/Experimental/DynamicMaterial/Source/DynamicMaterialEditor/Private/Slate/Properties/Generators/DMStagePropertyRowGenerator.h
+++ b/Engine/Plugins/Experimental/DynamicMaterial/Source/DynamicMaterialEditor/Private/Slate/Properties/Generators/DMStagePropertyRowGenerator.h
@@ -9,7 +9,6 @@ class FDMStagePropertyRowGenerator : public FDMComponentPropertyRowGenerator
 public:
 	static const TSharedRef<FDMStagePropertyRowGenerator>& Get();
 
-	FDMStagePropertyRowGenerator() = default;
 	virtual ~FDMStagePropertyRowGenerator() override = default;
 
 	virtual void AddComponentProperties(const TSharedRef<SDMComponentEdit>& InComponentEditWidget, UDMMaterialComponent* InComponent,
diff --git a/Engine/Plugins/Experimental/DynamicMaterial/Source/DynamicMaterialEditor/Private/Slate/Properties/Generators/DMTextureUVDynamicPropertyRowGenerator.cpp b/Engine/Plugins/Experimental/DynamicMaterial/Source/DynamicMaterialEditor/Private/Slate/Properties/Generators/DMTextureUVDynamicPropertyRowGenerator.cpp
new file mode 100644
--- /dev/null
+++ b/Engine/Plugins/Experimental/DynamicMaterial/Source/DynamicMaterialEditor/Private/Slate/Properties/Generators/DMTextureUVDynamicPropertyRowGenerator.cpp
@@ -0,0 +1,244 @@
+// Copyright Epic Games, Inc. All Rights Reserved.
+
+#include "Slate/Properties/Generators/DMTextureUVDynamicPropertyRowGenerator.h"
+#include "Components/DMMaterialComponent.h"
+#include "Components/DMTextureUV.h"
+#include "Components/DMTextureUVDynamic.h"
+#include "DMEDefs.h"
+#include "Modules/ModuleManager.h"
+#include "PropertyCustomizationHelpers.h"
+#include "PropertyEditorModule.h"
+#include "Slate/Properties/SDMTextureUVVisualizerProperty.h"
+#include "Slate/SDMComponentEdit.h"
+#include "Slate/SDMEditor.h"
+#include "Widgets/Layout/SBox.h"
+
+#define LOCTEXT_NAMESPACE "DMTextureUVDynamicPropertyRowGenerator"
+
+const TSharedRef<FDMTextureUVDynamicPropertyRowGenerator>& FDMTextureUVDynamicPropertyRowGenerator::Get()
+{
+	static TSharedRef<FDMTextureUVDynamicPropertyRowGenerator> Generator = MakeShared<FDMTextureUVDynamicPropertyRowGenerator>();
+	return Generator;
+}
+
+namespace UE::DynamicMaterialEditor::Private
+{
+	void AddTextureUVDynamicPropertyRow(const TSharedRef<SWidget>& InComponentEditWidget, UDMMaterialComponent* InTextureUVDynamic, 
+		FName InProperty, bool bInEnabled, TArray<FDMPropertyHandle>& InOutPropertyRows);
+
+	void AddTextureUVDynamicVisualizerRow(const TSharedRef<SWidget>& InComponentEditWidget, UDMTextureUVDynamic* InTextureUVDynamic, 
+		TArray<FDMPropertyHandle>& InOutPropertyRows);
+
+	bool CanResetTextureUVDynamicDynamicPropertyToDefault(TSharedPtr<IPropertyHandle> InPropertyHandle);
+
+	void ResetTextureUVDynamicDynamicPropertyToDefault(TSharedPtr<IPropertyHandle> InPropertyHandle);
+}
+
+void FDMTextureUVDynamicPropertyRowGenerator::AddComponentProperties(const TSharedRef<SDMComponentEdit>& InComponentEditWidget, 
+	UDMMaterialComponent* InComponent, TArray<FDMPropertyHandle>& InOutPropertyRows, TSet<UDMMaterialComponent*>& InOutProcessedObjects)
+{
+	if (!IsValid(InComponent))
+	{
+		return;
+	}
+
+	if (InOutProcessedObjects.Contains(InComponent))
+	{
+		return;
+	}
+
+	UDMTextureUVDynamic* TextureUVDynamic = Cast<UDMTextureUVDynamic>(InComponent);
+
+	if (!TextureUVDynamic)
+	{
+		return;
+	}
+
+	UDMTextureUV* TextureUV = TextureUVDynamic->GetParentTextureUV();
+
+	if (!TextureUV)
+	{
+		return;
+	}
+
+	InOutProcessedObjects.Add(InComponent);
+
+	using namespace UE::DynamicMaterialEditor::Private;
+
+	AddTextureUVDynamicPropertyRow(InComponentEditWidget, TextureUVDynamic, UDMTextureUV::NAME_Offset, /* Enabled */ true, InOutPropertyRows);
+	AddTextureUVDynamicPropertyRow(InComponentEditWidget, TextureUVDynamic, UDMTextureUV::NAME_Rotation, /* Enabled */ true, InOutPropertyRows);
+	AddTextureUVDynamicPropertyRow(InComponentEditWidget, TextureUVDynamic, UDMTextureUV::NAME_Tiling, /* Enabled */ true, InOutPropertyRows);
+	AddTextureUVDynamicPropertyRow(InComponentEditWidget, TextureUVDynamic, UDMTextureUV::NAME_Pivot, /* Enabled */ true, InOutPropertyRows);
+	AddTextureUVDynamicPropertyRow(InComponentEditWidget, TextureUV, UDMTextureUV::NAME_bMirrorOnX, /* Enabled */ false, InOutPropertyRows);
+	AddTextureUVDynamicPropertyRow(InComponentEditWidget, TextureUV, UDMTextureUV::NAME_bMirrorOnY, /* Enabled */ false, InOutPropertyRows);
+	AddTextureUVDynamicVisualizerRow(InComponentEditWidget, TextureUVDynamic, InOutPropertyRows);
+}
+
+void FDMTextureUVDynamicPropertyRowGenerator::AddPopoutComponentProperties(const TSharedRef<SWidget>& InParentWidget, UDMMaterialComponent* InComponent, 
+	TArray<FDMPropertyHandle>& InOutPropertyRows)
+{
+	if (!IsValid(InComponent))
+	{
+		return;
+	}
+
+	UDMTextureUVDynamic* TextureUVDynamic = Cast<UDMTextureUVDynamic>(InComponent);
+
+	if (!TextureUVDynamic)
+	{
+		return;
+	}
+
+	UDMTextureUV* TextureUV = TextureUVDynamic->GetParentTextureUV();
+
+	if (!TextureUV)
+	{
+		return;
+	}
+
+	using namespace UE::DynamicMaterialEditor::Private;
+
+	AddTextureUVDynamicPropertyRow(InParentWidget, TextureUVDynamic, UDMTextureUV::NAME_Offset, /* Enabled */ true, InOutPropertyRows);
+	AddTextureUVDynamicPropertyRow(InParentWidget, TextureUVDynamic, UDMTextureUV::NAME_Rotation, /* Enabled */ true, InOutPropertyRows);
+	AddTextureUVDynamicPropertyRow(InParentWidget, TextureUVDynamic, UDMTextureUV::NAME_Tiling, /* Enabled */ true, InOutPropertyRows);
+	AddTextureUVDynamicPropertyRow(InParentWidget, TextureUVDynamic, UDMTextureUV::NAME_Pivot, /* Enabled */ true, InOutPropertyRows);
+	AddTextureUVDynamicPropertyRow(InParentWidget, TextureUV, UDMTextureUV::NAME_bMirrorOnX, /* Enabled */ false, InOutPropertyRows);
+	AddTextureUVDynamicPropertyRow(InParentWidget, TextureUV, UDMTextureUV::NAME_bMirrorOnY, /* Enabled */ false, InOutPropertyRows);
+}
+
+bool FDMTextureUVDynamicPropertyRowGenerator::AllowKeyframeButton(UDMMaterialComponent* InComponent, FProperty* InProperty)
+{
+	if (InProperty)
+	{
+		const bool* AddKeyframeButtonPtr = UDMTextureUV::TextureProperties.Find(InProperty->GetFName());
+
+		if (AddKeyframeButtonPtr)
+		{
+			return *AddKeyframeButtonPtr;
+		}
+	}
+
+	return FDMComponentPropertyRowGenerator::AllowKeyframeButton(InComponent, InProperty);
+}
+
+void UE::DynamicMaterialEditor::Private::AddTextureUVDynamicPropertyRow(const TSharedRef<SWidget>& InComponentEditWidget, 
+	UDMMaterialComponent* InTextureUVDynamic, FName InProperty, bool bInEnabled, TArray<FDMPropertyHandle>& InOutPropertyRows)
+{
+	FDMPropertyHandle& NewHandle = InOutPropertyRows.Add_GetRef(SDMEditor::GetPropertyHandle(&*InComponentEditWidget, InTextureUVDynamic, InProperty));
+
+	NewHandle.ResetToDefaultOverride = FResetToDefaultOverride::Create(
+		FIsResetToDefaultVisible::CreateStatic(&UE::DynamicMaterialEditor::Private::CanResetTextureUVDynamicDynamicPropertyToDefault),
+		FResetToDefaultHandler::CreateStatic(&UE::DynamicMaterialEditor::Private::ResetTextureUVDynamicDynamicPropertyToDefault)
+	);
+
+	NewHandle.bEnabled = bInEnabled;
+}
+
+void UE::DynamicMaterialEditor::Private::AddTextureUVDynamicVisualizerRow(const TSharedRef<SWidget>& InComponentEditWidget, UDMTextureUVDynamic* InTextureUVDynamic,
+	TArray<FDMPropertyHandle>& InOutPropertyRows)
+{
+	UDMTextureUV* TextureUV = InTextureUVDynamic->GetParentTextureUV();
+
+	if (!TextureUV)
+	{
+		return;
+	}
+
+	// Make sure we don't get a substage
+	UDMMaterialStage* Stage = TextureUV->GetTypedParent<UDMMaterialStage>(/* Allow Subclasses */ false);
+
+	if (!Stage)
+	{
+		return;
+	}
+
+	if (InComponentEditWidget->GetWidgetClass().GetWidgetType() != SDMComponentEdit::StaticWidgetClass().GetWidgetType())
+	{
+		return;
+	}
+
+	TSharedPtr<SDMEditor> EditorWidget = StaticCastSharedRef<SDMComponentEdit>(InComponentEditWidget)->GetEditorWidget();
+
+	if (!EditorWidget.IsValid())
+	{
+		return;
+	}
+
+	FDMPropertyHandle VisualizerHandle;
+	VisualizerHandle.NameOverride = LOCTEXT("Visualizer", "UV Visualizer");
+	VisualizerHandle.NameToolTipOverride = LOCTEXT("VisualizerToolTip", "A graphical Texture UV editor.\n\n- Offset Mode: Change the Texture UV offset.\n- Pivot Mode: Change the Texture UV pivot, rotation and tiling.\n\nControl+click to reset values to default.");
+	VisualizerHandle.ValueName = FName(*InTextureUVDynamic->GetComponentPath());
+	VisualizerHandle.ValueWidget = SNew(SDMTextureUVVisualizerProperty, EditorWidget.ToSharedRef(), Stage).TextureUVDynamic(InTextureUVDynamic);
+	VisualizerHandle.CategoryOverrideName = TEXT("Texture UV");
+	VisualizerHandle.bEnabled = true;
+	InOutPropertyRows.Add(VisualizerHandle);
+}
+
+
+bool UE::DynamicMaterialEditor::Private::CanResetTextureUVDynamicDynamicPropertyToDefault(TSharedPtr<IPropertyHandle> InPropertyHandle)
+{
+	FProperty* Property = InPropertyHandle->GetProperty();
+
+	if (!Property)
+	{
+		return false;
+	}
+
+	const FName PropertyName = Property->GetFName();
+
+	if (PropertyName.IsNone())
+	{
+		return false;
+	}
+
+	TArray<UObject*> Outers;
+	InPropertyHandle->GetOuterObjects(Outers);
+
+	if (Outers.IsEmpty())
+	{
+		return false;
+	}
+
+	const UDMTextureUVDynamic* PropertyObject = Cast<UDMTextureUVDynamic>(Outers[0]);
+
+	if (!PropertyObject)
+	{
+		return false;
+	}
+
+	const UDMTextureUVDynamic* DefaultObject = GetDefault<UDMTextureUVDynamic>();
+
+	if (!DefaultObject)
+	{
+		return false;
+	}
+
+	if (PropertyName == UDMTextureUV::NAME_Offset)
+	{
+		return !DefaultObject->GetOffset().Equals(PropertyObject->GetOffset());
+	}
+
+	if (PropertyName == UDMTextureUV::NAME_Pivot)
+	{
+		return !DefaultObject->GetPivot().Equals(PropertyObject->GetPivot());
+	}
+
+	if (PropertyName == NAME_Rotation)
+	{
+		return DefaultObject->GetRotation() != PropertyObject->GetRotation();
+	}
+
+	if (PropertyName == UDMTextureUV::NAME_Tiling)
+	{
+		return !DefaultObject->GetTiling().Equals(PropertyObject->GetTiling());
+	}
+
+	return false;
+}
+
+void UE::DynamicMaterialEditor::Private::ResetTextureUVDynamicDynamicPropertyToDefault(TSharedPtr<IPropertyHandle> InPropertyHandle)
+{
+	InPropertyHandle->ResetToDefault();
+}
+
+#undef LOCTEXT_NAMESPACE
diff --git a/Engine/Plugins/Experimental/DynamicMaterial/Source/DynamicMaterialEditor/Private/Slate/Properties/Generators/DMTextureUVDynamicPropertyRowGenerator.h b/Engine/Plugins/Experimental/DynamicMaterial/Source/DynamicMaterialEditor/Private/Slate/Properties/Generators/DMTextureUVDynamicPropertyRowGenerator.h
new file mode 100644
--- /dev/null
+++ b/Engine/Plugins/Experimental/DynamicMaterial/Source/DynamicMaterialEditor/Private/Slate/Properties/Generators/DMTextureUVDynamicPropertyRowGenerator.h
@@ -0,0 +1,23 @@
+// Copyright Epic Games, Inc. All Rights Reserved.
+
+#pragma once
+
+#include "Slate/Properties/Generators/DMComponentPropertyRowGenerator.h"
+
+class SWidget;
+
+class FDMTextureUVDynamicPropertyRowGenerator : public FDMComponentPropertyRowGenerator
+{
+public:
+	static const TSharedRef<FDMTextureUVDynamicPropertyRowGenerator>& Get();
+
+	virtual ~FDMTextureUVDynamicPropertyRowGenerator() override = default;
+
+	virtual void AddComponentProperties(const TSharedRef<SDMComponentEdit>& InComponentEditWidget, UDMMaterialComponent* InComponent,
+		TArray<FDMPropertyHandle>& InOutPropertyRows, TSet<UDMMaterialComponent*>& InOutProcessedObjects) override;
+
+	static void AddPopoutComponentProperties(const TSharedRef<SWidget>& InParentWidget, UDMMaterialComponent* InComponent,
+		TArray<FDMPropertyHandle>& InOutPropertyRows);
+
+	virtual bool AllowKeyframeButton(UDMMaterialComponent* InComponent, FProperty* InProperty) override;
+};
diff --git a/Engine/Plugins/Experimental/DynamicMaterial/Source/DynamicMaterialEditor/Private/Slate/Properties/Generators/DMTextureUVPropertyRowGenerator.cpp b/Engine/Plugins/Experimental/DynamicMaterial/Source/DynamicMaterialEditor/Private/Slate/Properties/Generators/DMTextureUVPropertyRowGenerator.cpp
--- a/Engine/Plugins/Experimental/DynamicMaterial/Source/DynamicMaterialEditor/Private/Slate/Properties/Generators/DMTextureUVPropertyRowGenerator.cpp
+++ b/Engine/Plugins/Experimental/DynamicMaterial/Source/DynamicMaterialEditor/Private/Slate/Properties/Generators/DMTextureUVPropertyRowGenerator.cpp
@@ -4,15 +4,15 @@
 #include "Components/DMMaterialComponent.h"
 #include "Components/DMTextureUV.h"
 #include "DMEDefs.h"
+#include "DynamicMaterialEditorModule.h"
+#include "Model/DynamicMaterialModel.h"
+#include "Model/DynamicMaterialModelDynamic.h"
 #include "Modules/ModuleManager.h"
 #include "PropertyCustomizationHelpers.h"
 #include "PropertyEditorModule.h"
-#include "Slate/Properties/Editors/SDMPropertyEditVector.h"
 #include "Slate/Properties/SDMTextureUVVisualizerProperty.h"
 #include "Slate/SDMComponentEdit.h"
 #include "Slate/SDMEditor.h"
-#include "Styling/SlateIconFinder.h"
-#include "Utils/DMPrivate.h"
 #include "Widgets/Layout/SBox.h"
 
 #define LOCTEXT_NAMESPACE "DMTextureUVPropertyRowGenerator"
@@ -25,10 +25,10 @@ const TSharedRef<FDMTextureUVPropertyRowGenerator>& FDMTextureUVPropertyRowGener
 
 namespace UE::DynamicMaterialEditor::Private
 {
-	void AddPropertyRow(const TSharedRef<SWidget>& InComponentEditWidget, UDMTextureUV* InTextureUV, FName InProperty, 
+	void AddTextureUVPropertyRow(const TSharedRef<SWidget>& InComponentEditWidget, UDMTextureUV* InTextureUV, FName InProperty, 
 		TArray<FDMPropertyHandle>& InOutPropertyRows);
 
-	void AddVisualizerRow(const TSharedRef<SWidget>& InComponentEditWidget, UDMTextureUV* InTextureUV, 
+	void AddTextureUVVisualizerRow(const TSharedRef<SWidget>& InComponentEditWidget, UDMTextureUV* InTextureUV, 
 		TArray<FDMPropertyHandle>& InOutPropertyRows);
 
 	bool CanResetTextureUVPropertyToDefault(TSharedPtr<IPropertyHandle> InPropertyHandle);
@@ -58,15 +58,31 @@ void FDMTextureUVPropertyRowGenerator::AddComponentProperties(const TSharedRef<S
 
 	InOutProcessedObjects.Add(InComponent);
 
+	if (TSharedPtr<SDMEditor> EditorWidget = InComponentEditWidget->GetEditorWidget())
+	{
+		if (UDynamicMaterialModelBase* MaterialModelBase = EditorWidget->GetMaterialModelBase())
+		{
+			if (UDynamicMaterialModelDynamic* MaterialModelDynamic = Cast<UDynamicMaterialModelDynamic>(MaterialModelBase))
+			{
+				if (UDMMaterialComponentDynamic* ComponentDynamic = MaterialModelDynamic->GetComponentDynamic(TextureUV->GetFName()))
+				{
+					FDynamicMaterialEditorModule::Get().GeneratorComponentPropertyRows(InComponentEditWidget, ComponentDynamic, InOutPropertyRows, InOutProcessedObjects);
+				}
+
+				return;
+			}
+		}
+	}
+
 	using namespace UE::DynamicMaterialEditor::Private;
 
-	AddPropertyRow(InComponentEditWidget, TextureUV, UDMTextureUV::NAME_Offset, InOutPropertyRows);
-	AddPropertyRow(InComponentEditWidget, TextureUV, UDMTextureUV::NAME_Rotation, InOutPropertyRows);
-	AddPropertyRow(InComponentEditWidget, TextureUV, UDMTextureUV::NAME_Tiling, InOutPropertyRows);
-	AddPropertyRow(InComponentEditWidget, TextureUV, UDMTextureUV::NAME_Pivot, InOutPropertyRows);
-	AddPropertyRow(InComponentEditWidget, TextureUV, UDMTextureUV::NAME_bMirrorOnX, InOutPropertyRows);
-	AddPropertyRow(InComponentEditWidget, TextureUV, UDMTextureUV::NAME_bMirrorOnY, InOutPropertyRows);
-	AddVisualizerRow(InComponentEditWidget, TextureUV, InOutPropertyRows);
+	AddTextureUVPropertyRow(InComponentEditWidget, TextureUV, UDMTextureUV::NAME_Offset, InOutPropertyRows);
+	AddTextureUVPropertyRow(InComponentEditWidget, TextureUV, UDMTextureUV::NAME_Rotation, InOutPropertyRows);
+	AddTextureUVPropertyRow(InComponentEditWidget, TextureUV, UDMTextureUV::NAME_Tiling, InOutPropertyRows);
+	AddTextureUVPropertyRow(InComponentEditWidget, TextureUV, UDMTextureUV::NAME_Pivot, InOutPropertyRows);
+	AddTextureUVPropertyRow(InComponentEditWidget, TextureUV, UDMTextureUV::NAME_bMirrorOnX, InOutPropertyRows);
+	AddTextureUVPropertyRow(InComponentEditWidget, TextureUV, UDMTextureUV::NAME_bMirrorOnY, InOutPropertyRows);
+	AddTextureUVVisualizerRow(InComponentEditWidget, TextureUV, InOutPropertyRows);
 }
 
 void FDMTextureUVPropertyRowGenerator::AddPopoutComponentProperties(const TSharedRef<SWidget>& InParentWidget, UDMMaterialComponent* InComponent, 
@@ -86,12 +102,12 @@ void FDMTextureUVPropertyRowGenerator::AddPopoutComponentProperties(const TShare
 
 	using namespace UE::DynamicMaterialEditor::Private;
 
-	AddPropertyRow(InParentWidget, TextureUV, UDMTextureUV::NAME_Offset, InOutPropertyRows);
-	AddPropertyRow(InParentWidget, TextureUV, UDMTextureUV::NAME_Rotation, InOutPropertyRows);
-	AddPropertyRow(InParentWidget, TextureUV, UDMTextureUV::NAME_Tiling, InOutPropertyRows);
-	AddPropertyRow(InParentWidget, TextureUV, UDMTextureUV::NAME_Pivot, InOutPropertyRows);
-	AddPropertyRow(InParentWidget, TextureUV, UDMTextureUV::NAME_bMirrorOnX, InOutPropertyRows);
-	AddPropertyRow(InParentWidget, TextureUV, UDMTextureUV::NAME_bMirrorOnY, InOutPropertyRows);
+	AddTextureUVPropertyRow(InParentWidget, TextureUV, UDMTextureUV::NAME_Offset, InOutPropertyRows);
+	AddTextureUVPropertyRow(InParentWidget, TextureUV, UDMTextureUV::NAME_Rotation, InOutPropertyRows);
+	AddTextureUVPropertyRow(InParentWidget, TextureUV, UDMTextureUV::NAME_Tiling, InOutPropertyRows);
+	AddTextureUVPropertyRow(InParentWidget, TextureUV, UDMTextureUV::NAME_Pivot, InOutPropertyRows);
+	AddTextureUVPropertyRow(InParentWidget, TextureUV, UDMTextureUV::NAME_bMirrorOnX, InOutPropertyRows);
+	AddTextureUVPropertyRow(InParentWidget, TextureUV, UDMTextureUV::NAME_bMirrorOnY, InOutPropertyRows);
 }
 
 bool FDMTextureUVPropertyRowGenerator::AllowKeyframeButton(UDMMaterialComponent* InComponent, FProperty* InProperty)
@@ -109,7 +125,7 @@ bool FDMTextureUVPropertyRowGenerator::AllowKeyframeButton(UDMMaterialComponent*
 	return FDMComponentPropertyRowGenerator::AllowKeyframeButton(InComponent, InProperty);
 }
 
-void UE::DynamicMaterialEditor::Private::AddPropertyRow(const TSharedRef<SWidget>& InComponentEditWidget, UDMTextureUV* InTextureUV, 
+void UE::DynamicMaterialEditor::Private::AddTextureUVPropertyRow(const TSharedRef<SWidget>& InComponentEditWidget, UDMTextureUV* InTextureUV, 
 	FName InProperty, TArray<FDMPropertyHandle>& InOutPropertyRows)
 {
 	FDMPropertyHandle& NewHandle = InOutPropertyRows.Add_GetRef(SDMEditor::GetPropertyHandle(&*InComponentEditWidget, InTextureUV, InProperty));
@@ -118,9 +134,11 @@ void UE::DynamicMaterialEditor::Private::AddPropertyRow(const TSharedRef<SWidget
 		FIsResetToDefaultVisible::CreateStatic(&UE::DynamicMaterialEditor::Private::CanResetTextureUVPropertyToDefault),
 		FResetToDefaultHandler::CreateStatic(&UE::DynamicMaterialEditor::Private::ResetTextureUVPropertyToDefault)
 	);
+
+	NewHandle.bEnabled = true;
 }
 
-void UE::DynamicMaterialEditor::Private::AddVisualizerRow(const TSharedRef<SWidget>& InComponentEditWidget, UDMTextureUV* InTextureUV, 
+void UE::DynamicMaterialEditor::Private::AddTextureUVVisualizerRow(const TSharedRef<SWidget>& InComponentEditWidget, UDMTextureUV* InTextureUV, 
 	TArray<FDMPropertyHandle>& InOutPropertyRows)
 {
 	// Make sure we don't get a substage
@@ -131,12 +149,25 @@ void UE::DynamicMaterialEditor::Private::AddVisualizerRow(const TSharedRef<SWidg
 		return;
 	}
 
+	if (InComponentEditWidget->GetWidgetClass().GetWidgetType() != SDMComponentEdit::StaticWidgetClass().GetWidgetType())
+	{
+		return;
+	}
+
+	TSharedPtr<SDMEditor> EditorWidget = StaticCastSharedRef<SDMComponentEdit>(InComponentEditWidget)->GetEditorWidget();
+
+	if (!EditorWidget.IsValid())
+	{
+		return;
+	}
+
 	FDMPropertyHandle VisualizerHandle;
 	VisualizerHandle.NameOverride = LOCTEXT("Visualizer", "UV Visualizer");
 	VisualizerHandle.NameToolTipOverride = LOCTEXT("VisualizerToolTip", "A graphical Texture UV editor.\n\n- Offset Mode: Change the Texture UV offset.\n- Pivot Mode: Change the Texture UV pivot, rotation and tiling.\n\nControl+click to reset values to default.");
 	VisualizerHandle.ValueName = FName(*InTextureUV->GetComponentPath());
-	VisualizerHandle.ValueWidget = SNew(SDMTextureUVVisualizerProperty, Stage, InTextureUV);
+	VisualizerHandle.ValueWidget = SNew(SDMTextureUVVisualizerProperty, EditorWidget.ToSharedRef(), Stage).TextureUV(InTextureUV);
 	VisualizerHandle.CategoryOverrideName = TEXT("Texture UV");
+	VisualizerHandle.bEnabled = true;
 	InOutPropertyRows.Add(VisualizerHandle);
 }
 
@@ -150,9 +181,9 @@ bool UE::DynamicMaterialEditor::Private::CanResetTextureUVPropertyToDefault(TSha
 		return false;
 	}
 
-	FName PropertyName = Property->GetFName();
+	const FName PropertyName = Property->GetFName();
 
-	if (PropertyName == NAME_None)
+	if (PropertyName.IsNone())
 	{
 		return false;
 	}
diff --git a/Engine/Plugins/Experimental/DynamicMaterial/Source/DynamicMaterialEditor/Private/Slate/Properties/Generators/DMTextureUVPropertyRowGenerator.h b/Engine/Plugins/Experimental/DynamicMaterial/Source/DynamicMaterialEditor/Private/Slate/Properties/Generators/DMTextureUVPropertyRowGenerator.h
--- a/Engine/Plugins/Experimental/DynamicMaterial/Source/DynamicMaterialEditor/Private/Slate/Properties/Generators/DMTextureUVPropertyRowGenerator.h
+++ b/Engine/Plugins/Experimental/DynamicMaterial/Source/DynamicMaterialEditor/Private/Slate/Properties/Generators/DMTextureUVPropertyRowGenerator.h
@@ -11,7 +11,6 @@ class FDMTextureUVPropertyRowGenerator : public FDMComponentPropertyRowGenerator
 public:
 	static const TSharedRef<FDMTextureUVPropertyRowGenerator>& Get();
 
-	FDMTextureUVPropertyRowGenerator() = default;
 	virtual ~FDMTextureUVPropertyRowGenerator() override = default;
 
 	virtual void AddComponentProperties(const TSharedRef<SDMComponentEdit>& InComponentEditWidget, UDMMaterialComponent* InComponent,
diff --git a/Engine/Plugins/Experimental/DynamicMaterial/Source/DynamicMaterialEditor/Private/Slate/Properties/Generators/DMThroughputPropertyRowGenerator.h b/Engine/Plugins/Experimental/DynamicMaterial/Source/DynamicMaterialEditor/Private/Slate/Properties/Generators/DMThroughputPropertyRowGenerator.h
--- a/Engine/Plugins/Experimental/DynamicMaterial/Source/DynamicMaterialEditor/Private/Slate/Properties/Generators/DMThroughputPropertyRowGenerator.h
+++ b/Engine/Plugins/Experimental/DynamicMaterial/Source/DynamicMaterialEditor/Private/Slate/Properties/Generators/DMThroughputPropertyRowGenerator.h
@@ -9,7 +9,6 @@ class FDMThroughputPropertyRowGenerator : public FDMComponentPropertyRowGenerato
 public:
 	static const TSharedRef<FDMThroughputPropertyRowGenerator>& Get();
 
-	FDMThroughputPropertyRowGenerator() = default;
 	virtual ~FDMThroughputPropertyRowGenerator() override = default;
 
 	virtual void AddComponentProperties(const TSharedRef<SDMComponentEdit>& InComponentEditWidget, UDMMaterialComponent* InComponent,
diff --git a/Engine/Plugins/Experimental/DynamicMaterial/Source/DynamicMaterialEditor/Private/Slate/Properties/SDMMaterialProperty.cpp b/Engine/Plugins/Experimental/DynamicMaterial/Source/DynamicMaterialEditor/Private/Slate/Properties/SDMMaterialProperty.cpp
--- a/Engine/Plugins/Experimental/DynamicMaterial/Source/DynamicMaterialEditor/Private/Slate/Properties/SDMMaterialProperty.cpp
+++ b/Engine/Plugins/Experimental/DynamicMaterial/Source/DynamicMaterialEditor/Private/Slate/Properties/SDMMaterialProperty.cpp
@@ -3,7 +3,7 @@
 #include "SDMMaterialProperty.h"
 #include "Components/DMMaterialProperty.h"
 #include "DynamicMaterialEditorStyle.h"
-#include "Model/DynamicMaterialModel.h"
+#include "Model/DynamicMaterialModelBase.h"
 #include "Model/DynamicMaterialModelEditorOnlyData.h"
 #include "Slate/SDMEditor.h"
 #include "Slate/SDMSlot.h"
@@ -21,11 +21,11 @@ void SDMMaterialProperty::Construct(const FArguments& InArgs, const TSharedRef<S
 
 	if (ensure(MaterialEditor.IsValid()))
 	{
-		UDynamicMaterialModel* MaterialModel = MaterialEditor->GetMaterialModel();
+		UDynamicMaterialModelBase* MaterialModelBase = MaterialEditor->GetMaterialModelBase();
 
-		if (ensure(IsValid(MaterialModel)))
+		if (ensure(IsValid(MaterialModelBase)))
 		{
-			UDynamicMaterialModelEditorOnlyData* ModelEditorOnlyData = UDynamicMaterialModelEditorOnlyData::Get(MaterialModel);
+			UDynamicMaterialModelEditorOnlyData* ModelEditorOnlyData = UDynamicMaterialModelEditorOnlyData::Get(MaterialModelBase);
 
 			if (ensure(IsValid(ModelEditorOnlyData)))
 			{
diff --git a/Engine/Plugins/Experimental/DynamicMaterial/Source/DynamicMaterialEditor/Private/Slate/Properties/SDMSourceBlendType.cpp b/Engine/Plugins/Experimental/DynamicMaterial/Source/DynamicMaterialEditor/Private/Slate/Properties/SDMSourceBlendType.cpp
--- a/Engine/Plugins/Experimental/DynamicMaterial/Source/DynamicMaterialEditor/Private/Slate/Properties/SDMSourceBlendType.cpp
+++ b/Engine/Plugins/Experimental/DynamicMaterial/Source/DynamicMaterialEditor/Private/Slate/Properties/SDMSourceBlendType.cpp
@@ -182,7 +182,7 @@ void SDMSourceBlendType::EnsureBlendMap()
 			const FText BlendClassText = StageBlendCDO->GetDescription();
 			const TSubclassOf<UDMMaterialStageBlend> BlendClassObject = TSubclassOf<UDMMaterialStageBlend>(BlendClass.Get());
 
-			if (BlendClassName != NAME_None && !BlendClassText.IsEmpty())
+			if (!BlendClassName.IsNone() && !BlendClassText.IsEmpty())
 			{
 				BlendMap.Emplace(BlendClassName, {BlendClassText, BlendClassObject});
 			}
diff --git a/Engine/Plugins/Experimental/DynamicMaterial/Source/DynamicMaterialEditor/Private/Slate/Properties/SDMSourceBlendType.h b/Engine/Plugins/Experimental/DynamicMaterial/Source/DynamicMaterialEditor/Private/Slate/Properties/SDMSourceBlendType.h
--- a/Engine/Plugins/Experimental/DynamicMaterial/Source/DynamicMaterialEditor/Private/Slate/Properties/SDMSourceBlendType.h
+++ b/Engine/Plugins/Experimental/DynamicMaterial/Source/DynamicMaterialEditor/Private/Slate/Properties/SDMSourceBlendType.h
@@ -18,8 +18,6 @@ class UDMSourceBlendTypeContextObject : public UObject
 	GENERATED_BODY()
 	
 public:
-	UDMSourceBlendTypeContextObject() = default;
-	
 	TSharedPtr<SDMSourceBlendType> GetBlendTypeWidget() const { return BlendTypeWidgetWeak.Pin(); }
 	void SetBlendTypeWidget(const TSharedPtr<SDMSourceBlendType>& InBlendTypeWidget) { BlendTypeWidgetWeak = InBlendTypeWidget; }
 	
diff --git a/Engine/Plugins/Experimental/DynamicMaterial/Source/DynamicMaterialEditor/Private/Slate/Properties/SDMTextureUVVisualizer.cpp b/Engine/Plugins/Experimental/DynamicMaterial/Source/DynamicMaterialEditor/Private/Slate/Properties/SDMTextureUVVisualizer.cpp
--- a/Engine/Plugins/Experimental/DynamicMaterial/Source/DynamicMaterialEditor/Private/Slate/Properties/SDMTextureUVVisualizer.cpp
+++ b/Engine/Plugins/Experimental/DynamicMaterial/Source/DynamicMaterialEditor/Private/Slate/Properties/SDMTextureUVVisualizer.cpp
@@ -5,11 +5,13 @@
 #include "Brushes/SlateColorBrush.h"
 #include "Components/DMMaterialStage.h"
 #include "Components/DMTextureUV.h"
+#include "Components/DMTextureUVDynamic.h"
 #include "DynamicMaterialEditorModule.h"
 #include "Editor.h"
 #include "Framework/Application/SlateApplication.h"
-#include "Materials/MaterialInterface.h"
 #include "ScopedTransaction.h"
+#include "Slate/Previews/SDMComponentPreview.h"
+#include "Slate/SDMEditor.h"
 #include "Styling/StyleColors.h"
 #include "Widgets/Images/SImage.h"
 #include "Widgets/Input/SButton.h"
@@ -42,8 +44,7 @@ namespace UE::DynamicMaterialEditor::Private
 }
 
 SDMTextureUVVisualizer::SDMTextureUVVisualizer()
-	: Brush(FSlateMaterialBrush(UE::DynamicMaterialEditor::Private::TextureUVVisualizerImageSize))
-	, bIsPopout(false)
+	: bIsPopout(false)
 	, bPivotEditMode(false)
 	, CurrentAbsoluteSize(FVector2f::ZeroVector)
 	, CurrentAbsoluteCenter(FVector2f::ZeroVector)
@@ -54,20 +55,21 @@ SDMTextureUVVisualizer::SDMTextureUVVisualizer()
 	, ValueStart(FVector2D::ZeroVector)
 	, bInvertTiling(false)
 {
-	Brush.SetUVRegion(FBox2f(FVector2f(0, 0), FVector2f(1, 1)));
 }
 
-void SDMTextureUVVisualizer::Construct(const FArguments& InArgs, UDMMaterialStage* InMaterialStage, UDMTextureUV* InTextureUV)
+void SDMTextureUVVisualizer::Construct(const FArguments& InArgs, const TSharedRef<SDMEditor>& InEditorWidget, UDMMaterialStage* InMaterialStage)
 {
 	check(InMaterialStage);
-	check(InTextureUV);
+	check(InArgs._TextureUV || InArgs._TextureUVDynamic);
 
+	EditorWidgetWeak = InEditorWidget;
 	StageWeak = InMaterialStage;
-	TextureUVWeak = InTextureUV;
-	bIsPopout = InArgs._IsPopout;
 
-	InMaterialStage->GetOnUpdate().AddSP(this, &SDMTextureUVVisualizer::OnStageUpdated);
-	OnStageUpdated(InMaterialStage, EDMUpdateType::Structure);
+	TextureUVComponentWeak = InArgs._TextureUV
+		? static_cast<UDMMaterialComponent*>(InArgs._TextureUV)
+		: static_cast<UDMMaterialComponent*>(InArgs._TextureUVDynamic);
+
+	bIsPopout = InArgs._IsPopout;
 
 	SetCanTick(true);
 
@@ -75,9 +77,8 @@ void SDMTextureUVVisualizer::Construct(const FArguments& InArgs, UDMMaterialStag
 
 	ChildSlot
 	[
-		SAssignNew(StageImage, SImage)
-		.Image(&Brush)
-		.DesiredSizeOverride(bIsPopout ? TextureUVVisualizerPopoutImageSize : TextureUVVisualizerImageSize)
+		SAssignNew(StagePreview, SDMComponentPreview, InEditorWidget, InMaterialStage)
+		.PreviewSize(bIsPopout ? TextureUVVisualizerPopoutImageSize : TextureUVVisualizerImageSize)
 	];
 }
 
@@ -106,34 +107,155 @@ UDMMaterialStage* SDMTextureUVVisualizer::GetStage() const
 	return StageWeak.Get();
 }
 
+UDMMaterialComponent* SDMTextureUVVisualizer::GetTextureUVComponent() const
+{
+	return TextureUVComponentWeak.Get();
+}
+
 UDMTextureUV* SDMTextureUVVisualizer::GetTextureUV() const
 {
-	return TextureUVWeak.Get();
+	return Cast<UDMTextureUV>(GetTextureUVComponent());
 }
 
-void SDMTextureUVVisualizer::Tick(const FGeometry& InAllottedGeometry, const double InCurrentTime, const float InDeltaTime)
+UDMTextureUVDynamic* SDMTextureUVVisualizer::GetTextureUVDynamic() const
 {
-	SCompoundWidget::Tick(InAllottedGeometry, InCurrentTime, InDeltaTime);
+	return Cast<UDMTextureUVDynamic>(GetTextureUVComponent());
+}
 
-	if (!PreviewMaterialWeak.IsValid())
+const FVector2D& SDMTextureUVVisualizer::GetOffset() const
+{
+	if (UDMTextureUVDynamic* TextureUVDynamic = GetTextureUVDynamic())
 	{
-		Brush.SetMaterial(nullptr);
+		return TextureUVDynamic->GetOffset();
 	}
 
-	if (UDMMaterialComponent::CanClean())
+	if (UDMTextureUV* TextureUV = GetTextureUV())
 	{
-		if (UDMMaterialStage* Stage = StageWeak.Get())
-		{
-			if (Stage->IsComponentValid() && Stage->NeedsClean())
-			{
-				Stage->DoClean();
-				OnStageUpdated(StageWeak.Get(), EDMUpdateType::Structure);
-			}
-		}
+		return TextureUV->GetOffset();
 	}
 
-	CurrentAbsoluteSize = StageImage->GetTickSpaceGeometry().GetAbsoluteSize();
-	CurrentAbsoluteCenter = StageImage->GetTickSpaceGeometry().GetAbsolutePosition() + (CurrentAbsoluteSize * 0.5f);
+	return GetDefault<UDMTextureUV>()->GetOffset();
+}
+
+bool SDMTextureUVVisualizer::SetOffset(const FVector2D& InOffset)
+{
+	if (UDMTextureUVDynamic* TextureUVDynamic = GetTextureUVDynamic())
+	{
+		TextureUVDynamic->SetOffset(InOffset);
+		return true;
+	}
+
+	if (UDMTextureUV* TextureUV = GetTextureUV())
+	{
+		TextureUV->SetOffset(InOffset);
+		return true;
+	}
+
+	return false;
+}
+
+float SDMTextureUVVisualizer::GetRotation() const
+{
+	if (UDMTextureUVDynamic* TextureUVDynamic = GetTextureUVDynamic())
+	{
+		return TextureUVDynamic->GetRotation();
+	}
+
+	if (UDMTextureUV* TextureUV = GetTextureUV())
+	{
+		return TextureUV->GetRotation();
+	}
+
+	return GetDefault<UDMTextureUV>()->GetRotation();
+}
+
+bool SDMTextureUVVisualizer::SetRotation(float InRotation)
+{
+	if (UDMTextureUVDynamic* TextureUVDynamic = GetTextureUVDynamic())
+	{
+		TextureUVDynamic->SetRotation(InRotation);
+		return true;
+	}
+
+	if (UDMTextureUV* TextureUV = GetTextureUV())
+	{
+		TextureUV->SetRotation(InRotation);
+		return true;
+	}
+
+	return false;
+}
+
+const FVector2D& SDMTextureUVVisualizer::GetTiling() const
+{
+	if (UDMTextureUVDynamic* TextureUVDynamic = GetTextureUVDynamic())
+	{
+		return TextureUVDynamic->GetTiling();
+	}
+
+	if (UDMTextureUV* TextureUV = GetTextureUV())
+	{
+		return TextureUV->GetTiling();
+	}
+
+	return GetDefault<UDMTextureUV>()->GetTiling();
+}
+
+bool SDMTextureUVVisualizer::SetTiling(const FVector2D& InTiling)
+{
+	if (UDMTextureUVDynamic* TextureUVDynamic = GetTextureUVDynamic())
+	{
+		TextureUVDynamic->SetTiling(InTiling);
+		return true;
+	}
+
+	if (UDMTextureUV* TextureUV = GetTextureUV())
+	{
+		TextureUV->SetTiling(InTiling);
+		return true;
+	}
+
+	return false;
+}
+
+const FVector2D& SDMTextureUVVisualizer::GetPivot() const
+{
+	if (UDMTextureUVDynamic* TextureUVDynamic = GetTextureUVDynamic())
+	{
+		return TextureUVDynamic->GetPivot();
+	}
+
+	if (UDMTextureUV* TextureUV = GetTextureUV())
+	{
+		return TextureUV->GetPivot();
+	}
+
+	return GetDefault<UDMTextureUV>()->GetPivot();
+}
+
+bool SDMTextureUVVisualizer::SetPivot(const FVector2D& InPivot)
+{
+	if (UDMTextureUVDynamic* TextureUVDynamic = GetTextureUVDynamic())
+	{
+		TextureUVDynamic->SetPivot(InPivot);
+		return true;
+	}
+
+	if (UDMTextureUV* TextureUV = GetTextureUV())
+	{
+		TextureUV->SetPivot(InPivot);
+		return true;
+	}
+
+	return false;
+}
+
+void SDMTextureUVVisualizer::Tick(const FGeometry& InAllottedGeometry, const double InCurrentTime, const float InDeltaTime)
+{
+	SCompoundWidget::Tick(InAllottedGeometry, InCurrentTime, InDeltaTime);
+
+	CurrentAbsoluteSize = StagePreview->GetTickSpaceGeometry().GetAbsoluteSize();
+	CurrentAbsoluteCenter = StagePreview->GetTickSpaceGeometry().GetAbsolutePosition() + (CurrentAbsoluteSize * 0.5f);
 
 	if (ScrubbingMode != EScrubbingMode::None)
 	{
@@ -193,7 +315,7 @@ FCursorReply SDMTextureUVVisualizer::OnCursorQuery(const FGeometry& InGeometry,
 	const FVector2f MousePosition = FSlateApplication::Get().GetCursorPos();
 
 	// If the UV image isn't under the mouse, we don't need to do anything.
-	if (!StageImage->GetTickSpaceGeometry().IsUnderLocation(MousePosition))
+	if (!StagePreview->GetTickSpaceGeometry().IsUnderLocation(MousePosition))
 	{
 		return SCompoundWidget::OnCursorQuery(InGeometry, InCursorEvent);
 	}
@@ -234,13 +356,12 @@ FCursorReply SDMTextureUVVisualizer::OnCursorQuery(const FGeometry& InGeometry,
 	return SCompoundWidget::OnCursorQuery(InGeometry, InCursorEvent);
 }
 
-int32 SDMTextureUVVisualizer::OnPaint(const FPaintArgs& InArgs, const FGeometry& InAllottedGeometry, const FSlateRect& InMyCullingRect, FSlateWindowElementList& OutDrawElements, int32 InLayerId, const FWidgetStyle& InWidgetStyle, bool bInParentEnabled) const
+int32 SDMTextureUVVisualizer::OnPaint(const FPaintArgs& InArgs, const FGeometry& InAllottedGeometry, const FSlateRect& InMyCullingRect, 
+	FSlateWindowElementList& OutDrawElements, int32 InLayerId, const FWidgetStyle& InWidgetStyle, bool bInParentEnabled) const
 {
 	InLayerId = SCompoundWidget::OnPaint(InArgs, InAllottedGeometry, InMyCullingRect, OutDrawElements, InLayerId, InWidgetStyle, bInParentEnabled);
 
-	UDMTextureUV* TextureUV = TextureUVWeak.Get();
-
-	if (!IsValid(TextureUV))
+	if (!GetTextureUVComponent())
 	{
 		return InLayerId;
 	}
@@ -254,7 +375,7 @@ int32 SDMTextureUVVisualizer::OnPaint(const FPaintArgs& InArgs, const FGeometry&
 	static const FLinearColor NormalColor = FLinearColor(1.f, 1.f, 1.f, 0.25f);
 	static const FLinearColor HighlightColor = FStyleColors::Primary.GetSpecifiedColor();
 
-	const float Rotation = TextureUV->GetRotation();
+	const float Rotation = GetRotation();
 	const float RotationRadians = FMath::DegreesToRadians(Rotation);
 	const FVector2f LocalCenterOffset = GetOffsetLocation(InAllottedGeometry.GetLocalSize());
 	const FVector2f LocalPivotOffset = GetPivotLocation(InAllottedGeometry.GetLocalSize());
@@ -421,23 +542,6 @@ int32 SDMTextureUVVisualizer::OnPaint(const FPaintArgs& InArgs, const FGeometry&
 	return InLayerId;
 }
 
-void SDMTextureUVVisualizer::OnStageUpdated(UDMMaterialComponent* InComponent, EDMUpdateType InUpdateType)
-{
-	if (UDMMaterialStage* Stage = Cast<UDMMaterialStage>(InComponent))
-	{
-		if (Stage == StageWeak.Get() && IsValid(Stage) && Stage->IsComponentValid())
-		{
-			UMaterialInterface* PreviewMaterial = Stage->GetPreviewMaterial();
-
-			if (Brush.GetResourceObject() != PreviewMaterial)
-			{
-				Brush.SetMaterial(PreviewMaterial);
-				PreviewMaterialWeak = PreviewMaterial;
-			}
-		}
-	}
-}
-
 bool SDMTextureUVVisualizer::HasValidGeometry() const
 {
 	// The chances of the center being at anywhere near 0,0 is remote...
@@ -449,7 +553,7 @@ float SDMTextureUVVisualizer::GetCircleHandleBaseRadius() const
 {
 	using namespace UE::DynamicMaterialEditor::Private;
 
-	const FVector2f ImageSize = Brush.GetImageSize();
+	const FVector2f ImageSize = StagePreview->GetBrush().GetImageSize();
 	const float CircleHandleRadius = FMath::Min(ImageSize.X, ImageSize.Y) * TextureUVVisualizerCircleHandleBaseRadiusMultiplier;
 
 	return FMath::Clamp(CircleHandleRadius, 10, 50.f);
@@ -457,19 +561,12 @@ float SDMTextureUVVisualizer::GetCircleHandleBaseRadius() const
 
 FVector2f SDMTextureUVVisualizer::ApplyTextureUVTransform(const FVector2f& InUV) const
 {
-	UDMTextureUV* TextureUV = TextureUVWeak.Get();
-
-	if (!IsValid(TextureUV))
-	{
-		return InUV;
-	}
-
-	FVector2f Offset = static_cast<FVector2f>(TextureUV->GetOffset());
+	FVector2f Offset = static_cast<FVector2f>(GetOffset());
 	Offset.Y *= -1.f;
 
-	const float Rotation = TextureUV->GetRotation();
-	const FVector2f& Tiling = static_cast<FVector2f>(TextureUV->GetTiling());
-	const FVector2f& Pivot = static_cast<FVector2f>(TextureUV->GetPivot());
+	const float Rotation = GetRotation();
+	const FVector2f& Tiling = static_cast<FVector2f>(GetTiling());
+	const FVector2f& Pivot = static_cast<FVector2f>(GetPivot());
 
 	FVector2f TransformedUV = InUV;
 	TransformedUV -= Pivot;
@@ -508,19 +605,12 @@ FVector2f SDMTextureUVVisualizer::GetOffsetLocation(const FVector2f& InSize) con
 
 FVector2f SDMTextureUVVisualizer::GetPivotLocation(const FVector2f& InSize) const
 {
-	if (UDMTextureUV* TextureUV = TextureUVWeak.Get())
+	if (bIsPopout)
 	{
-		if (bIsPopout)
-		{
-			return ToPopoutLocation(FVector2f::UnitVector, static_cast<FVector2f>(TextureUV->GetPivot())) * InSize;
-		}
-
-		return static_cast<FVector2f>(TextureUV->GetPivot()) * InSize;
+		return ToPopoutLocation(FVector2f::UnitVector, static_cast<FVector2f>(GetPivot())) * InSize;
 	}
 
-	static const FVector2f Center = {0.5f, 0.5f};
-
-	return Center * InSize;
+	return static_cast<FVector2f>(GetPivot()) * InSize;
 }
 
 FVector2f SDMTextureUVVisualizer::GetAbsoluteOffsetLocation() const
@@ -536,14 +626,13 @@ FVector2f SDMTextureUVVisualizer::GetAbsolutePivotLocation() const
 float SDMTextureUVVisualizer::GetCircleHandleRadiusAtAngle(float InAngle) const
 {
 	const float BaseDistance = GetCircleHandleBaseRadius();
-	UDMTextureUV* TextureUV = TextureUVWeak.Get();
 
-	if (!IsValid(TextureUV))
+	if (!IsValid(GetTextureUVComponent()))
 	{
 		return BaseDistance;
 	}
 
-	const FVector2D& Tiling = TextureUV->GetTiling();
+	const FVector2D& Tiling = GetTiling();
 
 	if (FMath::IsNearlyEqual(Tiling.X, Tiling.Y))
 	{
@@ -569,9 +658,7 @@ SDMTextureUVVisualizer::EHandleAxis SDMTextureUVVisualizer::GetCenterHandleAxis(
 		return EHandleAxis::None;
 	}
 
-	UDMTextureUV* TextureUV = TextureUVWeak.Get();
-
-	if (!IsValid(TextureUV))
+	if (!IsValid(GetTextureUVComponent()))
 	{
 		return EHandleAxis::None;
 	}
@@ -583,7 +670,7 @@ SDMTextureUVVisualizer::EHandleAxis SDMTextureUVVisualizer::GetCenterHandleAxis(
 		
 	if (!bPivotEditMode)
 	{
-		const float Rotation = TextureUV->GetRotation();
+		const float Rotation = GetRotation();
 
 		if (!FMath::IsNearlyZero(Rotation))
 		{
@@ -621,9 +708,7 @@ SDMTextureUVVisualizer::EHandleAxis SDMTextureUVVisualizer::GetCenterHandleAxis(
 
 SDMTextureUVVisualizer::EHandleAxis SDMTextureUVVisualizer::GetCircleHandleAxis(const FVector2f& InAbsolutePosition) const
 {
-	UDMTextureUV* TextureUV = TextureUVWeak.Get();
-
-	if (!IsValid(TextureUV))
+	if (!IsValid(GetTextureUVComponent()))
 	{
 		return EHandleAxis::None;
 	}
@@ -632,9 +717,9 @@ SDMTextureUVVisualizer::EHandleAxis SDMTextureUVVisualizer::GetCircleHandleAxis(
 
 	const FVector2f HandleOffset = InAbsolutePosition - GetAbsolutePivotLocation();
 	const float DistanceFromHandle = HandleOffset.Size();
-	const FVector2D& Tiling2D = TextureUV->GetTiling();
+	const FVector2D& Tiling2D = GetTiling();
 	// We manage angle clockwise from +Y axis. Atan2 handles it anti-clockwise from +X axis.
-	float Angle = 90.f - FMath::RadiansToDegrees(FMath::Atan2(HandleOffset.Y, HandleOffset.X)) + TextureUV->GetRotation();
+	float Angle = 90.f - FMath::RadiansToDegrees(FMath::Atan2(HandleOffset.Y, HandleOffset.X)) + GetRotation();
 	float DistanceFromCircleHandle;
 
 	if (FMath::IsNearlyEqual(Tiling2D.X, Tiling2D.Y))
@@ -705,9 +790,7 @@ SDMTextureUVVisualizer::EHandleAxis SDMTextureUVVisualizer::GetCircleHandleAxis(
 
 bool SDMTextureUVVisualizer::TryClickCenterHandle(const FVector2f& InMousePosition, bool bInResetToDefault)
 {
-	UDMTextureUV* TextureUV = TextureUVWeak.Get();
-
-	if (!IsValid(TextureUV))
+	if (!IsValid(GetTextureUVComponent()))
 	{
 		return false;
 	}
@@ -730,8 +813,8 @@ bool SDMTextureUVVisualizer::TryClickCenterHandle(const FVector2f& InMousePositi
 		else
 		{
 			FScopedTransaction Transaction(LOCTEXT("ResetOffset", "Reset Offset to Default."));
-			TextureUV->Modify();
-			FVector2D NewOffset = TextureUV->GetOffset();
+			ModifyTextureUVComponent();
+			FVector2D NewOffset = GetOffset();
 
 			if (Axis == EHandleAxis::X)
 			{
@@ -746,7 +829,7 @@ bool SDMTextureUVVisualizer::TryClickCenterHandle(const FVector2f& InMousePositi
 				NewOffset = FVector2D::ZeroVector;
 			}
 
-			TextureUV->SetOffset(NewOffset);
+			SetOffset(NewOffset);
 		}
 	}
 	// Pivot handle
@@ -759,8 +842,8 @@ bool SDMTextureUVVisualizer::TryClickCenterHandle(const FVector2f& InMousePositi
 		else
 		{
 			FScopedTransaction Transaction(LOCTEXT("ResetPivot", "Reset Pivot to Default."));
-			TextureUV->Modify();
-			FVector2D NewPivot = TextureUV->GetOffset();
+			ModifyTextureUVComponent();
+			FVector2D NewPivot = GetOffset();
 
 			if (Axis == EHandleAxis::X)
 			{
@@ -775,7 +858,7 @@ bool SDMTextureUVVisualizer::TryClickCenterHandle(const FVector2f& InMousePositi
 				NewPivot = FVector2D(0.5, 0.5);
 			}
 
-			TextureUV->SetPivot(NewPivot);
+			SetPivot(NewPivot);
 		}
 	}
 
@@ -784,14 +867,12 @@ bool SDMTextureUVVisualizer::TryClickCenterHandle(const FVector2f& InMousePositi
 
 bool SDMTextureUVVisualizer::TryClickCircleHandle(const FVector2f& InMousePosition, bool bInResetToDefault)
 {
-	UDMTextureUV* TextureUV = TextureUVWeak.Get();
-
-	if (!IsValid(TextureUV))
+	if (!IsValid(GetTextureUVComponent()))
 	{
 		return false;
 	}
 
-	const float Rotation = TextureUV->GetRotation();
+	const float Rotation = GetRotation();
 	const FVector2f AbsolutePivotLocation = GetAbsolutePivotLocation();
 
 	// When the uv is rotated, do the opposite action
@@ -815,8 +896,8 @@ bool SDMTextureUVVisualizer::TryClickCircleHandle(const FVector2f& InMousePositi
 			else
 			{
 				FScopedTransaction Transaction(LOCTEXT("ResetRotation", "Reset Rotation to Default."));
-				TextureUV->Modify();
-				TextureUV->SetRotation(0.f);
+				ModifyTextureUVComponent();
+				SetRotation(0.f);
 			}
 
 			break;
@@ -832,10 +913,10 @@ bool SDMTextureUVVisualizer::TryClickCircleHandle(const FVector2f& InMousePositi
 			else
 			{
 				FScopedTransaction Transaction(LOCTEXT("ResetTilingeX", "Reset Tiling X to Default."));
-				TextureUV->Modify();
-				FVector2D NewTiling = TextureUV->GetTiling();
+				ModifyTextureUVComponent();
+				FVector2D NewTiling = GetTiling();
 				NewTiling.X = 1.0;
-				TextureUV->SetTiling(NewTiling);
+				SetTiling(NewTiling);
 			}
 
 			break;
@@ -851,10 +932,10 @@ bool SDMTextureUVVisualizer::TryClickCircleHandle(const FVector2f& InMousePositi
 			else
 			{
 				FScopedTransaction Transaction(LOCTEXT("ResetTilingY", "Reset Tiling Y to Default."));
-				TextureUV->Modify();
-				FVector2D NewTiling = TextureUV->GetTiling();
+				ModifyTextureUVComponent();
+				FVector2D NewTiling = GetTiling();
 				NewTiling.Y = 1.0;
-				TextureUV->SetTiling(NewTiling);
+				SetTiling(NewTiling);
 			}
 
 			break;
@@ -884,14 +965,14 @@ void SDMTextureUVVisualizer::UpdatePopoutUVs()
 		}
 	}
 
-	Brush.SetUVRegion(UVRegion);
+	StagePreview->GetBrush().SetUVRegion(UVRegion);
 }
 
 void SDMTextureUVVisualizer::SetScrubbingMode(EScrubbingMode InMode, EHandleAxis InAxis)
 {
-	UDMTextureUV* TextureUV = TextureUVWeak.Get();
+	UDMMaterialComponent* TextureUVComponent = GetTextureUVComponent();
 
-	if (InMode != EScrubbingMode::None && (InAxis == EHandleAxis::None || !IsValid(TextureUV) || !HasValidGeometry()))
+	if (InMode != EScrubbingMode::None && (InAxis == EHandleAxis::None || !IsValid(TextureUVComponent) || !HasValidGeometry()))
 	{
 		InMode = EScrubbingMode::None;
 	}
@@ -912,14 +993,14 @@ void SDMTextureUVVisualizer::SetScrubbingMode(EScrubbingMode InMode, EHandleAxis
 	switch (ScrubbingMode)
 	{
 		case EScrubbingMode::Offset:
-			ValueStart = TextureUV->GetOffset();
+			ValueStart = GetOffset();
 			TextureProperty = UDMTextureUV::StaticClass()->FindPropertyByName(UDMTextureUV::NAME_Offset);
 			UE_LOG(LogDynamicMaterialEditor, Verbose, TEXT("Started Offset mode"));
 			break;
 
 		case EScrubbingMode::Rotation:
 		{
-			ValueStart.X = TextureUV->GetRotation();
+			ValueStart.X = GetRotation();
 
 			/** Store the original mouse angle */
 			const FVector2f MouseOffset = ScrubbingStartAbsoluteMouse - GetAbsolutePivotLocation();
@@ -932,13 +1013,13 @@ void SDMTextureUVVisualizer::SetScrubbingMode(EScrubbingMode InMode, EHandleAxis
 		}
 
 		case EScrubbingMode::Tiling:
-			ValueStart = TextureUV->GetTiling();
+			ValueStart = GetTiling();
 			TextureProperty = UDMTextureUV::StaticClass()->FindPropertyByName(UDMTextureUV::NAME_Tiling);
 			UE_LOG(LogDynamicMaterialEditor, Verbose, TEXT("Started Tiling mode"));
 			break;
 
 		case EScrubbingMode::Pivot:
-			ValueStart = TextureUV->GetPivot();
+			ValueStart = GetPivot();
 			TextureProperty = UDMTextureUV::StaticClass()->FindPropertyByName(UDMTextureUV::NAME_Pivot);
 			UE_LOG(LogDynamicMaterialEditor, Verbose, TEXT("Started Pivot mode"));
 			break;
@@ -948,14 +1029,14 @@ void SDMTextureUVVisualizer::SetScrubbingMode(EScrubbingMode InMode, EHandleAxis
 	}
 
 	ScrubbingTransaction = MakeShared<FScopedTransaction>(LOCTEXT("VisualizerUVScrubbingTransaction", "UV Visualizer Scrub"));
-	TextureUV->Modify();
+	ModifyTextureUVComponent();
 
-	TextureUV->PreEditChange(TextureProperty);
+	TextureUVComponent->PreEditChange(TextureProperty);
 }
 
 FVector2f SDMTextureUVVisualizer::ToPopoutLocation(const FVector2f& InSize, FVector2f&& InLocation) const
 {
-	const FBox2d BrushUV = Brush.GetUVRegion();
+	const FBox2d BrushUV = StagePreview->GetBrush().GetUVRegion();
 
 	InLocation.X = FMath::GetMappedRangeValueUnclamped<float, float>(
 		FVector2f(BrushUV.Min.X, BrushUV.Max.X),
@@ -974,7 +1055,7 @@ FVector2f SDMTextureUVVisualizer::ToPopoutLocation(const FVector2f& InSize, FVec
 
 FVector2f SDMTextureUVVisualizer::FromPopoutLocation(const FVector2f& InSize, FVector2f&& InLocation) const
 {
-	const FBox2d BrushUV = Brush.GetUVRegion();
+	const FBox2d BrushUV = StagePreview->GetBrush().GetUVRegion();
 
 	InLocation.X = FMath::GetMappedRangeValueUnclamped<float, float>(
 		{0, 1},
@@ -991,6 +1072,14 @@ FVector2f SDMTextureUVVisualizer::FromPopoutLocation(const FVector2f& InSize, FV
 	return InLocation;
 }
 
+void SDMTextureUVVisualizer::ModifyTextureUVComponent()
+{
+	if (UDMMaterialComponent* TextureUVComponent = GetTextureUVComponent())
+	{
+		TextureUVComponent->Modify();
+	}
+}
+
 void SDMTextureUVVisualizer::UpdateScrub()
 {
 	switch (ScrubbingMode)
@@ -1023,28 +1112,28 @@ void SDMTextureUVVisualizer::UpdateScrub()
 
 void SDMTextureUVVisualizer::UpdateScrub_Offset()
 {
-	UDMTextureUV* TextureUV = TextureUVWeak.Get();
+	UDMMaterialComponent* TextureUVComponent = GetTextureUVComponent();
 
-	if (!IsValid(TextureUV))
+	if (!IsValid(TextureUVComponent))
 	{
 		return;
 	}
 
 	const FVector2f MouseOffset = FSlateApplication::Get().GetCursorPos() - ScrubbingStartAbsoluteMouse;
-	const FBox2d BrushUV = Brush.GetUVRegion();
+	const FBox2d BrushUV = StagePreview->GetBrush().GetUVRegion();
 
 	FVector2D OffsetChange = static_cast<FVector2D>(MouseOffset / CurrentAbsoluteSize);
 	OffsetChange.X *= (BrushUV.Max.X - BrushUV.Min.X);
 	OffsetChange.Y *= (BrushUV.Max.Y - BrushUV.Min.Y);
 
-	const float Rotation = TextureUV->GetRotation();
+	const float Rotation = GetRotation();
 
 	if (!FMath::IsNearlyZero(Rotation))
 	{
 		OffsetChange = OffsetChange.GetRotated(-Rotation);
 	}
 
-	OffsetChange *= TextureUV->GetTiling();
+	OffsetChange *= GetTiling();
 
 	if (HandleAxis == EHandleAxis::X)
 	{
@@ -1060,22 +1149,22 @@ void SDMTextureUVVisualizer::UpdateScrub_Offset()
 		OffsetChange.X = 0;
 	}
 
-	TextureUV->SetOffset(ValueStart + OffsetChange);
+	SetOffset(ValueStart + OffsetChange);
 
 	FPropertyChangedEvent ChangedEvent(
 		UDMTextureUV::StaticClass()->FindPropertyByName(UDMTextureUV::NAME_Offset),
 		EPropertyChangeType::Interactive,
-		{TextureUV}
+		{TextureUVComponent}
 	);
 
-	TextureUV->PostEditChangeProperty(ChangedEvent);
+	TextureUVComponent->PostEditChangeProperty(ChangedEvent);
 }
 
 void SDMTextureUVVisualizer::UpdateScrub_Rotation()
 {
-	UDMTextureUV* TextureUV = TextureUVWeak.Get();
+	UDMMaterialComponent* TextureUVComponent = GetTextureUVComponent();
 
-	if (!IsValid(TextureUV))
+	if (!IsValid(TextureUVComponent))
 	{
 		return;
 	}
@@ -1085,40 +1174,40 @@ void SDMTextureUVVisualizer::UpdateScrub_Rotation()
 	const float Angle = FMath::RadiansToDegrees(FMath::Atan2(CurrentMouseOffset.Y, CurrentMouseOffset.X));
 	const float NewValue = UE::Math::TRotator<float>::ClampAxis(ValueStart.X + Angle - ValueStart.Y);
 
-	TextureUV->SetRotation(NewValue);
+	SetRotation(NewValue);
 
 	FPropertyChangedEvent ChangedEvent(
 		UDMTextureUV::StaticClass()->FindPropertyByName(UDMTextureUV::NAME_Rotation),
 		EPropertyChangeType::Interactive,
-		{TextureUV}
+		{TextureUVComponent}
 	);
 
-	TextureUV->PostEditChangeProperty(ChangedEvent);
+	TextureUVComponent->PostEditChangeProperty(ChangedEvent);
 }
 
 void SDMTextureUVVisualizer::UpdateScrub_Tiling()
 {
-	UDMTextureUV* TextureUV = TextureUVWeak.Get();
+	UDMMaterialComponent* TextureUVComponent = GetTextureUVComponent();
 
-	if (!IsValid(TextureUV))
+	if (!IsValid(TextureUVComponent))
 	{
 		return;
 	}
 
 	const FVector2f MouseOffset = FSlateApplication::Get().GetCursorPos() - ScrubbingStartAbsoluteMouse;
-	const FBox2d BrushUV = Brush.GetUVRegion();
+	const FBox2d BrushUV = StagePreview->GetBrush().GetUVRegion();
 
 	FVector2D TilingChange = static_cast<FVector2D>(MouseOffset / CurrentAbsoluteSize);
 	TilingChange.X *= (BrushUV.Max.X - BrushUV.Min.X);
 	TilingChange.Y *= (BrushUV.Max.Y - BrushUV.Min.Y);
-	TilingChange /= TextureUV->GetTiling();
+	TilingChange /= GetTiling();
 
 	if (!bInvertTiling)
 	{
 		TilingChange *= -1.f;
 	}
 
-	const float Rotation = TextureUV->GetRotation();
+	const float Rotation = GetRotation();
 
 	if (!FMath::IsNearlyZero(Rotation))
 	{
@@ -1164,28 +1253,28 @@ void SDMTextureUVVisualizer::UpdateScrub_Tiling()
 		}
 	}
 
-	TextureUV->SetTiling(NewTiling);
+	SetTiling(NewTiling);
 
 	FPropertyChangedEvent ChangedEvent(
 		UDMTextureUV::StaticClass()->FindPropertyByName(UDMTextureUV::NAME_Tiling),
 		EPropertyChangeType::Interactive,
-		{TextureUV}
+		{TextureUVComponent}
 	);
 
-	TextureUV->PostEditChangeProperty(ChangedEvent);
+	TextureUVComponent->PostEditChangeProperty(ChangedEvent);
 }
 
 void SDMTextureUVVisualizer::UpdateScrub_Pivot()
 {
-	UDMTextureUV* TextureUV = TextureUVWeak.Get();
+	UDMMaterialComponent* TextureUVComponent = GetTextureUVComponent();
 
-	if (!IsValid(TextureUV))
+	if (!IsValid(TextureUVComponent))
 	{
 		return;
 	}
 
 	const FVector2f MouseOffset = FSlateApplication::Get().GetCursorPos() - ScrubbingStartAbsoluteMouse;
-	const FBox2d BrushUV = Brush.GetUVRegion();
+	const FBox2d BrushUV = StagePreview->GetBrush().GetUVRegion();
 
 	FVector2D PivotChange = static_cast<FVector2D>(MouseOffset / CurrentAbsoluteSize);
 	PivotChange.X *= (BrushUV.Max.X - BrushUV.Min.X);
@@ -1200,15 +1289,15 @@ void SDMTextureUVVisualizer::UpdateScrub_Pivot()
 		PivotChange.X = 0;
 	}
 
-	TextureUV->SetPivot(ValueStart + PivotChange);
+	SetPivot(ValueStart + PivotChange);
 
 	FPropertyChangedEvent ChangedEvent(
 		UDMTextureUV::StaticClass()->FindPropertyByName(UDMTextureUV::NAME_Pivot),
 		EPropertyChangeType::Interactive,
-		{TextureUV}
+		{TextureUVComponent}
 	);
 
-	TextureUV->PostEditChangeProperty(ChangedEvent);
+	TextureUVComponent->PostEditChangeProperty(ChangedEvent);
 }
 
 #undef LOCTEXT_NAMESPACE
diff --git a/Engine/Plugins/Experimental/DynamicMaterial/Source/DynamicMaterialEditor/Private/Slate/Properties/SDMTextureUVVisualizer.h b/Engine/Plugins/Experimental/DynamicMaterial/Source/DynamicMaterialEditor/Private/Slate/Properties/SDMTextureUVVisualizer.h
--- a/Engine/Plugins/Experimental/DynamicMaterial/Source/DynamicMaterialEditor/Private/Slate/Properties/SDMTextureUVVisualizer.h
+++ b/Engine/Plugins/Experimental/DynamicMaterial/Source/DynamicMaterialEditor/Private/Slate/Properties/SDMTextureUVVisualizer.h
@@ -8,12 +8,14 @@
 #include "UObject/WeakObjectPtr.h"
 
 class FScopedTransaction;
-class SImage;
+class SDMComponentPreview;
+class SDMEditor;
 class SWidget;
 class UDMMaterialComponent;
 class UDMMaterialStage;
 class UDMTextureUV;
-class UMaterialInterface;
+class UDMTextureUVDynamic;
+class UMaterial;
 enum class EDMUpdateType : uint8;
 
 /**
@@ -42,15 +44,21 @@ public:
 	};
 
 	SLATE_BEGIN_ARGS(SDMTextureUVVisualizer)
-		: _IsPopout(false)
+		: _TextureUV(nullptr)
+		, _TextureUVDynamic(nullptr)
+		, _IsPopout(false)
 		{}
+		SLATE_ARGUMENT(UDMTextureUV*, TextureUV)
+		SLATE_ARGUMENT(UDMTextureUVDynamic*, TextureUVDynamic)
 		SLATE_ARGUMENT(bool, IsPopout)
 	SLATE_END_ARGS()
 
 	SDMTextureUVVisualizer();
 
 	/** The TextureUV should be a sub-property of the stage */
-	void Construct(const FArguments& InArgs, UDMMaterialStage* InMaterialStage, UDMTextureUV* InTextureUV);
+	void Construct(const FArguments& InArgs, const TSharedRef<SDMEditor>& InEditorWidget, UDMMaterialStage* InMaterialStage);
+
+	TSharedPtr<SDMEditor> GetEditorWidget() const { return EditorWidgetWeak.Pin(); }
 
 	EScrubbingMode GetScrubbingMode() const;
 
@@ -62,8 +70,24 @@ public:
 
 	UDMMaterialStage* GetStage() const;
 
+	UDMMaterialComponent* GetTextureUVComponent() const;
+
 	UDMTextureUV* GetTextureUV() const;
 
+	UDMTextureUVDynamic* GetTextureUVDynamic() const;
+
+	const FVector2D& GetOffset() const;
+	bool SetOffset(const FVector2D& InOffset);
+
+	float GetRotation() const;
+	bool SetRotation(float InRotation);
+
+	const FVector2D& GetTiling() const;
+	bool SetTiling(const FVector2D& InTiling);
+
+	const FVector2D& GetPivot() const;
+	bool SetPivot(const FVector2D& InPivot);
+
 	//~ Begin SWidget
 	virtual void Tick(const FGeometry& InAllottedGeometry, const double InCurrentTime, const float InDeltaTime) override;
 	virtual FReply OnMouseButtonDown(const FGeometry& InGeometry, const FPointerEvent& InMouseEvent) override;
@@ -73,11 +97,10 @@ public:
 	//~ End SWidget
 
 protected:
+	TWeakPtr<SDMEditor> EditorWidgetWeak;
 	TWeakObjectPtr<UDMMaterialStage> StageWeak;
-	TWeakObjectPtr<UDMTextureUV> TextureUVWeak;
-	TSharedPtr<SImage> StageImage;
-	TWeakObjectPtr<UMaterialInterface> PreviewMaterialWeak;
-	FSlateMaterialBrush Brush;
+	TWeakObjectPtr<UDMMaterialComponent> TextureUVComponentWeak;
+	TSharedPtr<SDMComponentPreview> StagePreview;
 	bool bIsPopout;
 	bool bPivotEditMode;
 	FVector2f CurrentAbsoluteSize;
@@ -90,8 +113,6 @@ protected:
 	bool bInvertTiling;
 	TSharedPtr<FScopedTransaction> ScrubbingTransaction;
 
-	void OnStageUpdated(UDMMaterialComponent* InComponent, EDMUpdateType InUpdateType);
-
 	bool HasValidGeometry() const;
 
 	float GetCircleHandleBaseRadius() const;
@@ -125,6 +146,8 @@ protected:
 
 	FVector2f FromPopoutLocation(const FVector2f& InSize, FVector2f&& InLocation) const;
 
+	void ModifyTextureUVComponent();
+
 	void UpdateScrub();
 	void UpdateScrub_Offset();
 	void UpdateScrub_Rotation();
diff --git a/Engine/Plugins/Experimental/DynamicMaterial/Source/DynamicMaterialEditor/Private/Slate/Properties/SDMTextureUVVisualizerPopout.cpp b/Engine/Plugins/Experimental/DynamicMaterial/Source/DynamicMaterialEditor/Private/Slate/Properties/SDMTextureUVVisualizerPopout.cpp
--- a/Engine/Plugins/Experimental/DynamicMaterial/Source/DynamicMaterialEditor/Private/Slate/Properties/SDMTextureUVVisualizerPopout.cpp
+++ b/Engine/Plugins/Experimental/DynamicMaterial/Source/DynamicMaterialEditor/Private/Slate/Properties/SDMTextureUVVisualizerPopout.cpp
@@ -3,13 +3,15 @@
 #include "Slate/Properties/SDMTextureUVVisualizerPopout.h"
 
 #include "Components/DMMaterialStage.h"
+#include "Components/DMMaterialValueDynamic.h"
 #include "Components/DMTextureUV.h"
+#include "Components/DMTextureUVDynamic.h"
 #include "CustomDetailsViewArgs.h"
 #include "CustomDetailsViewModule.h"
 #include "DetailLayoutBuilder.h"
 #include "DMWorldSubsystem.h"
-#include "DynamicMaterialEditorStyle.h"
 #include "Engine/World.h"
+#include "Generators/DMTextureUVDynamicPropertyRowGenerator.h"
 #include "Generators/DMTextureUVPropertyRowGenerator.h"
 #include "ICustomDetailsView.h"
 #include "Items/ICustomDetailsViewCustomItem.h"
@@ -29,33 +31,60 @@
 
 const FName SDMTextureUVVisualizerPopout::TabId = TEXT("SDMTextureUVVisualizerPopout");
 
-void SDMTextureUVVisualizerPopout::CreatePopout(UDMMaterialStage* InMaterialStage, UDMTextureUV* InTextureUV)
+namespace UE::DynamicMaterialEditor::Private
+{
+	TSharedPtr<SDockTab> GetVisualizerTab(FName InTabId)
+	{
+		if (!FGlobalTabmanager::Get()->HasTabSpawner(InTabId))
+		{
+			FGlobalTabmanager::Get()->RegisterNomadTabSpawner(
+				InTabId,
+				FOnSpawnTab::CreateLambda(
+					[InTabId](const FSpawnTabArgs& InArgs)
+					{
+						TSharedRef<SDockTab> DockTab = SNew(SDockTab)
+							.Label(FText::FromName(InTabId))
+							.LabelSuffix(LOCTEXT("TabSuffix", "[UV Vis]"));
+
+						DockTab->SetTabIcon(FSlateIconFinder::FindIconForClass(UMaterial::StaticClass()).GetIcon());
+
+						return DockTab;
+					}
+				)
+			);
+		}
+
+		return FGlobalTabmanager::Get()->TryInvokeTab(InTabId);
+	}
+}
+
+void SDMTextureUVVisualizerPopout::CreatePopout(const TSharedRef<SDMEditor>& InEditorWidget, UDMMaterialStage* InMaterialStage, UDMTextureUV* InTextureUV)
 {
 	if (!IsValid(InMaterialStage) || !IsValid(InTextureUV))
 	{
 		return;
 	}
 
-	if (!FGlobalTabmanager::Get()->HasTabSpawner(TabId))
+	TSharedPtr<SDockTab> Tab = UE::DynamicMaterialEditor::Private::GetVisualizerTab(TabId);
+
+	if (!Tab.IsValid())
 	{
-		FGlobalTabmanager::Get()->RegisterNomadTabSpawner(
-			TabId,
-			FOnSpawnTab::CreateLambda(
-				[](const FSpawnTabArgs& InArgs)
-				{
-					TSharedRef<SDockTab> DockTab = SNew(SDockTab)
-						.Label(FText::FromName(TabId))
-						.LabelSuffix(LOCTEXT("TabSuffix", "[UV Vis]"));
-
-					DockTab->SetTabIcon(FSlateIconFinder::FindIconForClass(UMaterial::StaticClass()).GetIcon());
-
-					return DockTab;
-				}
-			)
-		);
+		return;
 	}
 
-	TSharedPtr<SDockTab> Tab = FGlobalTabmanager::Get()->TryInvokeTab(TabId);
+	Tab->ActivateInParent(ETabActivationCause::SetDirectly);
+	Tab->SetLabel(FText::FromString(InTextureUV->GetPathName()));
+	Tab->SetContent(SNew(SDMTextureUVVisualizerPopout, InEditorWidget, InMaterialStage).TextureUV(InTextureUV));
+}
+
+void SDMTextureUVVisualizerPopout::CreatePopout(const TSharedRef<SDMEditor>& InEditorWidget, UDMMaterialStage* InMaterialStage, UDMTextureUVDynamic* InTextureUVDynamic)
+{
+	if (!IsValid(InMaterialStage) || !IsValid(InTextureUVDynamic))
+	{
+		return;
+	}
+
+	TSharedPtr<SDockTab> Tab = UE::DynamicMaterialEditor::Private::GetVisualizerTab(TabId);
 
 	if (!Tab.IsValid())
 	{
@@ -63,14 +92,14 @@ void SDMTextureUVVisualizerPopout::CreatePopout(UDMMaterialStage* InMaterialStag
 	}
 
 	Tab->ActivateInParent(ETabActivationCause::SetDirectly);
-	Tab->SetLabel(FText::FromString(InTextureUV->GetPathName()));
-	Tab->SetContent(SNew(SDMTextureUVVisualizerPopout, InMaterialStage, InTextureUV));
+	Tab->SetLabel(FText::FromString(InTextureUVDynamic->GetPathName()));
+	Tab->SetContent(SNew(SDMTextureUVVisualizerPopout, InEditorWidget, InMaterialStage).TextureUVDynamic(InTextureUVDynamic));
 }
 
-void SDMTextureUVVisualizerPopout::Construct(const FArguments& InArgs, UDMMaterialStage* InMaterialStage, UDMTextureUV* InTextureUV)
+void SDMTextureUVVisualizerPopout::Construct(const FArguments& InArgs, const TSharedRef<SDMEditor>& InEditorWidget, UDMMaterialStage* InMaterialStage)
 {
 	check(InMaterialStage);
-	check(InTextureUV);
+	check(InArgs._TextureUV || InArgs._TextureUVDynamic);
 
 	ChildSlot
 	[
@@ -83,7 +112,9 @@ void SDMTextureUVVisualizerPopout::Construct(const FArguments& InArgs, UDMMateri
 			SNew(SOverlay)
 			+ SOverlay::Slot()
 			[
-				SAssignNew(Visualizer, SDMTextureUVVisualizer, InMaterialStage, InTextureUV)
+				SAssignNew(Visualizer, SDMTextureUVVisualizer, InEditorWidget, InMaterialStage)
+				.TextureUV(InArgs._TextureUV)
+				.TextureUVDynamic(InArgs._TextureUVDynamic)
 				.IsPopout(true)
 			]
 			+ SOverlay::Slot()
@@ -145,7 +176,11 @@ void SDMTextureUVVisualizerPopout::Construct(const FArguments& InArgs, UDMMateri
 			.HAlign(EHorizontalAlignment::HAlign_Fill)
 			.VAlign(EVerticalAlignment::VAlign_Top)
 			[
-				CreatePropertyWidget(InTextureUV)
+				CreatePropertyWidget(
+					InArgs._TextureUV 
+						? static_cast<UDMMaterialComponent*>(InArgs._TextureUV)
+						: static_cast<UDMMaterialComponent*>(InArgs._TextureUVDynamic)
+				)
 			]
 		]
 	];
@@ -201,17 +236,24 @@ FVector2D SDMTextureUVVisualizerPopout::GetSideBlockSize() const
 	return FVector2D::UnitVector;
 }
 
-TSharedRef<SWidget> SDMTextureUVVisualizerPopout::CreatePropertyWidget(UDMTextureUV* InTextureUV)
+TSharedRef<SWidget> SDMTextureUVVisualizerPopout::CreatePropertyWidget(UDMMaterialComponent* InComponent)
 {
 	SDMEditor::ClearPropertyHandles(this);
 
+	if (!InComponent)
+	{
+		return SNullWidget::NullWidget;
+	}
+
+	const bool bIsDynamic = InComponent && InComponent->IsA<UDMTextureUVDynamic>();
+
 	FCustomDetailsViewArgs Args;
 	Args.KeyframeHandler = nullptr;
 	Args.bAllowGlobalExtensions = true;
 	Args.bAllowResetToDefault = true;
 	Args.bShowCategories = false;
 
-	if (UWorld* World = InTextureUV->GetWorld())
+	if (UWorld* World = InComponent->GetWorld())
 	{
 		if (UDMWorldSubsystem* WorldSubsystem = World->GetSubsystem<UDMWorldSubsystem>())
 		{
@@ -223,7 +265,15 @@ TSharedRef<SWidget> SDMTextureUVVisualizerPopout::CreatePropertyWidget(UDMTextur
 	const FCustomDetailsViewItemId RootId = DetailsView->GetRootItem()->GetItemId();
 
 	TArray<FDMPropertyHandle> TextureUVPropertyRows;
-	FDMTextureUVPropertyRowGenerator::AddPopoutComponentProperties(SharedThis(this), InTextureUV, TextureUVPropertyRows);
+
+	if (bIsDynamic)
+	{
+		FDMTextureUVDynamicPropertyRowGenerator::AddPopoutComponentProperties(SharedThis(this), InComponent, TextureUVPropertyRows);
+	}
+	else
+	{
+		FDMTextureUVPropertyRowGenerator::AddPopoutComponentProperties(SharedThis(this), InComponent, TextureUVPropertyRows);
+	}
 
 	FDMPropertyHandle EditModeButtonRow;
 	EditModeButtonRow.ValueName = TEXT("EditMode");
@@ -302,6 +352,11 @@ TSharedRef<SWidget> SDMTextureUVVisualizerPopout::CreatePropertyWidget(UDMTextur
 				continue;
 			}
 
+			if (!EditRow.bEnabled)
+			{
+				Item->AsItem()->SetEnabledOverride(false);
+			}
+
 			Item->SetValueWidget(EditRow.ValueWidget.ToSharedRef());
 			DetailsView->ExtendTree(RootId, Position, Item->AsItem());
 			continue;
@@ -325,6 +380,11 @@ TSharedRef<SWidget> SDMTextureUVVisualizerPopout::CreatePropertyWidget(UDMTextur
 			);
 		}
 
+		if (!EditRow.bEnabled)
+		{
+			Item->SetEnabledOverride(false);
+		}
+
 		if (EditRow.PropertyHandle->HasMetaData("NotKeyframeable"))
 		{
 			Item->SetKeyframeEnabled(false);
diff --git a/Engine/Plugins/Experimental/DynamicMaterial/Source/DynamicMaterialEditor/Private/Slate/Properties/SDMTextureUVVisualizerPopout.h b/Engine/Plugins/Experimental/DynamicMaterial/Source/DynamicMaterialEditor/Private/Slate/Properties/SDMTextureUVVisualizerPopout.h
--- a/Engine/Plugins/Experimental/DynamicMaterial/Source/DynamicMaterialEditor/Private/Slate/Properties/SDMTextureUVVisualizerPopout.h
+++ b/Engine/Plugins/Experimental/DynamicMaterial/Source/DynamicMaterialEditor/Private/Slate/Properties/SDMTextureUVVisualizerPopout.h
@@ -6,10 +6,13 @@
 #include "Widgets/SCompoundWidget.h"
 
 class ICustomDetailsView;
+class SDMEditor;
 class SDMTextureUVVisualizer;
 class SDockTab;
+class UDMMaterialComponent;
 class UDMMaterialStage;
 class UDMTextureUV;
+class UDMTextureUVDynamic;
 enum class ECheckBoxState : uint8;
 
 /**
@@ -25,13 +28,20 @@ class SDMTextureUVVisualizerPopout : public SCompoundWidget
 public:
 	static const FName TabId;
 
-	static void CreatePopout(UDMMaterialStage* InMaterialStage, UDMTextureUV* InTextureUV);
+	static void CreatePopout(const TSharedRef<SDMEditor>& InEditorWidget, UDMMaterialStage* InMaterialStage, UDMTextureUV* InTextureUV);
 
-	SLATE_BEGIN_ARGS(SDMTextureUVVisualizerPopout) {}
+	static void CreatePopout(const TSharedRef<SDMEditor>& InEditorWidget, UDMMaterialStage* InMaterialStage, UDMTextureUVDynamic* InTextureUVDynamic);
+
+	SLATE_BEGIN_ARGS(SDMTextureUVVisualizerPopout)
+		: _TextureUV(nullptr)
+		, _TextureUVDynamic(nullptr)
+		{}
+		SLATE_ARGUMENT(UDMTextureUV*, TextureUV)
+		SLATE_ARGUMENT(UDMTextureUVDynamic*, TextureUVDynamic)
 	SLATE_END_ARGS()
 
 	/** The TextureUV should be a sub-property of the stage */
-	void Construct(const FArguments& InArgs, UDMMaterialStage* InMaterialStage, UDMTextureUV* InTextureUV);
+	void Construct(const FArguments& InArgs, const TSharedRef<SDMEditor>& InEditorWidget, UDMMaterialStage* InMaterialStage);
 
 protected:
 	TSharedPtr<SDMTextureUVVisualizer> Visualizer;
@@ -42,7 +52,7 @@ protected:
 
 	FVector2D GetSideBlockSize() const;
 
-	TSharedRef<SWidget> CreatePropertyWidget(UDMTextureUV* InTextureUV);
+	TSharedRef<SWidget> CreatePropertyWidget(UDMMaterialComponent* InComponent);
 
 	ECheckBoxState GetModeCheckBoxState(bool bInIsPivot) const;
 
diff --git a/Engine/Plugins/Experimental/DynamicMaterial/Source/DynamicMaterialEditor/Private/Slate/Properties/SDMTextureUVVisualizerProperty.cpp b/Engine/Plugins/Experimental/DynamicMaterial/Source/DynamicMaterialEditor/Private/Slate/Properties/SDMTextureUVVisualizerProperty.cpp
--- a/Engine/Plugins/Experimental/DynamicMaterial/Source/DynamicMaterialEditor/Private/Slate/Properties/SDMTextureUVVisualizerProperty.cpp
+++ b/Engine/Plugins/Experimental/DynamicMaterial/Source/DynamicMaterialEditor/Private/Slate/Properties/SDMTextureUVVisualizerProperty.cpp
@@ -4,6 +4,7 @@
 
 #include "Components/DMMaterialStage.h"
 #include "Components/DMTextureUV.h"
+#include "Components/DMTextureUVDynamic.h"
 #include "DetailLayoutBuilder.h"
 #include "DynamicMaterialEditorSettings.h"
 #include "Slate/Properties/SDMTextureUVVisualizer.h"
@@ -16,10 +17,10 @@
 
 #define LOCTEXT_NAMESPACE "SDMTextureUVVisualizerProperty"
 
-void SDMTextureUVVisualizerProperty::Construct(const FArguments& InArgs, UDMMaterialStage* InMaterialStage, UDMTextureUV* InTextureUV)
+void SDMTextureUVVisualizerProperty::Construct(const FArguments& InArgs, const TSharedRef<SDMEditor>& InEditorWidget, UDMMaterialStage* InMaterialStage)
 {
 	check(InMaterialStage);
-	check(InTextureUV);
+	check(InArgs._TextureUV || InArgs._TextureUVDynamic);
 
 	ChildSlot
 	[
@@ -65,7 +66,9 @@ void SDMTextureUVVisualizerProperty::Construct(const FArguments& InArgs, UDMMate
 			.MinAspectRatio(1)
 			.MaxAspectRatio(1)
 			[
-				SAssignNew(Visualizer, SDMTextureUVVisualizer, InMaterialStage, InTextureUV)
+				SAssignNew(Visualizer, SDMTextureUVVisualizer, InEditorWidget, InMaterialStage)
+				.TextureUV(InArgs._TextureUV)
+				.TextureUVDynamic(InArgs._TextureUVDynamic)
 				.IsPopout(false)
 			]			
 		]
@@ -113,7 +116,6 @@ void SDMTextureUVVisualizerProperty::Construct(const FArguments& InArgs, UDMMate
 	];
 }
 
-
 FReply SDMTextureUVVisualizerProperty::OnToggleVisualizerClicked()
 {
 	if (UDynamicMaterialEditorSettings* Settings = GetMutableDefault<UDynamicMaterialEditorSettings>())
@@ -143,14 +145,28 @@ FReply SDMTextureUVVisualizerProperty::OnOpenPopoutClicked()
 	}
 
 	UDMMaterialStage* Stage = Visualizer->GetStage();
-	UDMTextureUV* TextureUV = Visualizer->GetTextureUV();
+	UDMMaterialComponent* TextureUVComponent = Visualizer->GetTextureUVComponent();
+
+	if (!IsValid(Stage) || !IsValid(TextureUVComponent))
+	{
+		return FReply::Handled();
+	}
 
-	if (!IsValid(Stage) || !IsValid(TextureUV))
+	TSharedPtr<SDMEditor> EditorWidget = Visualizer->GetEditorWidget();
+
+	if (!EditorWidget.IsValid())
 	{
 		return FReply::Handled();
 	}
 
-	SDMTextureUVVisualizerPopout::CreatePopout(Stage, TextureUV);
+	if (UDMTextureUV* TextureUV = Cast<UDMTextureUV>(TextureUVComponent))
+	{
+		SDMTextureUVVisualizerPopout::CreatePopout(EditorWidget.ToSharedRef(), Stage, TextureUV);
+	}
+	else if (UDMTextureUVDynamic* TextureUVDynamic = Cast<UDMTextureUVDynamic>(TextureUVComponent))
+	{
+		SDMTextureUVVisualizerPopout::CreatePopout(EditorWidget.ToSharedRef(), Stage, TextureUVDynamic);
+	}
 
 	return FReply::Handled();
 }
diff --git a/Engine/Plugins/Experimental/DynamicMaterial/Source/DynamicMaterialEditor/Private/Slate/Properties/SDMTextureUVVisualizerProperty.h b/Engine/Plugins/Experimental/DynamicMaterial/Source/DynamicMaterialEditor/Private/Slate/Properties/SDMTextureUVVisualizerProperty.h
--- a/Engine/Plugins/Experimental/DynamicMaterial/Source/DynamicMaterialEditor/Private/Slate/Properties/SDMTextureUVVisualizerProperty.h
+++ b/Engine/Plugins/Experimental/DynamicMaterial/Source/DynamicMaterialEditor/Private/Slate/Properties/SDMTextureUVVisualizerProperty.h
@@ -4,10 +4,12 @@
 
 #include "Widgets/SCompoundWidget.h"
 
+class SDMEditor;
 class SDMTextureUVVisualizer;
 class SDMTextureUVVisualizerPopout;
 class UDMMaterialStage;
 class UDMTextureUV;
+class UDMTextureUVDynamic;
 enum class ECheckBoxState : uint8;
 
 /**
@@ -18,12 +20,16 @@ enum class ECheckBoxState : uint8;
 class SDMTextureUVVisualizerProperty : public SCompoundWidget
 {
 public:
-
-	SLATE_BEGIN_ARGS(SDMTextureUVVisualizerProperty) {}
+	SLATE_BEGIN_ARGS(SDMTextureUVVisualizerProperty)
+		: _TextureUV(nullptr)
+		, _TextureUVDynamic(nullptr)
+		{}
+		SLATE_ARGUMENT(UDMTextureUV*, TextureUV)
+		SLATE_ARGUMENT(UDMTextureUVDynamic*, TextureUVDynamic)
 	SLATE_END_ARGS()
 
 	/** The TextureUV should be a sub-property of the stage */
-	void Construct(const FArguments& InArgs, UDMMaterialStage* InMaterialStage, UDMTextureUV* InTextureUV);
+	void Construct(const FArguments& InArgs, const TSharedRef<SDMEditor>& InEditorWidget, UDMMaterialStage* InMaterialStage);
 
 protected:
 	TSharedPtr<SDMTextureUVVisualizer> Visualizer;
