Material Designer

diff --git a/Engine/Plugins/Experimental/DynamicMaterial/Source/DynamicMaterialEditor/Private/Slate/SDMComponentEdit.cpp b/Engine/Plugins/Experimental/DynamicMaterial/Source/DynamicMaterialEditor/Private/Slate/SDMComponentEdit.cpp
--- a/Engine/Plugins/Experimental/DynamicMaterial/Source/DynamicMaterialEditor/Private/Slate/SDMComponentEdit.cpp
+++ b/Engine/Plugins/Experimental/DynamicMaterial/Source/DynamicMaterialEditor/Private/Slate/SDMComponentEdit.cpp
@@ -34,6 +34,7 @@
 #include "Menus/DMMaterialStageSourceMenus.h"
 #include "Misc/CoreDelegates.h"
 #include "Model/DynamicMaterialModel.h"
+#include "Model/DynamicMaterialModelDynamic.h"
 #include "Model/DynamicMaterialModelEditorOnlyData.h"
 #include "PropertyCustomizationHelpers.h"
 #include "PropertyEditorModule.h"
@@ -49,6 +50,7 @@
 #include "Styling/SlateTypes.h"
 #include "ToolMenu.h"
 #include "ToolMenus.h"
+#include "Model/DynamicMaterialModelDynamic.h"
 #include "Utils/DMPrivate.h"
 #include "Widgets/Input/SComboButton.h"
 #include "Widgets/Layout/SBox.h"
@@ -255,7 +257,7 @@ void SDMComponentEdit::Construct(const FArguments& InArgs, UDMMaterialComponent*
 	{
 		if (TSharedPtr<SDMEditor> EditorWidget = InEditorWidget.Pin())
 		{
-			WorldContext = EditorWidget->GetMaterialModel();
+			WorldContext = EditorWidget->GetMaterialModelBase();
 		}
 	}
 
@@ -286,6 +288,10 @@ void SDMComponentEdit::Construct(const FArguments& InArgs, UDMMaterialComponent*
 		];
 }
 
+void SDMComponentEdit::PrivateRegisterAttributes(struct FSlateAttributeDescriptor::FInitializer&)
+{
+}
+
 TSharedRef<SWidget> SDMComponentEdit::CreateEditWidget()
 {
 	SDMEditor::ClearPropertyHandles(this);
@@ -294,6 +300,16 @@ TSharedRef<SWidget> SDMComponentEdit::CreateEditWidget()
 
 	UDMMaterialComponent* Component = ComponentWeak.Get();
 
+	bool bIsDynamic = false;
+
+	if (TSharedPtr<SDMEditor> EditorWidget = GetEditorWidget())
+	{
+		if (UDynamicMaterialModelBase* MaterialModelBase = EditorWidget->GetMaterialModelBase())
+		{
+			bIsDynamic = !MaterialModelBase->IsA<UDynamicMaterialModel>();
+		}
+	}
+
 	FCustomDetailsViewArgs Args;
 	Args.KeyframeHandler = KeyframeHandler;
 	Args.bAllowGlobalExtensions = true;
@@ -338,6 +354,11 @@ TSharedRef<SWidget> SDMComponentEdit::CreateEditWidget()
 		{
 			TypeSelectorItem->SetValueWidget(CreateSourceTypeEditWidget());
 
+			if (bIsDynamic)
+			{
+				TypeSelectorItem->AsItem()->SetEnabledOverride(false);
+			}
+
 			DetailsView->ExtendTree(
 				GetDefaultCategory()->GetItemId(),
 				ECustomDetailsTreeInsertPosition::FirstChild,
@@ -373,7 +394,7 @@ TSharedRef<SWidget> SDMComponentEdit::CreateEditWidget()
 
 		FName CategoryName = EditRow.CategoryOverrideName;
 
-		if (CategoryName == NAME_None && EditRow.PropertyHandle.IsValid())
+		if (CategoryName.IsNone() && EditRow.PropertyHandle.IsValid())
 		{
 			// Sub category (possibly)
 			if (TSharedPtr<IPropertyHandle> SubCategoryProperty = EditRow.PropertyHandle->GetParentHandle())
@@ -394,7 +415,7 @@ TSharedRef<SWidget> SDMComponentEdit::CreateEditWidget()
 
 		TSharedPtr<ICustomDetailsViewItem> CategoryItem;
 
-		if (CategoryName == NAME_None)
+		if (CategoryName.IsNone())
 		{
 			CategoryItem = GetDefaultCategory();
 		}
@@ -432,6 +453,14 @@ TSharedRef<SWidget> SDMComponentEdit::CreateEditWidget()
 
 			Item->SetValueWidget(EditRow.ValueWidget.ToSharedRef());
 
+			if (!EditRow.bEnabled)
+			{
+				Item->AsItem()->SetEnabledOverride(false);
+
+				// Disable the expansion widgets (SNullWidget is treated as removing the override).
+				Item->SetExpansionWidget(SNew(SBox));
+			}
+
 			if (EditRow.MaxWidth.IsSet())
 			{
 				Item->AsItem()->SetValueWidgetWidthOverride(EditRow.MaxWidth);
@@ -467,6 +496,14 @@ TSharedRef<SWidget> SDMComponentEdit::CreateEditWidget()
 			);
 		}
 
+		if (!EditRow.bEnabled)
+		{
+			Item->SetEnabledOverride(false);
+
+			// Disable the expansion widgets (SNullWidget is treated as removing the override).
+			Item->SetOverrideWidget(ECustomDetailsViewWidgetType::Extensions, SNew(SBox));
+		}
+
 		if (EditRow.PropertyHandle.IsValid() && EditRow.PropertyHandle->HasMetaData("NotKeyframeable"))
 		{
 			Item->SetKeyframeEnabled(false);
@@ -510,9 +547,11 @@ TArray<FDMPropertyHandle> SDMComponentEdit::GetEditRows()
 	{
 		if (TSharedPtr<SDMEditor> EditorWidget = GetEditorWidget())
 		{
-			if (UDynamicMaterialModel* MaterialModel = EditorWidget->GetMaterialModel())
+			if (UDynamicMaterialModelBase* MaterialModelBase = EditorWidget->GetMaterialModelBase())
 			{
-				GenerateMaterialModelPropertyRows(EditorWidget.ToSharedRef(), MaterialModel, PropertyRows, ProcessedObjects);
+				bool bIsDynamic = !MaterialModelBase->IsA<UDynamicMaterialModel>();
+
+				GenerateMaterialModelPropertyRows(EditorWidget.ToSharedRef(), MaterialModelBase, bIsDynamic, PropertyRows, ProcessedObjects);
 			}
 		}
 	}
@@ -886,37 +925,66 @@ void SDMComponentEdit::OnExpansionStateChanged(const TSharedRef<ICustomDetailsVi
 	EditorWidget->SetExpansionState(ComponentWeak.Get(), *ItemId.GetItemName(), bInExpansionState);
 }
 
-void SDMComponentEdit::GenerateMaterialModelPropertyRows(const TSharedRef<SDMEditor> InEditorWidget, UDynamicMaterialModel* InMaterialModel,
-	TArray<FDMPropertyHandle>& InOutPropertyRows, TSet<UDMMaterialComponent*>& InOutProcessedObjects)
+void SDMComponentEdit::GenerateMaterialModelPropertyRows(const TSharedRef<SDMEditor> InEditorWidget, UDynamicMaterialModelBase* InMaterialModelBase,
+	bool bInDynamic, TArray<FDMPropertyHandle>& InOutPropertyRows, TSet<UDMMaterialComponent*>& InOutProcessedObjects)
 {
+	UDynamicMaterialModel* MaterialModel = InMaterialModelBase->ResolveMaterialModel();
+
+	if (!MaterialModel)
+	{
+		return;
+	}
+
 	const FName MaterialSettingsCategory = FName("Material Settings");
 
-	auto AddGlobalValue = [InEditorWidget, &InOutPropertyRows, &InOutProcessedObjects, &MaterialSettingsCategory]
-		(UDMMaterialValue* InValue, const FText& InNameOverride)
+	auto AddGlobalValue = [InEditorWidget, &InOutPropertyRows, &InOutProcessedObjects, &MaterialSettingsCategory, InMaterialModelBase]
+		(UDMMaterialComponent* InComponent, const FText& InNameOverride)
 		{
-			FDMPropertyHandle& ValueHandle = InOutPropertyRows.Add_GetRef(InEditorWidget->GetPropertyHandle(&*InEditorWidget,
-				InValue, UDMMaterialValue::ValueName));
-
-			ValueHandle.CategoryOverrideName = MaterialSettingsCategory;
-			ValueHandle.NameOverride = InNameOverride;
-			ValueHandle.ResetToDefaultOverride = FResetToDefaultOverride::Create(
-				FIsResetToDefaultVisible::CreateUObject(InValue, &UDMMaterialValue::CanResetToDefault),
-				FResetToDefaultHandler::CreateUObject(InValue, &UDMMaterialValue::ResetToDefault),
-				/* Propagate to children */ false
-			);
+			if (UDynamicMaterialModelDynamic* MaterialModelDynamic = Cast<UDynamicMaterialModelDynamic>(InMaterialModelBase))
+			{
+				InComponent = MaterialModelDynamic->GetComponentDynamic(InComponent->GetFName());
+
+				if (!InComponent)
+				{
+					return;
+				}
+			}
+
+			FDMPropertyHandle& ComponentHandle = InOutPropertyRows.Add_GetRef(InEditorWidget->GetPropertyHandle(&*InEditorWidget,
+				InComponent, UDMMaterialValue::ValueName));
+
+			ComponentHandle.CategoryOverrideName = MaterialSettingsCategory;
+			ComponentHandle.NameOverride = InNameOverride;
+
+			if (UDMMaterialValue* MaterialValue = Cast<UDMMaterialValue>(InComponent))
+			{
+				ComponentHandle.ResetToDefaultOverride = FResetToDefaultOverride::Create(
+					FIsResetToDefaultVisible::CreateUObject(MaterialValue, &UDMMaterialValue::CanResetToDefault),
+					FResetToDefaultHandler::CreateUObject(MaterialValue, &UDMMaterialValue::ResetToDefault),
+					/* Propagate to children */ false
+				);
+			}
+			else if (UDMMaterialValueDynamic* MaterialValueDynamic = Cast<UDMMaterialValueDynamic>(InComponent))
+			{
+				ComponentHandle.ResetToDefaultOverride = FResetToDefaultOverride::Create(
+					FIsResetToDefaultVisible::CreateUObject(MaterialValueDynamic, &UDMMaterialValueDynamic::CanResetToDefault),
+					FResetToDefaultHandler::CreateUObject(MaterialValueDynamic, &UDMMaterialValueDynamic::ResetToDefault),
+					/* Propagate to children */ false
+				);
+			}
 
-			InOutProcessedObjects.Add(InValue);
+			InOutProcessedObjects.Add(InComponent);
 		};
 
-	AddGlobalValue(InMaterialModel->GetGlobalParameterValue(UDynamicMaterialModel::GlobalOffsetValueName), LOCTEXT("GlobalOffset", "Global Offset"));
-	AddGlobalValue(InMaterialModel->GetGlobalParameterValue(UDynamicMaterialModel::GlobalTilingValueName), LOCTEXT("GlobalTiling", "Global Tiling"));
-	AddGlobalValue(InMaterialModel->GetGlobalParameterValue(UDynamicMaterialModel::GlobalRotationValueName), LOCTEXT("GlobalRotation", "Global Rotation"));
+	AddGlobalValue(MaterialModel->GetGlobalParameterValue(UDynamicMaterialModel::GlobalOffsetValueName), LOCTEXT("GlobalOffset", "Global Offset"));
+	AddGlobalValue(MaterialModel->GetGlobalParameterValue(UDynamicMaterialModel::GlobalTilingValueName), LOCTEXT("GlobalTiling", "Global Tiling"));
+	AddGlobalValue(MaterialModel->GetGlobalParameterValue(UDynamicMaterialModel::GlobalRotationValueName), LOCTEXT("GlobalRotation", "Global Rotation"));
 
-	if (UDynamicMaterialModelEditorOnlyData* EditorOnlyData = UDynamicMaterialModelEditorOnlyData::Get(InMaterialModel))
+	if (UDynamicMaterialModelEditorOnlyData* EditorOnlyData = UDynamicMaterialModelEditorOnlyData::Get(InMaterialModelBase))
 	{
 		if (EditorOnlyData->GetBlendMode() != BLEND_Opaque)
 		{
-			AddGlobalValue(InMaterialModel->GetGlobalParameterValue(UDynamicMaterialModel::GlobalOpacityValueName), LOCTEXT("GlobalOpacity", "Global Opacity"));
+			AddGlobalValue(MaterialModel->GetGlobalParameterValue(UDynamicMaterialModel::GlobalOpacityValueName), LOCTEXT("GlobalOpacity", "Global Opacity"));
 		}
 
 		const FText PropertyFormat = LOCTEXT("PropertyFormat", "Global {0}");
@@ -953,15 +1021,16 @@ void SDMComponentEdit::GenerateMaterialModelPropertyRows(const TSharedRef<SDMEdi
 	}
 
 	const FName MaterialTypeCategory = FName("Material Type");
-	auto AddVariable = [InEditorWidget, &InOutPropertyRows, &MaterialTypeCategory](UObject* InObject, FName InPropertyName)
+	auto AddVariable = [InEditorWidget, &InOutPropertyRows, &MaterialTypeCategory, bInDynamic](UObject* InObject, FName InPropertyName)
 		{
 			FDMPropertyHandle& ValueHandle = InOutPropertyRows.Add_GetRef(InEditorWidget->GetPropertyHandle(&*InEditorWidget,
 				InObject, InPropertyName));
 
 			ValueHandle.CategoryOverrideName = MaterialTypeCategory;
+			ValueHandle.bEnabled = !bInDynamic;
 		};
 
-	if (UDynamicMaterialModelEditorOnlyData* EditorOnlyData = UDynamicMaterialModelEditorOnlyData::Get(InMaterialModel))
+	if (UDynamicMaterialModelEditorOnlyData* EditorOnlyData = UDynamicMaterialModelEditorOnlyData::Get(InMaterialModelBase))
 	{
 		AddVariable(EditorOnlyData, GET_MEMBER_NAME_CHECKED(UDynamicMaterialModelEditorOnlyData, ChannelListPreset));
 		AddVariable(EditorOnlyData, GET_MEMBER_NAME_CHECKED(UDynamicMaterialModelEditorOnlyData, Domain));
diff --git a/Engine/Plugins/Experimental/DynamicMaterial/Source/DynamicMaterialEditor/Private/Slate/SDMComponentEdit.h b/Engine/Plugins/Experimental/DynamicMaterial/Source/DynamicMaterialEditor/Private/Slate/SDMComponentEdit.h
--- a/Engine/Plugins/Experimental/DynamicMaterial/Source/DynamicMaterialEditor/Private/Slate/SDMComponentEdit.h
+++ b/Engine/Plugins/Experimental/DynamicMaterial/Source/DynamicMaterialEditor/Private/Slate/SDMComponentEdit.h
@@ -32,6 +32,7 @@ class UDMMaterialStageSource;
 class UDMMaterialStageThroughput;
 class UDMMaterialValueFloat1;
 class UDMTextureUV;
+class UDynamicMaterialModelBase;
 class UDynamicMaterialModelEditorOnlyData;
 class UMaterial;
 enum class ECheckBoxState : uint8;
@@ -49,11 +50,12 @@ public:
 	static TSharedRef<SWidget> CreateExtensionButtons(const TSharedPtr<SDMComponentEdit>& InComponentEditWidget, UDMMaterialComponent* InComponent,
 		const FName& InPropertyName, bool bInAllowKeyframe);
 
+	SLATE_DECLARE_WIDGET(SDMComponentEdit, SCompoundWidget)
+
 	SLATE_BEGIN_ARGS(SDMComponentEdit)
 		{}
 	SLATE_END_ARGS()
 
-	SDMComponentEdit() = default;
 	virtual ~SDMComponentEdit() override;
 
 	void Construct(const FArguments& InArgs, UDMMaterialComponent* InComponent, const TWeakPtr<SDMEditor>& InEditorWidget);
@@ -69,8 +71,8 @@ public:
 	//~ End FUndoClient
 
 protected:
-	static void GenerateMaterialModelPropertyRows(const TSharedRef<SDMEditor> InEditorWidget, UDynamicMaterialModel* InMaterialModel, 
-		TArray<FDMPropertyHandle>& InOutPropertyRows, TSet<UDMMaterialComponent*>& InOutProcessedObjects);
+	static void GenerateMaterialModelPropertyRows(const TSharedRef<SDMEditor> InEditorWidget, UDynamicMaterialModelBase* InMaterialModelBase, 
+		bool bInDynamic, TArray<FDMPropertyHandle>& InOutPropertyRows, TSet<UDMMaterialComponent*>& InOutProcessedObjects);
 
 	TWeakObjectPtr<UDMMaterialComponent> ComponentWeak;
 	TWeakPtr<SDMEditor> EditorWidgetWeak;
diff --git a/Engine/Plugins/Experimental/DynamicMaterial/Source/DynamicMaterialEditor/Private/Slate/SDMEditor.cpp b/Engine/Plugins/Experimental/DynamicMaterial/Source/DynamicMaterialEditor/Private/Slate/SDMEditor.cpp
--- a/Engine/Plugins/Experimental/DynamicMaterial/Source/DynamicMaterialEditor/Private/Slate/SDMEditor.cpp
+++ b/Engine/Plugins/Experimental/DynamicMaterial/Source/DynamicMaterialEditor/Private/Slate/SDMEditor.cpp
@@ -1,6 +1,7 @@
 // Copyright Epic Games, Inc. All Rights Reserved.
 
 #include "SDMEditor.h"
+#include "AssetRegistry/AssetRegistryModule.h"
 #include "AssetThumbnail.h"
 #include "Components/DMMaterialLayer.h"
 #include "Components/DMMaterialProperty.h"
@@ -8,7 +9,6 @@
 #include "Components/MaterialStageExpressions/DMMSETextureSample.h"
 #include "Components/MaterialValues/DMMaterialValueFloat.h"
 #include "Components/PrimitiveComponent.h"
-#include "DetailLayoutBuilder.h"
 #include "DMWorldSubsystem.h"
 #include "DynamicMaterialEditorCommands.h"
 #include "DynamicMaterialEditorModule.h"
@@ -16,24 +16,26 @@
 #include "DynamicMaterialEditorStyle.h"
 #include "DynamicMaterialModule.h"
 #include "Engine/World.h"
+#include "Factories/MaterialFactoryNew.h"
 #include "Framework/Commands/GenericCommands.h"
 #include "Framework/Commands/InputChord.h"
 #include "HAL/PlatformApplicationMisc.h"
 #include "IDetailTreeNode.h"
 #include "IPropertyRowGenerator.h"
+#include "Material/DynamicMaterialInstance.h"
+#include "Materials/Material.h"
 #include "Materials/MaterialInterface.h"
 #include "Menus/DMToolBarMenus.h"
 #include "Misc/CoreDelegates.h"
 #include "Model/DynamicMaterialModel.h"
+#include "Model/DynamicMaterialModelDynamic.h"
 #include "Model/DynamicMaterialModelEditorOnlyData.h"
 #include "PropertyHandle.h"
-#include "Slate/Properties/SDMMaterialParameters.h"
 #include "Slate/SDMComponentEdit.h"
 #include "Slate/SDMMaterialWizard.h"
 #include "Slate/SDMSlot.h"
 #include "Slate/SDMToolBar.h"
 #include "SlateOptMacros.h"
-#include "Material/DynamicMaterialInstance.h"
 #include "Styling/StyleColors.h"
 #include "ThumbnailRendering/ThumbnailManager.h"
 #include "Utils/DMBlueprintFunctionLibrary.h"
@@ -190,32 +192,32 @@ namespace UE::DynamicMaterialEditor::Private
 		}
 	}
 
-	bool CheckMaterialModelValidity(UDynamicMaterialModel* InMaterialModel)
+	bool CheckMaterialModelValidity(UDynamicMaterialModelBase* InMaterialModelBase)
 		{
-			if (UWorld* World = InMaterialModel->GetWorld())
+			if (UWorld* World = InMaterialModelBase->GetWorld())
 			{
 				if (UDMWorldSubsystem* WorldSubsystem = World->GetSubsystem<UDMWorldSubsystem>())
 				{
-					if (!WorldSubsystem->ExecuteIsValidDelegate(InMaterialModel))
+					if (!WorldSubsystem->ExecuteIsValidDelegate(InMaterialModelBase))
 					{
 						return false;
 					}
 				}
 			}
 
-			UActorComponent* ComponentOuter = InMaterialModel->GetTypedOuter<UActorComponent>();
+			UActorComponent* ComponentOuter = InMaterialModelBase->GetTypedOuter<UActorComponent>();
 			if (ComponentOuter && !IsValid(ComponentOuter))
 			{
 				return false;
 			}
 
-			AActor* ActorOuter = InMaterialModel->GetTypedOuter<AActor>();
+			AActor* ActorOuter = InMaterialModelBase->GetTypedOuter<AActor>();
 			if (ActorOuter && !IsValid(ActorOuter))
 			{
 				return false;
 			}
 
-			UPackage* PackageOuter = InMaterialModel->GetPackage();
+			UPackage* PackageOuter = InMaterialModelBase->GetPackage();
 			if (PackageOuter && !IsValid(PackageOuter))
 			{
 				return false;
@@ -248,7 +250,7 @@ TSharedRef<FAssetThumbnailPool> SDMEditor::GetThumbnailPool()
 	return ThumbnailPool.ToSharedRef();
 }
 
-void SDMEditor::Construct(const FArguments& InArgs, TWeakObjectPtr<UDynamicMaterialModel> InModelWeak)
+void SDMEditor::Construct(const FArguments& InArgs)
 {
 	SetCanTick(true);
 
@@ -268,7 +270,7 @@ void SDMEditor::Construct(const FArguments& InArgs, TWeakObjectPtr<UDynamicMater
 
 	SetEmptyLayout();
 
-	SetMaterialModel(InModelWeak.Get());
+	SetMaterialModelBase(nullptr);
 
 	UDynamicMaterialEditorSettings::Get()->OnSettingsChanged.AddSP(this, &SDMEditor::OnSettingsChanged);
 
@@ -295,21 +297,28 @@ SDMEditor::~SDMEditor()
 {
 	if (FDynamicMaterialModule::AreUObjectsSafe())
 	{
-		if (UDynamicMaterialModel* MaterialModel = MaterialModelWeak.Get())
+		if (UDynamicMaterialModelEditorOnlyData* ModelEditorOnlyData = UDynamicMaterialModelEditorOnlyData::Get(MaterialModelBaseWeak))
 		{
-			if (UDynamicMaterialModelEditorOnlyData* ModelEditorOnlyData = UDynamicMaterialModelEditorOnlyData::Get(MaterialModel))
-			{
-				ModelEditorOnlyData->GetOnMaterialBuiltDelegate().RemoveAll(this);
-				ModelEditorOnlyData->GetOnValueListUpdateDelegate().RemoveAll(this);
-				ModelEditorOnlyData->GetOnSlotListUpdateDelegate().RemoveAll(this);
-			}
+			ModelEditorOnlyData->GetOnMaterialBuiltDelegate().RemoveAll(this);
+			ModelEditorOnlyData->GetOnValueListUpdateDelegate().RemoveAll(this);
+			ModelEditorOnlyData->GetOnSlotListUpdateDelegate().RemoveAll(this);
 		}
 	}
 }
 
+UDynamicMaterialModel* SDMEditor::GetMaterialModel() const
+{
+	if (UDynamicMaterialModelBase* MaterialModelBase = MaterialModelBaseWeak.Get())
+	{
+		return MaterialModelBase->ResolveMaterialModel();
+	}
+
+	return nullptr;
+}
+
 void SDMEditor::ClearEditor()
 {
-	MaterialModelWeak.Reset();
+	MaterialModelBaseWeak.Reset();
 	ObjectProperty.Reset();
 
 	SlotPickerContainer.Reset();
@@ -319,6 +328,9 @@ void SDMEditor::ClearEditor()
 	SplitterContainer.Reset();
 
 	bInvalidateComponentEditWidget = false;
+
+	PreviewMaterials.Empty();
+	PreviewMaterialDynamics.Empty();
 }
 
 void SDMEditor::ResetEditor()
@@ -339,7 +351,7 @@ void SDMEditor::SetActiveSlotIndex(int InSlotIndex)
 	}
 	else if (InSlotIndex != ActiveSlotIndex)
 	{
-		if (UDynamicMaterialModelEditorOnlyData* ModelEditorOnlyData = UDynamicMaterialModelEditorOnlyData::Get(MaterialModelWeak))
+		if (UDynamicMaterialModelEditorOnlyData* ModelEditorOnlyData = UDynamicMaterialModelEditorOnlyData::Get(MaterialModelBaseWeak))
 		{
 			if (ModelEditorOnlyData->GetSlots().IsValidIndex(InSlotIndex))
 			{
@@ -437,27 +449,32 @@ FDMPropertyHandle SDMEditor::CreatePropertyHandle(const void* InOwningWidget, UO
 	return PropertyHandle;
 }
 
-void SDMEditor::SetMaterialModel(UDynamicMaterialModel* InMaterialModel)
+void SDMEditor::SetMaterialModelBase(UDynamicMaterialModelBase* InMaterialModelBase)
 {
-	if (InMaterialModel && !UE::DynamicMaterialEditor::Private::CheckMaterialModelValidity(InMaterialModel))
+	if (GetMaterialModelBase() == InMaterialModelBase)
 	{
-		InMaterialModel = nullptr;
+		return;
 	}
 
-	if (MaterialModelWeak.Get() == InMaterialModel)
+	if (InMaterialModelBase && !UE::DynamicMaterialEditor::Private::CheckMaterialModelValidity(InMaterialModelBase))
 	{
-		return;
+		InMaterialModelBase = nullptr;
 	}
 
 	ClearEditor();
 
-	MaterialModelWeak = InMaterialModel;
+	MaterialModelBaseWeak = InMaterialModelBase;
+
+	if (UDynamicMaterialModelDynamic* MaterialModelDynamic = Cast<UDynamicMaterialModelDynamic>(InMaterialModelBase))
+	{
+		MaterialModelDynamic->EnsureComponents();
+	}
 
-	Toolbar->SetMaterialModel(InMaterialModel);
+	Toolbar->OnMaterialModelChanged();
 
 	SetEditedComponent(nullptr);
 
-	if (!InMaterialModel)
+	if (!InMaterialModelBase)
 	{
 		SetEmptyLayout();
 		return;
@@ -465,7 +482,7 @@ void SDMEditor::SetMaterialModel(UDynamicMaterialModel* InMaterialModel)
 
 	bHasActiveLayout = true;
 
-	UDynamicMaterialModelEditorOnlyData* EditorOnlyData = UDynamicMaterialModelEditorOnlyData::Get(InMaterialModel);
+	UDynamicMaterialModelEditorOnlyData* EditorOnlyData = UDynamicMaterialModelEditorOnlyData::Get(InMaterialModelBase);
 
 	if (EditorOnlyData && EditorOnlyData->NeedsWizard())
 	{
@@ -484,6 +501,16 @@ void SDMEditor::SetMaterialModel(UDynamicMaterialModel* InMaterialModel)
 	}
 }
 
+bool SDMEditor::IsDynamicModel() const
+{
+	if (UDynamicMaterialModelBase* MaterialModelBase = GetMaterialModelBase())
+	{
+		return !MaterialModelBase->IsA<UDynamicMaterialModel>();
+	}
+
+	return false;
+}
+
 void SDMEditor::SetMaterialObjectProperty(const FDMObjectMaterialProperty& InObjectProperty)
 {
 	if (!InObjectProperty.IsValid())
@@ -492,15 +519,15 @@ void SDMEditor::SetMaterialObjectProperty(const FDMObjectMaterialProperty& InObj
 		return;
 	}
 
-	UDynamicMaterialModel* MaterialModel = InObjectProperty.GetMaterialModel();
+	UDynamicMaterialModelBase* MaterialModelBase = InObjectProperty.GetMaterialModelBase();
 
-	if (!IsValid(MaterialModel))
+	if (!IsValid(MaterialModelBase))
 	{
 		ResetEditor();
 		return;
 	}
 
-	SetMaterialModel(MaterialModel);
+	SetMaterialModelBase(MaterialModelBase);
 	ObjectProperty = InObjectProperty;
 
 	if (AActor* Actor = ObjectProperty.GetTypedOuter<AActor>())
@@ -526,7 +553,7 @@ void SDMEditor::SetMaterialActor(AActor* InActor)
 		return;
 	}
 
-	if (!GetMaterialModel())
+	if (!GetMaterialModelBase())
 	{
 		Container->SetContent(
 			SNew(SBox)
@@ -545,17 +572,19 @@ void SDMEditor::SetEmptyLayout()
 {
 	bHasActiveLayout = false;
 	Container->SetContent(SDMEditor::GetEmptyContent());
+	PreviewMaterials.Empty();
+	PreviewMaterialDynamics.Empty();
 }
 
-void SDMEditor::OnMaterialModelSelected(UDynamicMaterialModel* InMaterialModel)
+void SDMEditor::OnMaterialModelSelected(UDynamicMaterialModelBase* InMaterialModelBase)
 {
-	if (!IsValid(InMaterialModel))
+	if (!IsValid(InMaterialModelBase))
 	{
 		return;
 	}
 
 	// Only check this setting if we have an active material model.
-	if (GetMaterialModel())
+	if (GetMaterialModelBase())
 	{
 		UDynamicMaterialEditorSettings* Settings = UDynamicMaterialEditorSettings::Get();
 
@@ -565,7 +594,7 @@ void SDMEditor::OnMaterialModelSelected(UDynamicMaterialModel* InMaterialModel)
 		}
 	}
 
-	SetMaterialModel(InMaterialModel);
+	SetMaterialModelBase(InMaterialModelBase);
 }
 
 void SDMEditor::OnMaterialInstanceSelected(UDynamicMaterialInstance* InMaterialInstance)
@@ -581,7 +610,7 @@ void SDMEditor::OnMaterialInstanceSelected(UDynamicMaterialInstance* InMaterialI
 void SDMEditor::OnActorSelected(AActor* InActor)
 {
 	// Only check this setting if we have an active material model.
-	if (GetMaterialModel())
+	if (GetMaterialModelBase())
 	{
 		UDynamicMaterialEditorSettings* Settings = UDynamicMaterialEditorSettings::Get();
 
@@ -610,14 +639,14 @@ void SDMEditor::OnActorSelected(AActor* InActor)
 	{
 		const FDMObjectMaterialProperty& MaterialProperty = ActorProperties[MaterialPropertyIdx];
 
-		if (MaterialProperty.GetMaterialModel())
+		if (MaterialProperty.GetMaterialModelBase())
 		{
 			SetMaterialObjectProperty(MaterialProperty);
 			break;
 		}
 	}
 
-	if (!ActorProperties.IsEmpty() && !GetMaterialModel())
+	if (!ActorProperties.IsEmpty() && !GetMaterialModelBase())
 	{
 		SetMaterialObjectProperty(ActorProperties[0]);
 	}
@@ -835,19 +864,11 @@ TSharedRef<SWidget> SDMEditor::CreateMainLayout()
 		.VAlign(VAlign_Top)
 		[
 			SAssignNew(Toolbar, SDMToolBar, SharedThis(this))
-			.MaterialModel(MaterialModelWeak.Get())
 			.OnSlotChanged(this, &SDMEditor::OnToolBarPropertyChanged)
 			.OnGetSettingsMenu(this, &SDMEditor::MakeToolBarSettingsMenu)
 		]
 		+ SVerticalBox::Slot()
 		.AutoHeight()
-		[
-			UE::DynamicMaterialEditor::bGlobalValuesEnabled && MaterialModelWeak.IsValid()
-				? CreateParametersArea()
-				: SNullWidget::NullWidget
-		]
-		+ SVerticalBox::Slot()
-		.AutoHeight()
 		.HAlign(HAlign_Fill)
 		.VAlign(VAlign_Top)
 		.Padding(5.f, 0.f)
@@ -916,28 +937,6 @@ TSharedRef<SWidget> SDMEditor::CreateMainLayout()
 		];
 }
 
-TSharedRef<SWidget> SDMEditor::CreateParametersArea()
-{
-	return 
-		SNew(SExpandableArea)
-		.InitiallyCollapsed(false)
-		.HeaderPadding(FMargin(3.0f, 5.0f, 3.0f, 5.0f))
-		.HeaderContent()
-		[
-			SNew(SBox)
-			.VAlign(EVerticalAlignment::VAlign_Center)
-			[
-				SNew(STextBlock)
-				.Font(IDetailLayoutBuilder::GetDetailFont())
-				.Text(LOCTEXT("MaterialParameters", "Material Parameters"))
-			]
-		]
-		.BodyContent()
-		[
-			SAssignNew(ParametersWidget, SDMMaterialParameters, MaterialModelWeak)
-		];
-}
-
 TSharedRef<SWidget> SDMEditor::CreateSlotPickerWidget()
 {
 	TSharedRef<SWrapBox> SlotSelector = SNew(SWrapBox)
@@ -945,7 +944,7 @@ TSharedRef<SWidget> SDMEditor::CreateSlotPickerWidget()
 		.InnerSlotPadding(FVector2D(5, 5))
 		.Orientation(EOrientation::Orient_Horizontal);
 
-	UDynamicMaterialModelEditorOnlyData* ModelEditorOnlyData = UDynamicMaterialModelEditorOnlyData::Get(MaterialModelWeak);
+	UDynamicMaterialModelEditorOnlyData* ModelEditorOnlyData = UDynamicMaterialModelEditorOnlyData::Get(MaterialModelBaseWeak);
 
 	if (!ModelEditorOnlyData)
 	{
@@ -1021,14 +1020,6 @@ TSharedRef<SWidget> SDMEditor::CreateSlotPickerWidget()
 	return SlotSelector;
 }
 
-void SDMEditor::RefreshParametersList()
-{
-	if (ParametersWidget.IsValid())
-	{
-		ParametersWidget->RefreshWidgets();
-	}
-}
-
 void SDMEditor::RefreshSlotPickerList()
 {
 	if (SlotPickerContainer.IsValid())
@@ -1060,9 +1051,9 @@ void SDMEditor::Tick(const FGeometry& AllottedGeometry, const double InCurrentTi
 		return;
 	}
 
-	UDynamicMaterialModel* MaterialModel = MaterialModelWeak.Get();
+	UDynamicMaterialModelBase* MaterialModelBase = GetMaterialModelBase();
 
-	if (!IsValid(MaterialModel))
+	if (!IsValid(MaterialModelBase))
 	{
 		if (bHasActiveLayout)
 		{
@@ -1079,16 +1070,16 @@ void SDMEditor::Tick(const FGeometry& AllottedGeometry, const double InCurrentTi
 
 	if (ObjectProperty.IsValid())
 	{
-		UDynamicMaterialModel* MaterialModelProperty = ObjectProperty.GetMaterialModel();
+		UDynamicMaterialModelBase* MaterialModelBaseProperty = ObjectProperty.GetMaterialModelBase();
 
-		if (!IsValid(MaterialModelProperty))
+		if (!IsValid(MaterialModelBaseProperty))
 		{
-			MaterialModelProperty = nullptr;
+			MaterialModelBaseProperty = nullptr;
 		}
 
-		if (MaterialModel != MaterialModelProperty)
+		if (MaterialModelBase != MaterialModelBaseProperty)
 		{
-			if (MaterialModelProperty)
+			if (MaterialModelBaseProperty)
 			{
 				// Re-set the property
 				SetMaterialObjectProperty(ObjectProperty);
@@ -1101,7 +1092,7 @@ void SDMEditor::Tick(const FGeometry& AllottedGeometry, const double InCurrentTi
 			}
 		}
 	}
-	else if (!UE::DynamicMaterialEditor::Private::CheckMaterialModelValidity(MaterialModel))
+	else if (!UE::DynamicMaterialEditor::Private::CheckMaterialModelValidity(MaterialModelBase))
 	{
 		ResetEditor();
 		return;
@@ -1109,11 +1100,13 @@ void SDMEditor::Tick(const FGeometry& AllottedGeometry, const double InCurrentTi
 
 	if (Toolbar.IsValid())
 	{
+		/*
 		UDynamicMaterialModel* ToolbarMaterialModel = Toolbar->GetMaterialModel();
 		if (ToolbarMaterialModel != MaterialModel)
 		{
-			Toolbar->SetMaterialModel(MaterialModel);
+			Toolbar->OnMaterialModelChanged(GetMaterialModel());
 		}
+		*/
 
 		if (ObjectProperty.IsValid())
 		{
@@ -1126,13 +1119,13 @@ void SDMEditor::Tick(const FGeometry& AllottedGeometry, const double InCurrentTi
 		}
 		else if (Toolbar->GetMaterialActor())
 		{
-			if (!ToolbarMaterialModel)
+			if (!MaterialModelBase)
 			{
 				Toolbar->SetMaterialActor(nullptr);
 			}
 			else
 			{
-				AActor* Actor = ToolbarMaterialModel->GetTypedOuter<AActor>();
+				AActor* Actor = MaterialModelBase->GetTypedOuter<AActor>();
 
 				if (Actor != Toolbar->GetMaterialActor())
 				{
@@ -1153,32 +1146,24 @@ void SDMEditor::Tick(const FGeometry& AllottedGeometry, const double InCurrentTi
 	}
 }
 
-void SDMEditor::OnMaterialBuilt(UDynamicMaterialModel* InMaterialModel)
+void SDMEditor::OnMaterialBuilt(UDynamicMaterialModelBase* InMaterialModelBase)
 {
 	RefreshSlotPickerList();
 	RefreshComponentEditWidget();
 }
 
-void SDMEditor::OnValuesUpdated(UDynamicMaterialModel* InMaterialModel)
+void SDMEditor::OnValuesUpdated(UDynamicMaterialModelBase* InMaterialModelBase)
 {
-	RefreshParametersList();
 }
 
-void SDMEditor::OnSlotsUpdated(UDynamicMaterialModel* InMaterialModel)
+void SDMEditor::OnSlotsUpdated(UDynamicMaterialModelBase* InMaterialModelBase)
 {
 	RefreshSlotWidget();
 }
 
 bool SDMEditor::IsPropertyValidForModel(EDMMaterialPropertyType InProperty) const
 {
-	UDynamicMaterialModel* MaterialModel = MaterialModelWeak.Get();
-
-	if (!MaterialModel)
-	{
-		return false;
-	}
-
-	UDynamicMaterialModelEditorOnlyData* EditorOnlyData = UDynamicMaterialModelEditorOnlyData::Get(MaterialModel);
+	UDynamicMaterialModelEditorOnlyData* EditorOnlyData = UDynamicMaterialModelEditorOnlyData::Get(MaterialModelBaseWeak);
 
 	if (!EditorOnlyData)
 	{
@@ -1218,7 +1203,7 @@ ECheckBoxState SDMEditor::GetSlotCheckState(EDMMaterialPropertyType InProperty)
 		return ECheckBoxState::Unchecked;
 	}
 
-	UDynamicMaterialModelEditorOnlyData* EditorOnlyData = UDynamicMaterialModelEditorOnlyData::Get(MaterialModelWeak);
+	UDynamicMaterialModelEditorOnlyData* EditorOnlyData = UDynamicMaterialModelEditorOnlyData::Get(MaterialModelBaseWeak);
 
 	if (!EditorOnlyData)
 	{
@@ -1260,7 +1245,7 @@ void SDMEditor::OnSlotCheckStateChanged(ECheckBoxState InCheckState, EDMMaterial
 		return;
 	}
 
-	UDynamicMaterialModelEditorOnlyData* ModelEditorOnlyData = UDynamicMaterialModelEditorOnlyData::Get(MaterialModelWeak);
+	UDynamicMaterialModelEditorOnlyData* ModelEditorOnlyData = UDynamicMaterialModelEditorOnlyData::Get(MaterialModelBaseWeak);
 
 	if (!ModelEditorOnlyData)
 	{
@@ -1299,7 +1284,7 @@ FReply SDMEditor::OnCreateMaterialButtonClicked(FDMObjectMaterialProperty InMate
 	{
 		if (Toolbar.IsValid())
 		{
-			Toolbar->SetMaterialModel(NewModel);
+			Toolbar->OnMaterialModelChanged();
 		}
 
 		if (AActor* Actor = InMaterialProperty.GetTypedOuter<AActor>())
@@ -1326,13 +1311,14 @@ TSharedRef<SWidget> SDMEditor::MakeToolBarSettingsMenu()
 
 void SDMEditor::OnSettingsChanged(const FPropertyChangedEvent& InPropertyChangedEvent)
 {
+	UDynamicMaterialModelBase* MaterialModelBase = GetMaterialModelBase();
 	ClearEditor();
-	SetMaterialModel(MaterialModelWeak.Get());
+	SetMaterialModelBase(MaterialModelBase);
 }
 
 bool SDMEditor::CanAddNewLayer() const
 {
-	return ActiveSlotWidget.IsValid() && !!UDynamicMaterialModelEditorOnlyData::Get(MaterialModelWeak);
+	return ActiveSlotWidget.IsValid() && !!UDynamicMaterialModelEditorOnlyData::Get(MaterialModelBaseWeak);
 }
 
 void SDMEditor::AddNewLayer()
@@ -1352,7 +1338,7 @@ bool SDMEditor::CanInsertNewLayer() const
 {
 	return ActiveSlotWidget.IsValid()
 		&& ActiveSlotWidget->GetSelectedLayerIndices().Num() == 1
-		&& MaterialModelWeak.IsValid();
+		&& MaterialModelBaseWeak.IsValid();
 }
 
 void SDMEditor::InsertNewLayer()
@@ -1384,7 +1370,7 @@ bool SDMEditor::CanCopySelectedLayer() const
 {
 	return ActiveSlotWidget.IsValid()
 		&& ActiveSlotWidget->GetSelectedLayerIndices().Num() == 1
-		&& MaterialModelWeak.IsValid();
+		&& MaterialModelBaseWeak.IsValid();
 }
 
 void SDMEditor::CopySelectedLayer()
@@ -1411,7 +1397,7 @@ void SDMEditor::CutSelectedLayer()
 
 bool SDMEditor::CanPasteLayer() const
 {
-	if (!MaterialModelWeak.IsValid())
+	if (!MaterialModelBaseWeak.IsValid())
 	{
 		return false;
 	}
@@ -1424,7 +1410,7 @@ bool SDMEditor::CanPasteLayer() const
 
 void SDMEditor::PasteLayer()
 {
-	if (UDynamicMaterialModelEditorOnlyData* ModelEditorOnlyData = UDynamicMaterialModelEditorOnlyData::Get(MaterialModelWeak))
+	if (UDynamicMaterialModelEditorOnlyData* ModelEditorOnlyData = UDynamicMaterialModelEditorOnlyData::Get(MaterialModelBaseWeak))
 	{
 		if (UDMMaterialSlot* Slot = ModelEditorOnlyData->GetSlot(ActiveSlotIndex))
 		{
@@ -1481,6 +1467,77 @@ void SDMEditor::DeleteSelectedLayer()
 	}
 }
 
+UMaterial* SDMEditor::CreatePreviewMaterial(UObject* InPreviewing)
+{
+	UMaterial* PreviewMaterial = nullptr;
+
+	if (FDynamicMaterialModule::IsMaterialExportEnabled() == false)
+	{
+		PreviewMaterial = Cast<UMaterial>(GetMutableDefault<UMaterialFactoryNew>()->FactoryCreateNew(
+			UMaterial::StaticClass(),
+			GetTransientPackage(),
+			NAME_None,
+			RF_Transient,
+			nullptr,
+			GWarn
+		));
+
+		PreviewMaterial->bIsPreviewMaterial = true;
+	}
+	else
+	{
+		FString MaterialBaseName = InPreviewing->GetName() + "-" + FGuid::NewGuid().ToString();
+		const FString FullName = "/Game/MaterialDesignerMaterials/" + MaterialBaseName;
+		UPackage* Package = CreatePackage(*FullName);
+
+		PreviewMaterial = Cast<UMaterial>(GetMutableDefault<UMaterialFactoryNew>()->FactoryCreateNew(
+			UMaterial::StaticClass(),
+			Package,
+			*MaterialBaseName,
+			RF_Standalone | RF_Public,
+			nullptr,
+			GWarn
+		));
+
+		FAssetRegistryModule::AssetCreated(PreviewMaterial);
+	}
+
+	PreviewMaterials.FindOrAdd(InPreviewing).Reset(PreviewMaterial);
+
+	return PreviewMaterial;
+}
+
+void SDMEditor::FreePreviewMaterial(UObject* InPreviewing)
+{
+	 if (TStrongObjectPtr<UMaterial>* PreviewMaterial = PreviewMaterials.Find(InPreviewing))
+	 {
+		if (UMaterial* Material = PreviewMaterial->Get())
+		{
+			FreeMID(Material);
+		}
+
+		PreviewMaterials.Remove(InPreviewing);
+	 }
+}
+
+UMaterialInstanceDynamic* SDMEditor::CreateMID(UMaterial* InMaterialBase)
+{
+	if (!InMaterialBase)
+	{
+		return nullptr;
+	}
+
+	UMaterialInstanceDynamic* MID = UMaterialInstanceDynamic::Create(InMaterialBase, GetTransientPackage());
+	PreviewMaterialDynamics.FindOrAdd(InMaterialBase).Reset(MID);
+
+	return MID;
+}
+
+void SDMEditor::FreeMID(UMaterial* InMaterialBase)
+{
+	PreviewMaterialDynamics.Remove(InMaterialBase);
+}
+
 bool SDMEditor::GetExpansionState(UObject* InOwner, FName InName, bool& bOutExpanded)
 {
 	const FExpansionItem ExpansionItem = {InOwner, InName};
@@ -1561,7 +1618,7 @@ TSharedRef<SWidget> SDMEditor::CreateSlotWidget()
 		return CreateEmptySlotsContent();
 	}
 	
-	UDynamicMaterialModelEditorOnlyData* ModelEditorOnlyData = UDynamicMaterialModelEditorOnlyData::Get(MaterialModelWeak);
+	UDynamicMaterialModelEditorOnlyData* ModelEditorOnlyData = UDynamicMaterialModelEditorOnlyData::Get(MaterialModelBaseWeak);
 
 	if (!ModelEditorOnlyData)
 	{
@@ -1585,14 +1642,14 @@ TSharedRef<SWidget> SDMEditor::CreateSlotWidget()
 
 void SDMEditor::OnUndo()
 {
-	UDynamicMaterialModel* MaterialModel = MaterialModelWeak.Get();
+	UDynamicMaterialModel* MaterialModel = GetMaterialModel();
 
 	if (!IsValid(MaterialModel))
 	{
 		return;
 	}
 
-	if (UDynamicMaterialModelEditorOnlyData* ModelEditorOnlyData = UDynamicMaterialModelEditorOnlyData::Get(MaterialModelWeak))
+	if (UDynamicMaterialModelEditorOnlyData* ModelEditorOnlyData = UDynamicMaterialModelEditorOnlyData::Get(MaterialModelBaseWeak))
 	{
 		if (ModelEditorOnlyData->GetBlendMode() == BLEND_Opaque)
 		{
@@ -1603,7 +1660,7 @@ void SDMEditor::OnUndo()
 
 UDMMaterialSlot* SDMEditor::GetSlotForMaterialProperty(EDMMaterialPropertyType InProperty) const
 {
-	UDynamicMaterialModelEditorOnlyData* ModelEditorOnlyData = UDynamicMaterialModelEditorOnlyData::Get(MaterialModelWeak);
+	UDynamicMaterialModelEditorOnlyData* ModelEditorOnlyData = UDynamicMaterialModelEditorOnlyData::Get(MaterialModelBaseWeak);
 
 	if (!ModelEditorOnlyData)
 	{
diff --git a/Engine/Plugins/Experimental/DynamicMaterial/Source/DynamicMaterialEditor/Private/Slate/SDMEditor.h b/Engine/Plugins/Experimental/DynamicMaterial/Source/DynamicMaterialEditor/Private/Slate/SDMEditor.h
--- a/Engine/Plugins/Experimental/DynamicMaterial/Source/DynamicMaterialEditor/Private/Slate/SDMEditor.h
+++ b/Engine/Plugins/Experimental/DynamicMaterial/Source/DynamicMaterialEditor/Private/Slate/SDMEditor.h
@@ -17,7 +17,6 @@ class IPropertyHandle;
 class IPropertyRowGenerator;
 class IStructureDetailsView;
 class SBox;
-class SDMMaterialParameters;
 class SDMPropertyEdit;
 class SDMSlot;
 class SDMToolBar;
@@ -30,6 +29,9 @@ class UDMMaterialStage;
 class UDMMaterialStageExpression;
 class UDMMaterialValueFloat1;
 class UDynamicMaterialModel;
+class UDynamicMaterialModelBase;
+class UMaterial;
+class UMaterialInstanceDynamic;
 enum class ECheckBoxState : uint8;
 enum class EDMExpressionMenu : uint8;
 enum EMaterialDomain : int;
@@ -54,12 +56,16 @@ public:
 
 	static void ClearPropertyHandles(const SWidget* InOwningWidget);
 
-	void Construct(const FArguments& InArgs, TWeakObjectPtr<UDynamicMaterialModel> InModelWeak);
+	void Construct(const FArguments& InArgs);
 
 	virtual ~SDMEditor() override;
 
-	UDynamicMaterialModel* GetMaterialModel() const { return MaterialModelWeak.Get(); }
-	void SetMaterialModel(UDynamicMaterialModel* InMaterialModel);
+	UDynamicMaterialModelBase* GetMaterialModelBase() const { return MaterialModelBaseWeak.Get(); }
+	void SetMaterialModelBase(UDynamicMaterialModelBase* InMaterialModelBase);
+
+	UDynamicMaterialModel* GetMaterialModel() const;
+
+	bool IsDynamicModel() const;
 
 	const FDMObjectMaterialProperty& GetMaterialObjectProperty() const { return ObjectProperty; }
 	void SetMaterialObjectProperty(const FDMObjectMaterialProperty& InObjectProperty);
@@ -67,7 +73,7 @@ public:
 	AActor* GetMaterialActor() const;
 	void SetMaterialActor(AActor* InActor);
 
-	void OnMaterialModelSelected(UDynamicMaterialModel* InMaterialModel);
+	void OnMaterialModelSelected(UDynamicMaterialModelBase* InMaterialModelBase);
 
 	void OnMaterialInstanceSelected(UDynamicMaterialInstance* InMaterialInstance);
 
@@ -75,7 +81,6 @@ public:
 
 	TSharedPtr<SDMSlot> GetActiveSlotWidget() const;
 
-	void RefreshParametersList();
 	void RefreshSlotPickerList();
 	void RefreshSlotWidget();
 	void RefreshComponentEditWidget();
@@ -86,11 +91,6 @@ public:
 
 	void SetEditedComponent(UDMMaterialComponent* InComponent);
 
-	//~ Begin SWidget
-	virtual bool SupportsKeyboardFocus() const override { return true; }
-	virtual void Tick(const FGeometry& AllottedGeometry, const double InCurrentTime, const float InDeltaTime) override;
-	//~ End SWidget
-
 	void ClearEditor();
 
 	void ResetEditor();
@@ -121,7 +121,15 @@ public:
 	bool CanDeleteSelectedLayer() const;
 	void DeleteSelectedLayer();
 
+	UMaterial* CreatePreviewMaterial(UObject* InPreviewing);
+	void FreePreviewMaterial(UObject* InPreviewing);
+
+	UMaterialInstanceDynamic* CreateMID(UMaterial* InMaterialBase);
+	void FreeMID(UMaterial* InMaterialBase);
+
 	//~ Begin SWidget
+	virtual bool SupportsKeyboardFocus() const override { return true; }
+	virtual void Tick(const FGeometry& AllottedGeometry, const double InCurrentTime, const float InDeltaTime) override;
 	virtual FReply OnKeyDown(const FGeometry& MyGeometry, const FKeyEvent& InKeyEvent) override;
 	//~ End SWidget
 
@@ -161,7 +169,6 @@ protected:
 
 	TSharedPtr<SBox> Container;
 	TSharedPtr<SDMToolBar> Toolbar;
-	TSharedPtr<SDMMaterialParameters> ParametersWidget;
 	TSharedPtr<SBox> SlotPickerContainer;
 	TSharedPtr<SSplitter> SplitterContainer;
 	TSharedPtr<SScrollBox> SlotContainer;
@@ -173,15 +180,17 @@ protected:
 
 	TSharedPtr<FUICommandList> CommandList;
 	int32 ActiveSlotIndex;
-	TWeakObjectPtr<UDynamicMaterialModel> MaterialModelWeak;
+	TWeakObjectPtr<UDynamicMaterialModelBase> MaterialModelBaseWeak;
 	FDMObjectMaterialProperty ObjectProperty;
 	TWeakObjectPtr<UDMMaterialComponent> EditedComponent;
 	bool bInvalidateComponentEditWidget;
 
+	TMap<FObjectKey, TStrongObjectPtr<UMaterial>> PreviewMaterials;
+	TMap<FObjectKey, TStrongObjectPtr<UMaterialInstanceDynamic>> PreviewMaterialDynamics;
+
 	void BindCommands();
 
 	TSharedRef<SWidget> CreateMainLayout();
-	TSharedRef<SWidget> CreateParametersArea();
 	TSharedRef<SWidget> CreateSlotPickerWidget();
 	TSharedRef<SWidget> CreateSlotWidget();
 	TSharedRef<SWidget> CreateComponentEditWidget();
@@ -192,9 +201,9 @@ protected:
 
 	TSharedRef<SWidget> CreateActorMaterialSlotSelector(AActor* InActor);
 
-	void OnMaterialBuilt(UDynamicMaterialModel* InMaterialModel);
-	void OnValuesUpdated(UDynamicMaterialModel* InMaterialModel);
-	void OnSlotsUpdated(UDynamicMaterialModel* InMaterialModel);
+	void OnMaterialBuilt(UDynamicMaterialModelBase* InMaterialModelBase);
+	void OnValuesUpdated(UDynamicMaterialModelBase* InMaterialModelBase);
+	void OnSlotsUpdated(UDynamicMaterialModelBase* InMaterialModelBase);
 
 	bool IsPropertyValidForModel(EDMMaterialPropertyType InProperty) const;
 	ECheckBoxState GetSlotCheckState(EDMMaterialPropertyType InProperty) const;
diff --git a/Engine/Plugins/Experimental/DynamicMaterial/Source/DynamicMaterialEditor/Private/Slate/SDMMaterialWizard.cpp b/Engine/Plugins/Experimental/DynamicMaterial/Source/DynamicMaterialEditor/Private/Slate/SDMMaterialWizard.cpp
--- a/Engine/Plugins/Experimental/DynamicMaterial/Source/DynamicMaterialEditor/Private/Slate/SDMMaterialWizard.cpp
+++ b/Engine/Plugins/Experimental/DynamicMaterial/Source/DynamicMaterialEditor/Private/Slate/SDMMaterialWizard.cpp
@@ -267,7 +267,7 @@ FReply SDMMaterialWizard::Accept_OnClick()
 				}
 				else
 				{
-					Editor->SetMaterialModel(MaterialModel);
+					Editor->SetMaterialModelBase(MaterialModel);
 				}
 
 				Editor->SetMaterialActor(MaterialActor);
diff --git a/Engine/Plugins/Experimental/DynamicMaterial/Source/DynamicMaterialEditor/Private/Slate/SDMSlot.cpp b/Engine/Plugins/Experimental/DynamicMaterial/Source/DynamicMaterialEditor/Private/Slate/SDMSlot.cpp
--- a/Engine/Plugins/Experimental/DynamicMaterial/Source/DynamicMaterialEditor/Private/Slate/SDMSlot.cpp
+++ b/Engine/Plugins/Experimental/DynamicMaterial/Source/DynamicMaterialEditor/Private/Slate/SDMSlot.cpp
@@ -97,6 +97,8 @@ void SDMSlot::Construct(const FArguments& InArgs, const TSharedRef<SDMEditor>& I
 		return;
 	}
 
+	bIsDynamic = !Cast<UDynamicMaterialModel>(InEditor->GetMaterialModelBase());
+
 	InSlot->GetOnPropertiesUpdateDelegate().AddSP(this, &SDMSlot::OnSlotPropertiesUpdated);
 	InSlot->GetOnLayersUpdateDelegate().AddSP(this, &SDMSlot::OnSlotLayersUpdated);
 
@@ -436,13 +438,13 @@ TSharedRef<SWidget> SDMSlot::CreateHeaderPropertyListWidget()
 		return SNullWidget::NullWidget;
 	}
 
-	UDynamicMaterialModel* MaterialModel = MaterialEditor->GetMaterialModel();
-	if (!ensure(IsValid(MaterialModel)))
+	UDynamicMaterialModelBase* MaterialModelBase = MaterialEditor->GetMaterialModelBase();
+	if (!ensure(IsValid(MaterialModelBase)))
 	{
 		return SNullWidget::NullWidget;
 	}
 
-	UDynamicMaterialModelEditorOnlyData* ModelEditorOnlyData = UDynamicMaterialModelEditorOnlyData::Get(MaterialModel);
+	UDynamicMaterialModelEditorOnlyData* ModelEditorOnlyData = UDynamicMaterialModelEditorOnlyData::Get(MaterialModelBase);
 	if (!ensure(IsValid(ModelEditorOnlyData)))
 	{
 		return SNullWidget::NullWidget;
@@ -511,6 +513,7 @@ TSharedRef<SWidget> SDMSlot::CreateLayerButtonsRowWidget()
 {
 	return 
 		SNew(SHorizontalBox)
+		.IsEnabled(!bIsDynamic)
 		+ SHorizontalBox::Slot()
 		.FillWidth(1.0f)
 		.HAlign(HAlign_Left)
@@ -651,6 +654,7 @@ TSharedRef<SWidget> SDMSlot::CreateSlotSettingsRow()
 		.Padding(5.0f, 0.0f, 0.0f, 0.0f)
 		[
 			SNew(SHorizontalBox)
+			.IsEnabled(!bIsDynamic)
 			.ToolTipText(LOCTEXT("MaterialDesignerInstanceBlendModeTooltip", "Change the Blend Mode for this Material Designer Instance."))
 			+ SHorizontalBox::Slot()
 			.AutoWidth()
diff --git a/Engine/Plugins/Experimental/DynamicMaterial/Source/DynamicMaterialEditor/Private/Slate/SDMSlot.h b/Engine/Plugins/Experimental/DynamicMaterial/Source/DynamicMaterialEditor/Private/Slate/SDMSlot.h
--- a/Engine/Plugins/Experimental/DynamicMaterial/Source/DynamicMaterialEditor/Private/Slate/SDMSlot.h
+++ b/Engine/Plugins/Experimental/DynamicMaterial/Source/DynamicMaterialEditor/Private/Slate/SDMSlot.h
@@ -91,6 +91,7 @@ public:
 
 protected:
 	TWeakPtr<SDMEditor> EditorWidgetWeak;
+	bool bIsDynamic;
 
 	TAttribute<int32> SlotPreviewSize;
 	TAttribute<int32> LayerPreviewSize;
diff --git a/Engine/Plugins/Experimental/DynamicMaterial/Source/DynamicMaterialEditor/Private/Slate/SDMStage.cpp b/Engine/Plugins/Experimental/DynamicMaterial/Source/DynamicMaterialEditor/Private/Slate/SDMStage.cpp
--- a/Engine/Plugins/Experimental/DynamicMaterial/Source/DynamicMaterialEditor/Private/Slate/SDMStage.cpp
+++ b/Engine/Plugins/Experimental/DynamicMaterial/Source/DynamicMaterialEditor/Private/Slate/SDMStage.cpp
@@ -19,8 +19,7 @@
 #include "DynamicMaterialEditorStyle.h"
 #include "Engine/Texture.h"
 #include "SDMEditor.h"
-#include "Slate/Previews/SDMStagePreview.h"
-#include "Slate/SMaterialToolTip.h"
+#include "Slate/Previews/SDMComponentPreview.h"
 #include "SlateOptMacros.h"
 #include "Styling/CoreStyle.h"
 #include "Utils/DMPrivate.h"
@@ -38,9 +37,16 @@ SDMStage::~SDMStage()
 	SDMEditor::ClearPropertyHandles(this);
 }
 
-void SDMStage::Construct(const FArguments& InArgs, UDMMaterialStage* InStage)
+void SDMStage::Construct(const FArguments& InArgs, const TSharedRef<SDMSlot> InSlotWidget, UDMMaterialStage* InStage)
 {
 	ensure(IsValid(InStage));
+
+	if (!ensure(InSlotWidget->GetEditorWidget().IsValid()))
+	{
+		return;
+	}
+
+	SlotWidgetWeak = InSlotWidget;
 	StageWeak = InStage;
 
 	bStageEnabled = InArgs._StageEnabled;
@@ -74,11 +80,8 @@ void SDMStage::Construct(const FArguments& InArgs, UDMMaterialStage* InStage)
 		.IsInteractive(false)
 		.BorderImage(FCoreStyle::Get().GetBrush("ToolTip.Background"))
 		[
-			SNew(SMaterialToolTip)
-			.Material(StageWeak->GetPreviewMaterial())
-			.Text(ToolTipText)
-			.MaterialSize_Lambda([]() { return FVector2D(UDynamicMaterialEditorSettings::Get()->TooltipTextureSize); })
-			.ShowMaterial(this, &SDMStage::IsToolTipInteractable)
+			SNew(SDMComponentPreview, InSlotWidget->GetEditorWidget().ToSharedRef(), StageWeak.Get())
+			.PreviewSize(FVector2D(256.f))
 		]
 	);
 
@@ -106,7 +109,7 @@ void SDMStage::Construct(const FArguments& InArgs, UDMMaterialStage* InStage)
 					.WidthOverride(this, &SDMStage::GetStagePreviewDesiredWidth)
 					.HeightOverride(this, &SDMStage::GetStagePreviewDesiredHeight)
 					[
-						SNew(SDMStagePreview, StageWeak.Get())
+						SNew(SDMComponentPreview, InSlotWidget->GetEditorWidget().ToSharedRef(), StageWeak.Get())
 						.PreviewSize(DesiredSize)
 					]
 				]
diff --git a/Engine/Plugins/Experimental/DynamicMaterial/Source/DynamicMaterialEditor/Private/Slate/SDMStage.h b/Engine/Plugins/Experimental/DynamicMaterial/Source/DynamicMaterialEditor/Private/Slate/SDMStage.h
--- a/Engine/Plugins/Experimental/DynamicMaterial/Source/DynamicMaterialEditor/Private/Slate/SDMStage.h
+++ b/Engine/Plugins/Experimental/DynamicMaterial/Source/DynamicMaterialEditor/Private/Slate/SDMStage.h
@@ -4,6 +4,7 @@
 
 #include "Widgets/SCompoundWidget.h"
 
+class SDMSlot;
 class FAssetDragDropOp;
 class FDMStageDragDropOperation;
 class UDMMaterialStage;
@@ -35,7 +36,7 @@ public:
 		, _DesiredSize(FVector2D(48.0f))
 		, _MinUniformSize(32)
 		, _MaxUniformSize(128)
-	{}
+		{}
 		SLATE_ATTRIBUTE(bool, StageEnabled)
 		SLATE_ATTRIBUTE(bool, StageSelected)
 		SLATE_ARGUMENT(bool, StageIsMask)
@@ -71,7 +72,9 @@ public:
 
 	virtual ~SDMStage() override;
 
-	void Construct(const FArguments& InArgs, UDMMaterialStage* InStage);
+	void Construct(const FArguments& InArgs, const TSharedRef<SDMSlot> InSlotWidget, UDMMaterialStage* InStage);
+
+	TSharedPtr<SDMSlot> GetSlotWidget() const { return SlotWidgetWeak.Pin(); }
 
 	FORCEINLINE UDMMaterialStage* GetStage() const { return StageWeak.Get(); }
 
@@ -96,6 +99,7 @@ public:
 	//~ End SWidget
 
 protected:
+	TWeakPtr<SDMSlot> SlotWidgetWeak;
 	TWeakObjectPtr<UDMMaterialStage> StageWeak;
 
 	TAttribute<bool> bStageEnabled;
diff --git a/Engine/Plugins/Experimental/DynamicMaterial/Source/DynamicMaterialEditor/Private/Slate/SDMToolBar.cpp b/Engine/Plugins/Experimental/DynamicMaterial/Source/DynamicMaterialEditor/Private/Slate/SDMToolBar.cpp
--- a/Engine/Plugins/Experimental/DynamicMaterial/Source/DynamicMaterialEditor/Private/Slate/SDMToolBar.cpp
+++ b/Engine/Plugins/Experimental/DynamicMaterial/Source/DynamicMaterialEditor/Private/Slate/SDMToolBar.cpp
@@ -1,8 +1,12 @@
 // Copyright Epic Games, Inc. All Rights Reserved.
 
 #include "Slate/SDMToolBar.h"
+
+#include "AssetRegistry/AssetRegistryModule.h"
+#include "AssetToolsModule.h"
 #include "ContentBrowserModule.h"
 #include "DMWorldSubsystem.h"
+#include "DynamicMaterialEditorModule.h"
 #include "DynamicMaterialEditorSettings.h"
 #include "DynamicMaterialEditorStyle.h"
 #include "Editor.h"
@@ -11,7 +15,9 @@
 #include "GameFramework/Actor.h"
 #include "IContentBrowserSingleton.h"
 #include "Material/DynamicMaterialInstance.h"
+#include "Material/DynamicMaterialInstanceFactory.h"
 #include "Model/DynamicMaterialModel.h"
+#include "Model/DynamicMaterialModelDynamic.h"
 #include "PackageTools.h"
 #include "SDMEditor.h"
 #include "Selection.h"
@@ -37,8 +43,6 @@ void SDMToolBar::Construct(const FArguments& InArgs, const TSharedRef<SDMEditor>
 	SetCanTick(true);
 
 	EditorWeak = InEditor;
-	MaterialActorWeak = InArgs._MaterialActor;
-	MaterialModelWeak = InArgs._MaterialModel;
 	OnSlotChanged = InArgs._OnSlotChanged;
 	OnGetSettingsMenu = InArgs._OnGetSettingsMenu;
 	
@@ -56,7 +60,7 @@ void SDMToolBar::Construct(const FArguments& InArgs, const TSharedRef<SDMEditor>
 		]
 	];
 
-	SetMaterialModel(MaterialModelWeak.Get());
+	OnMaterialModelChanged();
 	SetMaterialActor(MaterialActorWeak.Get());
 }
 
@@ -82,14 +86,13 @@ TSharedRef<SWidget> SDMToolBar::CreateToolBarEntries()
 			.Padding(5.0f, 0.0f, 0.0f, 0.0f)
 			[
 				SNew(SHorizontalBox)
-				.Visibility(this, &SDMToolBar::GetActorVisibility)
 				+ SHorizontalBox::Slot()
 				.AutoWidth()
 				.Padding(5.0f, 0.0f, 0.0f, 0.0f)
 				.VAlign(EVerticalAlignment::VAlign_Center)
 				[
-					SAssignNew(SaveButton, SButton)
-					.IsEnabled(this, &SDMToolBar::CanSave)
+					SAssignNew(SaveButtonWidget, SButton)
+					.IsEnabled(false)
 					.ContentPadding(GetLargeIconToolBarButtonContentPadding())
 					.ButtonStyle(FDynamicMaterialEditorStyle::Get(), "HoverHintOnly")
 					.ToolTipText(LOCTEXT("MaterialDesignerSaveTooltip", "Save the Material Designer asset\n\nCaution: If this asset lives inside an actor, the actor/level will be saved."))
@@ -109,7 +112,6 @@ TSharedRef<SWidget> SDMToolBar::CreateToolBarEntries()
 					.ContentPadding(GetLargeIconToolBarButtonContentPadding())
 					.ButtonStyle(FDynamicMaterialEditorStyle::Get(), "HoverHintOnly")
 					.ToolTipText(LOCTEXT("ExportMaterialInstance", "Save As"))
-					.Visibility(this, &SDMToolBar::GetExportMaterialInstanceButtonVisibility)
 					.OnClicked(this, &SDMToolBar::OnExportMaterialInstanceButtonClicked)
 					[
 						SNew(SImage)
@@ -117,6 +119,40 @@ TSharedRef<SWidget> SDMToolBar::CreateToolBarEntries()
 						.DesiredSizeOverride(GetLargeIconToolBarButtonSize())
 					]
 				]
+				+ SHorizontalBox::Slot()
+				.AutoWidth()
+				.VAlign(VAlign_Center)
+				.Padding(0.f, 0.f, 0.f, 0.f)
+				[
+					SAssignNew(OpenParentButton, SButton)
+					.Visibility(EVisibility::Collapsed)
+					.ContentPadding(GetLargeIconToolBarButtonContentPadding())
+					.ButtonStyle(FDynamicMaterialEditorStyle::Get(), "HoverHintOnly")
+					.ToolTipText(LOCTEXT("MaterialDesignerOpenParentTooltip", "Open the parent of this Material Designer Dynamic."))
+					.OnClicked(this, &SDMToolBar::OnOpenParentClicked)
+					[
+						SNew(SImage)
+						.Image(FAppStyle::GetBrush(TEXT("Icons.Blueprints")))
+						.DesiredSizeOverride(GetLargeIconToolBarButtonSize())
+					]
+				]
+				+ SHorizontalBox::Slot()
+				.AutoWidth()
+				.VAlign(VAlign_Center)
+				.Padding(0.f, 0.f, 0.f, 0.f)
+				[
+					SAssignNew(ConvertToEditableButton, SButton)
+					.Visibility(EVisibility::Collapsed)
+					.ContentPadding(GetLargeIconToolBarButtonContentPadding())
+					.ButtonStyle(FDynamicMaterialEditorStyle::Get(), "HoverHintOnly")
+					.ToolTipText(LOCTEXT("MaterialDesignerConvertToEditableTooltip", "Convert this Material Designer Dyanmic to a fully editable material (and create a new shader)."))
+					.OnClicked(this, &SDMToolBar::OnConvertToEditableClicked)
+					[
+						SNew(SImage)
+						.Image(FAppStyle::GetBrush(TEXT("Icons.Edit")))
+						.DesiredSizeOverride(GetLargeIconToolBarButtonSize())
+					]
+				]
 			]
 
 			+ SWrapBox::Slot()
@@ -124,8 +160,40 @@ TSharedRef<SWidget> SDMToolBar::CreateToolBarEntries()
 			.VAlign(VAlign_Center)
 			.Padding(5.0f, 0.0f, 0.0f, 0.0f)
 			[
-				SNew(SHorizontalBox)
-				.Visibility(this, &SDMToolBar::GetActorVisibility)
+				SAssignNew(AssetRowWidget, SHorizontalBox)
+				+ SHorizontalBox::Slot()
+				.AutoWidth()
+				.VAlign(VAlign_Center)
+				.Padding(5.0f, 0.0f, 0.0f, 0.0f)
+				[
+					SNew(SButton)
+					.ContentPadding(GetLargeIconToolBarButtonContentPadding())
+					.ButtonStyle(FDynamicMaterialEditorStyle::Get(), "HoverHintOnly")
+					.ToolTipText(LOCTEXT("MaterialDesignerBrowseTooltip", "Browse to the selected asset in the content browser."))
+					.OnClicked(this, &SDMToolBar::OnBrowseClicked)
+					[
+						SNew(SImage)
+						.Image(FAppStyle::GetBrush(TEXT("Icons.BrowseContent")))
+						.DesiredSizeOverride(GetLargeIconToolBarButtonSize())
+					]
+				]
+				+ SHorizontalBox::Slot()
+				.AutoWidth()
+				.VAlign(VAlign_Center)
+				.Padding(0.f, 0.f, 0.f, 0.f)
+				[
+					SAssignNew(AssetNameWidget, STextBlock)
+					.TextStyle(FDynamicMaterialEditorStyle::Get(), "ActorName")
+				]
+			]
+
+			+ SWrapBox::Slot()
+			.HAlign(HAlign_Left)
+			.VAlign(VAlign_Center)
+			.Padding(5.0f, 0.0f, 0.0f, 0.0f)
+			[
+				SAssignNew(ActorRowWidget, SHorizontalBox)
+				.Visibility(EVisibility::Collapsed)
 				+ SHorizontalBox::Slot()
 				.AutoWidth()
 				.Padding(0.f, 0.f, 5.f, 0.f)
@@ -140,9 +208,8 @@ TSharedRef<SWidget> SDMToolBar::CreateToolBarEntries()
 				.Padding(0.f, 0.f, 0.f, 0.f)
 				.VAlign(EVerticalAlignment::VAlign_Center)
 				[
-					SNew(STextBlock)
+					SAssignNew(ActorNameWidget, STextBlock)
 					.TextStyle(FDynamicMaterialEditorStyle::Get(), "ActorName")
-					.Text(this, &SDMToolBar::GetActorName)
 				]	
 				+ SHorizontalBox::Slot()
 				.AutoWidth()
@@ -150,7 +217,6 @@ TSharedRef<SWidget> SDMToolBar::CreateToolBarEntries()
 				.Padding(5.0f, 0.0f, 0.0f, 0.0f)
 				[
 					SAssignNew(SlotSelectorContainer, SBox)
-					.Visibility(this, &SDMToolBar::GetActorVisibility)
 					[
 						CreateSlotsComboBoxWidget()
 					]
@@ -160,8 +226,7 @@ TSharedRef<SWidget> SDMToolBar::CreateToolBarEntries()
 				.VAlign(EVerticalAlignment::VAlign_Center)
 				.Padding(5.0f, 0.0f, 0.0f, 0.0f)
 				[
-					SAssignNew(UseButton, SButton)
-					.Visibility(EVisibility::Collapsed)
+					SNew(SButton)
 					.ContentPadding(GetLargeIconToolBarButtonContentPadding())
 					.ButtonStyle(FDynamicMaterialEditorStyle::Get(), "HoverHintOnly")
 					.ToolTipText(LOCTEXT("MaterialDesignerUseTooltip", "Replace the material in this slot with the one selected in the content browser."))
@@ -173,42 +238,6 @@ TSharedRef<SWidget> SDMToolBar::CreateToolBarEntries()
 					]
 				]
 			]
-
-			+ SWrapBox::Slot()
-			.HAlign(HAlign_Left)
-			.VAlign(VAlign_Center)
-			.Padding(5.0f, 0.0f, 0.0f, 0.0f)
-			[
-				SNew(SHorizontalBox)
-				.Visibility(this, &SDMToolBar::GetAssetVisibility)
-				+ SHorizontalBox::Slot()
-				.AutoWidth()
-				.VAlign(VAlign_Center)
-				.Padding(5.0f, 0.0f, 0.0f, 0.0f)
-				[
-					SAssignNew(BrowseButton, SButton)
-					.Visibility(EVisibility::Collapsed)
-					.ContentPadding(GetLargeIconToolBarButtonContentPadding())
-					.ButtonStyle(FDynamicMaterialEditorStyle::Get(), "HoverHintOnly")
-					.ToolTipText(LOCTEXT("MaterialDesignerBrowseTooltip", "Browse to the selected asset in the content browser."))
-					.OnClicked(this, &SDMToolBar::OnBrowseClicked)
-					[
-						SNew(SImage)
-						.Image(FAppStyle::GetBrush(TEXT("Icons.BrowseContent")))
-						.DesiredSizeOverride(GetLargeIconToolBarButtonSize())
-					]
-				]
-				+ SHorizontalBox::Slot()
-				.AutoWidth()
-				.VAlign(VAlign_Center)
-				.Padding(0.f, 0.f, 0.f, 0.f)
-				[
-					SNew(STextBlock)
-					.TextStyle(FDynamicMaterialEditorStyle::Get(), "ActorName")
-					.Text(this, &SDMToolBar::GetAssetName)
-					.ToolTipText(this, &SDMToolBar::GetAssetToolTip)
-				]
-			]
 		]
 		
 		+ SHorizontalBox::Slot()
@@ -269,7 +298,9 @@ TSharedRef<SWidget> SDMToolBar::CreateToolBarButton(TAttribute<const FSlateBrush
 
 TSharedRef<SWidget> SDMToolBar::CreateSlotsComboBoxWidget()
 {
-	if (!MaterialActorWeak.IsValid() || !MaterialModelWeak.IsValid())
+	UDynamicMaterialModelBase* MaterialModelBase = GetMaterialModelBase();
+
+	if (!MaterialActorWeak.IsValid() || !IsValid(MaterialModelBase))
 	{
 		return SNullWidget::NullWidget;
 	}
@@ -325,11 +356,11 @@ void SDMToolBar::OnMaterialSlotChanged(TSharedPtr<FDMObjectMaterialProperty> InS
 		return;
 	}
 
-	UDynamicMaterialModel* SelectedMaterialModel = InSelectedSlot->GetMaterialModel();
+	UDynamicMaterialModelBase* SelectedMaterialModelBase = InSelectedSlot->GetMaterialModelBase();
 
-	if (IsValid(SelectedMaterialModel))
+	if (IsValid(SelectedMaterialModelBase))
 	{
-		SetMaterialModel(SelectedMaterialModel);
+		OnMaterialModelChanged();
 	}
 	else if (InSelectedSlot->OuterWeak.IsValid())
 	{
@@ -346,44 +377,65 @@ void SDMToolBar::OnMaterialSlotChanged(TSharedPtr<FDMObjectMaterialProperty> InS
 	OnSlotChanged.ExecuteIfBound(InSelectedSlot);
 }
 
-EVisibility SDMToolBar::GetSlotsComboBoxWidgetVisibiltiy() const
-{
-	return ActorMaterialProperties.Num() > 1 ? EVisibility::Visible : EVisibility::Collapsed;
-}
-
 void SDMToolBar::SetMaterialProperties(const TArray<TSharedPtr<FDMObjectMaterialProperty>>& InActorMaterialProperties)
 {
 	ActorMaterialProperties = InActorMaterialProperties;
 	SelectedMaterialSlotIndex = 0;
 }
 
-void SDMToolBar::SetMaterialModel(UDynamicMaterialModel* InModel)
+void SDMToolBar::OnMaterialModelChanged()
 {
-	if (!IsValid(InModel))
-	{
-		return;
-	}
-
-	MaterialModelWeak = InModel;
-
 	bool bIsAsset = false;
+	bool bIsDynamic = false;
 
-	if (InModel->IsAsset())
-	{
-		bIsAsset = true;
-	}
-	else if (UDynamicMaterialInstance* MaterialInstance = InModel->GetDynamicMaterialInstance())
+	UDynamicMaterialModelBase* MaterialModelBase = GetMaterialModelBase();
+
+	if (IsValid(MaterialModelBase))
 	{
-		if (MaterialInstance->IsAsset())
+		if (MaterialModelBase->IsAsset())
 		{
 			bIsAsset = true;
 		}
+		else if (UDynamicMaterialInstance* MaterialInstance = MaterialModelBase->GetDynamicMaterialInstance())
+		{
+			if (MaterialInstance->IsAsset())
+			{
+				bIsAsset = true;
+			}
+		}
+
+		bIsDynamic = !MaterialModelBase->IsA<UDynamicMaterialModel>();
 	}
 
-	BrowseButton->SetVisibility(bIsAsset ? EVisibility::Visible : EVisibility::Collapsed);
+	SaveButtonWidget->SetEnabled(CanSave());
+
+	if (bIsAsset)
+	{
+		AssetNameWidget->SetText(GetAssetName());
+		AssetNameWidget->SetToolTipText(GetAssetToolTip());
+		AssetNameWidget->SetVisibility(EVisibility::Visible);
 
-	// If this is a valid actor material slot, this will be updated by a SetMaterialActor call.
-	UseButton->SetVisibility(EVisibility::Collapsed);
+		AssetRowWidget->SetVisibility(EVisibility::Visible);
+	}
+	else
+	{
+		AssetNameWidget->SetText(FText::GetEmpty());
+		AssetNameWidget->SetToolTipText(FText::GetEmpty());
+		AssetNameWidget->SetVisibility(EVisibility::Collapsed);
+
+		AssetRowWidget->SetVisibility(EVisibility::Collapsed);
+	}
+
+	if (bIsDynamic)
+	{
+		OpenParentButton->SetVisibility(EVisibility::Visible);
+		ConvertToEditableButton->SetVisibility(EVisibility::Visible);
+	}
+	else
+	{
+		OpenParentButton->SetVisibility(EVisibility::Collapsed);
+		ConvertToEditableButton->SetVisibility(EVisibility::Collapsed);
+	}
 
 	ActorMaterialProperties.Empty(0);
 	SelectedMaterialSlotIndex = INDEX_NONE;
@@ -395,12 +447,13 @@ void SDMToolBar::SetMaterialActor(AActor* InActor, const int32 InActiveSlotIndex
 	ActorMaterialProperties.Empty();
 	SelectedMaterialSlotIndex = 0;
 
-	UseButton->SetVisibility(MaterialActorWeak.IsValid() ? EVisibility::Visible : EVisibility::Collapsed);
-
 	if (IsValid(InActor))
 	{
+		ActorNameWidget->SetText(GetActorName());
+		ActorRowWidget->SetVisibility(EVisibility::Visible);
+
 		TArray<FDMObjectMaterialProperty> ActorProperties = UDMBlueprintFunctionLibrary::GetActorMaterialProperties(InActor);
-		UDynamicMaterialModel* MaterialModel = MaterialModelWeak.Get();
+		UDynamicMaterialModelBase* MaterialModelBase = GetMaterialModelBase();
 
 		for (int32 MaterialPropertyIdx = 0; MaterialPropertyIdx < ActorProperties.Num(); ++MaterialPropertyIdx)
 		{
@@ -408,16 +461,31 @@ void SDMToolBar::SetMaterialActor(AActor* InActor, const int32 InActiveSlotIndex
 
 			ActorMaterialProperties.Add(MakeShared<FDMObjectMaterialProperty>(MaterialProperty));
 
-			if (MaterialProperty.GetMaterialModel() == MaterialModel)
+			if (MaterialProperty.GetMaterialModelBase() == MaterialModelBase)
 			{
 				SelectedMaterialSlotIndex = MaterialPropertyIdx;
 			}
 		}
 	}
+	else
+	{
+		ActorNameWidget->SetText(FText::GetEmpty());
+		ActorRowWidget->SetVisibility(EVisibility::Collapsed);
+	}
 
 	SlotSelectorContainer->SetContent(CreateSlotsComboBoxWidget());
 }
 
+UDynamicMaterialModelBase* SDMToolBar::GetMaterialModelBase() const
+{
+	if (TSharedPtr<SDMEditor> Editor = EditorWeak.Pin())
+	{
+		return Editor->GetMaterialModelBase();
+	}
+
+	return nullptr;
+}
+
 const FSlateBrush* SDMToolBar::GetFollowSelectionBrush() const
 {
 	static const FSlateBrush* Unlocked = FAppStyle::Get().GetBrush("Icons.Unlock");
@@ -462,57 +530,56 @@ FReply SDMToolBar::OnFollowSelectionButtonClicked()
 	return FReply::Handled();
 }
 
-EVisibility SDMToolBar::GetExportMaterialInstanceButtonVisibility() const
+FReply SDMToolBar::OnExportMaterialInstanceButtonClicked()
 {
-	UDynamicMaterialModel* MaterialModel = MaterialModelWeak.Get();
+	UDynamicMaterialModelBase* MaterialModelBase = GetMaterialModelBase();
 
-	if (!MaterialModel)
+	if (!MaterialModelBase)
 	{
-		return EVisibility::Collapsed;
+		return FReply::Handled();
 	}
 
-	UDynamicMaterialInstance* MaterialInstance = MaterialModel->GetDynamicMaterialInstance();
+	UDynamicMaterialInstance* MaterialInstance = MaterialModelBase->GetDynamicMaterialInstance();
 
 	if (!MaterialInstance)
 	{
-		return EVisibility::Collapsed;
+		return FReply::Handled();
 	}
+	
+	FString CurrentName = MaterialInstance->GetName();
 
-	return EVisibility::Visible;
-}
-
-FReply SDMToolBar::OnExportMaterialInstanceButtonClicked()
-{
-	UDynamicMaterialModel* MaterialModel = MaterialModelWeak.Get();
-
-	if (!MaterialModel)
+	if (MaterialModelBase->IsA<UDynamicMaterialModel>())
 	{
-		return FReply::Handled();
+		CurrentName = CurrentName.StartsWith(TEXT("MDI_"))
+			? CurrentName
+			: (TEXT("MDI_") + CurrentName);
 	}
-
-	UDynamicMaterialInstance* MaterialInstance = MaterialModel->GetDynamicMaterialInstance();
-
-	if (!MaterialInstance)
+	else
 	{
-		return FReply::Handled();
+		CurrentName = CurrentName.StartsWith(TEXT("MDD_"))
+			? CurrentName
+			: (TEXT("MDD_") + CurrentName);
 	}
 
-	IContentBrowserSingleton& ContentBrowser = FModuleManager::LoadModuleChecked<FContentBrowserModule>("ContentBrowser").Get();
-	const FContentBrowserItemPath CurrentPath = ContentBrowser.GetCurrentPath();
+	IAssetTools& AssetTools = FModuleManager::LoadModuleChecked<FAssetToolsModule>("AssetTools").Get();
+	FString PackageName, AssetName;
+	AssetTools.CreateUniqueAssetName(CurrentName, TEXT(""), PackageName, AssetName);
+
+	IContentBrowserSingleton& ContentBrowserModule = FModuleManager::LoadModuleChecked<FContentBrowserModule>("ContentBrowser").Get();
+	const FContentBrowserItemPath CurrentPath = ContentBrowserModule.GetCurrentPath();
 	const FString PathStr = CurrentPath.HasInternalPath() ? CurrentPath.GetInternalPathString() : "/Game";
 
 	FSaveAssetDialogConfig SaveAssetDialogConfig;
 	SaveAssetDialogConfig.DialogTitleOverride = LOCTEXT("SaveAssetDialogTitle", "Save Asset As");
 	SaveAssetDialogConfig.DefaultPath = PathStr;
-	SaveAssetDialogConfig.DefaultAssetName = MaterialInstance->GetName();
-	SaveAssetDialogConfig.ExistingAssetPolicy = ESaveAssetDialogExistingAssetPolicy::AllowButWarn;
+	SaveAssetDialogConfig.ExistingAssetPolicy = ESaveAssetDialogExistingAssetPolicy::Disallow;
+	SaveAssetDialogConfig.DefaultAssetName = AssetName;
 
-	FContentBrowserModule& ContentBrowserModule = FModuleManager::LoadModuleChecked<FContentBrowserModule>("ContentBrowser");
-	FString SaveObjectPath = ContentBrowserModule.Get().CreateModalSaveAssetDialog(SaveAssetDialogConfig);
+	const FString SaveObjectPath = ContentBrowserModule.Get().CreateModalSaveAssetDialog(SaveAssetDialogConfig);
 
 	if (!SaveObjectPath.IsEmpty())
 	{
-		UDMBlueprintFunctionLibrary::ExportMaterialInstance(MaterialInstance->GetMaterialModel(), SaveObjectPath);
+		UDMBlueprintFunctionLibrary::ExportMaterialInstance(MaterialInstance->GetMaterialModelBase(), SaveObjectPath);
 
 		if (FEngineAnalytics::IsAvailable())
 		{
@@ -525,20 +592,20 @@ FReply SDMToolBar::OnExportMaterialInstanceButtonClicked()
 
 FReply SDMToolBar::OnBrowseClicked()
 {
-	UDynamicMaterialModel* MaterialModel = MaterialModelWeak.Get();
+	UDynamicMaterialModelBase* MaterialModelBase = GetMaterialModelBase();
 
-	if (!IsValid(MaterialModel))
+	if (!IsValid(MaterialModelBase))
 	{
 		return FReply::Handled();
 	}
 
 	UObject* Asset = nullptr;
 
-	if (MaterialModelWeak->IsAsset())
+	if (MaterialModelBase->IsAsset())
 	{
-		Asset = MaterialModel;
+		Asset = MaterialModelBase;
 	}
-	else if (UDynamicMaterialInstance* MaterialInstance = MaterialModel->GetDynamicMaterialInstance())
+	else if (UDynamicMaterialInstance* MaterialInstance = MaterialModelBase->GetDynamicMaterialInstance())
 	{
 		if (MaterialInstance->IsAsset())
 		{
@@ -565,10 +632,10 @@ FReply SDMToolBar::OnUseClicked()
 		return FReply::Handled();
 	}
 
-	UDynamicMaterialModel* CurrentModel = MaterialModelWeak.Get();
+	UDynamicMaterialModelBase* CurrentModelBase = GetMaterialModelBase();
 	UDMWorldSubsystem* DMSubsystem = nullptr;
 
-	if (CurrentModel)
+	if (CurrentModelBase)
 	{
 		AActor* Actor = MaterialActorWeak.Get();
 
@@ -626,7 +693,7 @@ FReply SDMToolBar::OnUseClicked()
 
 	if (DMSubsystem && DMSubsystem->GetMaterialValueSetterDelegate().IsBound())
 	{
-		if (!DMSubsystem->ExecuteIsValidDelegate(CurrentModel))
+		if (!DMSubsystem->ExecuteIsValidDelegate(CurrentModelBase))
 		{
 			return FReply::Handled();
 		}
@@ -673,47 +740,20 @@ FText SDMToolBar::GetActorName() const
 	return FText::GetEmpty();
 }
 
-EVisibility SDMToolBar::GetActorVisibility() const
-{
-	return GetMaterialActor()
-		? EVisibility::Visible
-		: EVisibility::Collapsed;
-}
-
-EVisibility SDMToolBar::GetAssetVisibility() const
-{
-	if (UDynamicMaterialModel* MaterialModel = MaterialModelWeak.Get())
-	{
-		if (UDynamicMaterialInstance* MaterialInstance = MaterialModel->GetDynamicMaterialInstance())
-		{
-			if (MaterialInstance->IsAsset())
-			{
-				return EVisibility::Visible;
-			}
-		}
-		else if (MaterialModel->IsAsset())
-		{
-			return EVisibility::Visible;
-		}
-	}
-
-	return EVisibility::Collapsed;
-}
-
 FText SDMToolBar::GetAssetName() const
 {
-	if (UDynamicMaterialModel* MaterialModel = MaterialModelWeak.Get())
+	if (UDynamicMaterialModelBase* MaterialModelBase = GetMaterialModelBase())
 	{
-		if (UDynamicMaterialInstance* MaterialInstance = MaterialModel->GetDynamicMaterialInstance())
+		if (UDynamicMaterialInstance* MaterialInstance = MaterialModelBase->GetDynamicMaterialInstance())
 		{
 			if (MaterialInstance->IsAsset())
 			{
 				return FText::FromString(MaterialInstance->GetName());
 			}
 		}
-		else if (MaterialModel->IsAsset())
+		else if (MaterialModelBase->IsAsset())
 		{
-			return FText::FromString(MaterialModel->GetName());
+			return FText::FromString(MaterialModelBase->GetName());
 		}
 	}
 
@@ -722,18 +762,18 @@ FText SDMToolBar::GetAssetName() const
 
 FText SDMToolBar::GetAssetToolTip() const
 {
-	if (UDynamicMaterialModel* MaterialModel = MaterialModelWeak.Get())
+	if (UDynamicMaterialModelBase* MaterialModelBase = GetMaterialModelBase())
 	{
-		if (UDynamicMaterialInstance* MaterialInstance = MaterialModel->GetDynamicMaterialInstance())
+		if (UDynamicMaterialInstance* MaterialInstance = MaterialModelBase->GetDynamicMaterialInstance())
 		{
 			if (MaterialInstance->IsAsset())
 			{
 				return FText::FromString(MaterialInstance->GetPathName());
 			}
 		}
-		else if (MaterialModel->IsAsset())
+		else if (MaterialModelBase->IsAsset())
 		{
-			return FText::FromString(MaterialModel->GetPathName());
+			return FText::FromString(MaterialModelBase->GetPathName());
 		}
 	}
 
@@ -742,12 +782,12 @@ FText SDMToolBar::GetAssetToolTip() const
 
 bool SDMToolBar::CanSave() const
 {
-	return !!GetSaveablePackage(GetMaterialModel());
+	return !!GetSaveablePackage(GetMaterialModelBase());
 }
 
 const FSlateBrush* SDMToolBar::GetSaveIcon() const
 {
-	if (UPackage* Package = GetSaveablePackage(GetMaterialModel()))
+	if (UPackage* Package = GetSaveablePackage(GetMaterialModelBase()))
 	{
 		if (Package->IsDirty())
 		{
@@ -760,13 +800,203 @@ const FSlateBrush* SDMToolBar::GetSaveIcon() const
 
 FReply SDMToolBar::OnSaveClicked()
 {
-	if (UDynamicMaterialModel* MaterialModel = GetMaterialModel())
+	if (UDynamicMaterialModelBase* MaterialModelBase = GetMaterialModelBase())
+	{
+		if (UPackage* Package = GetSaveablePackage(MaterialModelBase))
+		{
+			TArray<UObject*> AssetsToSave;
+			AssetsToSave.Add(MaterialModelBase);
+			UPackageTools::SavePackagesForObjects(AssetsToSave);
+		}
+	}
+
+	return FReply::Handled();
+}
+
+FReply SDMToolBar::OnOpenParentClicked()
+{
+	if (TSharedPtr<SDMEditor> Editor = EditorWeak.Pin())
+	{
+		if (UDynamicMaterialModelDynamic* DynamicMaterialModel = Cast<UDynamicMaterialModelDynamic>(Editor->GetMaterialModelBase()))
+		{
+			if (UDynamicMaterialModel* ParentModel = DynamicMaterialModel->ResolveMaterialModel())
+			{
+				Editor->SetMaterialActor(nullptr);
+				Editor->SetMaterialModelBase(ParentModel);
+			}
+		}
+	}
+
+	return FReply::Handled();
+}
+
+FReply SDMToolBar::OnConvertToEditableClicked()
+{
+	UDynamicMaterialModelDynamic* CurrentModelDynamic = Cast<UDynamicMaterialModelDynamic>(GetMaterialModelBase());
+
+	if (!CurrentModelDynamic)
+	{
+		UE_LOG(LogDynamicMaterialEditor, Error, TEXT("Tried to convert a null or non-dynamic model to editable."));
+		return FReply::Handled();
+	}
+
+	UDynamicMaterialModel* ParentModel = CurrentModelDynamic->GetParentModel();
+
+	if (!ParentModel)
+	{
+		UE_LOG(LogDynamicMaterialEditor, Error, TEXT("Failed to find parent model."));
+		return FReply::Handled();
+	}
+
+	bool bIsAsset = false;
+
+	if (CurrentModelDynamic->IsAsset())
+	{
+		bIsAsset = true;
+	}
+	else if (UDynamicMaterialInstance* MaterialInstance = CurrentModelDynamic->GetDynamicMaterialInstance())
+	{
+		if (MaterialInstance->IsAsset())
+		{
+			bIsAsset = true;
+		}
+	}
+
+	UDMWorldSubsystem* DMSubsystem = nullptr;
+	AActor* Actor = MaterialActorWeak.Get();
+	TSharedPtr<FDMObjectMaterialProperty> CurrentActorProperty = nullptr;
+
+	if (Actor && ActorMaterialProperties.IsValidIndex(SelectedMaterialSlotIndex) && ActorMaterialProperties[SelectedMaterialSlotIndex].IsValid())
+	{
+		UWorld* World = Actor->GetWorld();
+
+		if (IsValid(World))
+		{
+			DMSubsystem = World->GetSubsystem<UDMWorldSubsystem>();
+		}
+
+		CurrentActorProperty = ActorMaterialProperties[SelectedMaterialSlotIndex];
+	}
+
+	// In-actor models/instance must have a world subsystem to query.
+	if (!bIsAsset && !DMSubsystem)
+	{
+		UE_LOG(LogDynamicMaterialEditor, Error, TEXT("Cannot create a new asset for embedded instances without an active world subsystem."));
+		return FReply::Handled();
+	}
+
+	UDynamicMaterialInstance* OldInstance = CurrentModelDynamic->GetDynamicMaterialInstance();
+
+	// Where should we save it? (Always export to CB)
+	FString CurrentName = OldInstance ? OldInstance->GetName() : CurrentModelDynamic->GetName();
+
+	CurrentName = CurrentName.StartsWith(TEXT("MDD_"))
+		? (TEXT("MDI_") + CurrentName.RightChop(4))
+		: (TEXT("MDI_") + CurrentName);
+
+	IAssetTools& AssetTools = FModuleManager::LoadModuleChecked<FAssetToolsModule>("AssetTools").Get();
+	FString PackageName, AssetName;
+	AssetTools.CreateUniqueAssetName(CurrentName, TEXT(""), PackageName, AssetName);
+
+	IContentBrowserSingleton& ContentBrowserModule = FModuleManager::LoadModuleChecked<FContentBrowserModule>("ContentBrowser").Get();
+	const FContentBrowserItemPath CurrentPath = ContentBrowserModule.GetCurrentPath();
+	const FString PathStr = CurrentPath.HasInternalPath() ? CurrentPath.GetInternalPathString() : "/Game";
+
+	FSaveAssetDialogConfig SaveAssetDialogConfig;
+	SaveAssetDialogConfig.DialogTitleOverride = LOCTEXT("SaveAssetDialogTitle", "Save Asset As");
+	SaveAssetDialogConfig.DefaultPath = PathStr;
+	SaveAssetDialogConfig.ExistingAssetPolicy = ESaveAssetDialogExistingAssetPolicy::Disallow;
+	SaveAssetDialogConfig.DefaultAssetName = AssetName;
+
+	const FString SaveObjectPath = ContentBrowserModule.Get().CreateModalSaveAssetDialog(SaveAssetDialogConfig);
+
+	if (SaveObjectPath.IsEmpty())
+	{
+		UE_LOG(LogDynamicMaterialEditor, Warning, TEXT("No path was chosen for saving the new editable asset, cancelling."));
+		return FReply::Handled();
+	}
+
+	// Create new model
+	UDynamicMaterialModel* NewModel = CurrentModelDynamic->ToEditable(GetTransientPackage());
+
+	if (!NewModel)
+	{
+		UE_LOG(LogDynamicMaterialEditor, Error, TEXT("Failed to convert dynamic asset to editable."));
+		return FReply::Handled();
+	}
+
+	// Create a package for it
+	PackageName = FPaths::GetBaseFilename(*SaveObjectPath, false);
+	UPackage* Package = CreatePackage(*PackageName);
+
+	if (!Package)
+	{
+		UE_LOG(LogDynamicMaterialEditor, Error, TEXT("Failed to create new package for editable asset."));
+		return FReply::Handled();
+	}
+
+	AssetName = FPaths::GetBaseFilename(*SaveObjectPath, true);
+
+	// Do we need an instance to? Or just a model?
+	UDynamicMaterialInstance* NewInstance = nullptr;
+
+	if (OldInstance)
+	{
+		NewInstance = Cast<UDynamicMaterialInstance>(GetMutableDefault<UDynamicMaterialInstanceFactory>()->FactoryCreateNew(
+			UDynamicMaterialInstance::StaticClass(),
+			Package,
+			*AssetName,
+			RF_Transactional | RF_Public | RF_Standalone,
+			NewModel,
+			nullptr
+		));
+	}
+	else
+	{
+		NewModel->Rename(*AssetName, Package, UE::DynamicMaterial::RenameFlags);
+	}
+
+	FAssetRegistryModule::AssetCreated(NewModel);
+
+	// If it was in an actor, set it on the actor
+	if (NewInstance)
+	{
+		if (CurrentActorProperty.IsValid())
+		{
+			if (DMSubsystem && DMSubsystem->GetMaterialValueSetterDelegate().IsBound())
+			{
+				if (DMSubsystem->ExecuteIsValidDelegate(CurrentModelDynamic))
+				{
+					DMSubsystem->ExecuteMaterialValueSetterDelegate(*CurrentActorProperty, NewInstance);
+				}
+				else
+				{
+					UE_LOG(LogDynamicMaterialEditor, Error, TEXT("Asset is not valid for current world, not assigning to actor."));
+				}
+			}
+			else
+			{
+				ActorMaterialProperties[SelectedMaterialSlotIndex]->SetMaterial(NewInstance);
+			}
+
+			if (TSharedPtr<SDMEditor> Editor = EditorWeak.Pin())
+			{
+				Editor->SetMaterialObjectProperty(*CurrentActorProperty);
+			}
+		}
+		else
+		{
+			if (TSharedPtr<SDMEditor> Editor = EditorWeak.Pin())
+			{
+				Editor->SetMaterialModelBase(NewModel);
+			}
+		}
+	}
+	else
 	{
-		if (UPackage* Package = GetSaveablePackage(MaterialModel))
+		if (TSharedPtr<SDMEditor> Editor = EditorWeak.Pin())
 		{
-			TArray<UObject*> TexturesToSaveArray;
-			TexturesToSaveArray.Add(MaterialModel);
-			UPackageTools::SavePackagesForObjects(TexturesToSaveArray);
+			Editor->SetMaterialModelBase(NewModel);
 		}
 	}
 
diff --git a/Engine/Plugins/Experimental/DynamicMaterial/Source/DynamicMaterialEditor/Private/Slate/SDMToolBar.h b/Engine/Plugins/Experimental/DynamicMaterial/Source/DynamicMaterialEditor/Private/Slate/SDMToolBar.h
--- a/Engine/Plugins/Experimental/DynamicMaterial/Source/DynamicMaterialEditor/Private/Slate/SDMToolBar.h
+++ b/Engine/Plugins/Experimental/DynamicMaterial/Source/DynamicMaterialEditor/Private/Slate/SDMToolBar.h
@@ -9,6 +9,7 @@
 class AActor;
 class SBox;
 class SDMEditor;
+class STextBlock;
 class SWidget;
 class UDMMaterialStageExpression;
 class UDynamicMaterialModel;
@@ -28,8 +29,6 @@ class SDMToolBar : public SCompoundWidget
 {
 public:
 	SLATE_BEGIN_ARGS(SDMToolBar) {}
-		SLATE_ARGUMENT(TWeakObjectPtr<AActor>, MaterialActor)
-		SLATE_ARGUMENT(TWeakObjectPtr<UDynamicMaterialModel>, MaterialModel)
 		SLATE_EVENT(FDMOnActorMaterailSlotChanged, OnSlotChanged)
 		SLATE_EVENT(FOnGetContent, OnGetSettingsMenu)
 	SLATE_END_ARGS()
@@ -44,8 +43,9 @@ public:
 	AActor* GetMaterialActor() const { return MaterialActorWeak.Get(); }
 	void SetMaterialActor(AActor* InActor, const int32 InActiveSlotIndex = 0);
 
-	UDynamicMaterialModel* GetMaterialModel() const { return MaterialModelWeak.Get(); }
-	void SetMaterialModel(UDynamicMaterialModel* InModel);
+	UDynamicMaterialModelBase* GetMaterialModelBase() const;
+
+	void OnMaterialModelChanged();
 
 	FText GetActorName() const;
 
@@ -54,7 +54,6 @@ protected:
 
 	TWeakPtr<SDMEditor> EditorWeak;
 	TWeakObjectPtr<AActor> MaterialActorWeak;
-	TWeakObjectPtr<UDynamicMaterialModel> MaterialModelWeak;
 	FDMOnActorMaterailSlotChanged OnSlotChanged;
 	FOnGetContent OnGetSettingsMenu;
 
@@ -62,9 +61,13 @@ protected:
 	int32 SelectedMaterialSlotIndex;
 
 	TSharedPtr<SBox> SlotSelectorContainer;
-	TSharedPtr<SWidget> SaveButton;
-	TSharedPtr<SWidget> BrowseButton;
-	TSharedPtr<SWidget> UseButton;
+	TSharedPtr<SWidget> SaveButtonWidget;
+	TSharedPtr<SWidget> ActorRowWidget;
+	TSharedPtr<SWidget> AssetRowWidget;
+	TSharedPtr<STextBlock> ActorNameWidget;
+	TSharedPtr<STextBlock> AssetNameWidget;
+	TSharedPtr<SWidget> OpenParentButton;
+	TSharedPtr<SWidget> ConvertToEditableButton;
 
 	TSharedRef<SWidget> CreateToolBarEntries();
 
@@ -82,23 +85,16 @@ protected:
 	const FMargin GetLargeIconToolBarButtonContentPadding() const { return FMargin(4.0f); }
 	const FVector2D GetLargeIconToolBarButtonSize() const { return FVector2D(16.0f); }
 
-	EVisibility GetSlotsComboBoxWidgetVisibiltiy() const;
-
 	const FSlateBrush* GetFollowSelectionBrush() const;
 	FSlateColor GetFollowSelectionColor() const;
 	FReply OnFollowSelectionButtonClicked();
 
-	EVisibility GetExportMaterialInstanceButtonVisibility() const;
 	FReply OnExportMaterialInstanceButtonClicked();
 
 	FReply OnBrowseClicked();
 
 	FReply OnUseClicked();
 
-	EVisibility GetActorVisibility() const;
-
-	EVisibility GetAssetVisibility() const;
-
 	FText GetAssetName() const;
 
 	FText GetAssetToolTip() const;
@@ -108,4 +104,8 @@ protected:
 	const FSlateBrush* GetSaveIcon() const;
 
 	FReply OnSaveClicked();
+
+	FReply OnOpenParentClicked();
+
+	FReply OnConvertToEditableClicked();
 };
diff --git a/Engine/Plugins/Experimental/DynamicMaterial/Source/DynamicMaterialEditor/Private/Utils/DMBlueprintFunctionLibrary.cpp b/Engine/Plugins/Experimental/DynamicMaterial/Source/DynamicMaterialEditor/Private/Utils/DMBlueprintFunctionLibrary.cpp
--- a/Engine/Plugins/Experimental/DynamicMaterial/Source/DynamicMaterialEditor/Private/Utils/DMBlueprintFunctionLibrary.cpp
+++ b/Engine/Plugins/Experimental/DynamicMaterial/Source/DynamicMaterialEditor/Private/Utils/DMBlueprintFunctionLibrary.cpp
@@ -32,6 +32,7 @@
 #include "Material/DynamicMaterialInstanceFactory.h"
 #include "Materials/Material.h"
 #include "Model/DynamicMaterialModel.h"
+#include "Model/DynamicMaterialModelDynamic.h"
 #include "Model/DynamicMaterialModelEditorOnlyData.h"
 #include "Model/IDynamicMaterialModelEditorOnlyDataInterface.h"
 #include "UObject/Package.h"
@@ -247,7 +248,7 @@ UDynamicMaterialModel* UDMBlueprintFunctionLibrary::CreateDynamicMaterialInObjec
 	return nullptr;
 }
 
-bool UDMBlueprintFunctionLibrary::ExportMaterialInstance(UDynamicMaterialModel* InMaterialModel, const FString& InSavePath)
+bool UDMBlueprintFunctionLibrary::ExportMaterialInstance(UDynamicMaterialModelBase* InMaterialModel, const FString& InSavePath)
 {
 	if (!IsValid(InMaterialModel))
 	{
@@ -294,13 +295,20 @@ bool UDMBlueprintFunctionLibrary::ExportMaterialInstance(UDynamicMaterialModel*
 
 	if (UDynamicMaterialInstance* NewInstance = Cast<UDynamicMaterialInstance>(NewAsset))
 	{
-		if (UDynamicMaterialModel* NewModel = NewInstance->GetMaterialModel())
+		if (InMaterialModel->IsA<UDynamicMaterialModel>())
 		{
-			if (IDynamicMaterialModelEditorOnlyDataInterface* ModelEditorOnlyData = NewModel->GetEditorOnlyData())
+			if (UDynamicMaterialModel* NewModel = NewInstance->GetMaterialModel())
 			{
-				ModelEditorOnlyData->RequestMaterialBuild();
+				if (IDynamicMaterialModelEditorOnlyDataInterface* ModelEditorOnlyData = NewModel->GetEditorOnlyData())
+				{
+					ModelEditorOnlyData->RequestMaterialBuild();
+				}
 			}
 		}
+		else if (InMaterialModel->IsA<UDynamicMaterialModelDynamic>())
+		{
+			NewInstance->InitializeMIDPublic();
+		}
 	}
 
 	FAssetRegistryModule::AssetCreated(NewAsset);
@@ -308,7 +316,7 @@ bool UDMBlueprintFunctionLibrary::ExportMaterialInstance(UDynamicMaterialModel*
 	return true;
 }
 
-bool UDMBlueprintFunctionLibrary::ExportGeneratedMaterial(UDynamicMaterialModel* InMaterialModel, const FString& InSavePath)
+bool UDMBlueprintFunctionLibrary::ExportGeneratedMaterial(UDynamicMaterialModelBase* InMaterialModel, const FString& InSavePath)
 {
 	if (!IsValid(InMaterialModel))
 	{
diff --git a/Engine/Plugins/Experimental/DynamicMaterial/Source/DynamicMaterialEditor/Private/Utils/DMBlueprintFunctionLibrary.h b/Engine/Plugins/Experimental/DynamicMaterial/Source/DynamicMaterialEditor/Private/Utils/DMBlueprintFunctionLibrary.h
--- a/Engine/Plugins/Experimental/DynamicMaterial/Source/DynamicMaterialEditor/Private/Utils/DMBlueprintFunctionLibrary.h
+++ b/Engine/Plugins/Experimental/DynamicMaterial/Source/DynamicMaterialEditor/Private/Utils/DMBlueprintFunctionLibrary.h
@@ -11,7 +11,7 @@ class AActor;
 class UDMMaterialStage;
 class UDMMaterialStageInputValue;
 class UDMRenderTargetRenderer;
-class UDynamicMaterialModel;
+class UDynamicMaterialModelBase;
 class UTexture;
 struct FDMObjectMaterialProperty;
 
@@ -32,9 +32,9 @@ public:
 
 	static UDynamicMaterialModel* CreateDynamicMaterialInObject(FDMObjectMaterialProperty& InMaterialProperty);
 
-	static bool ExportMaterialInstance(UDynamicMaterialModel* InMaterialModel, const FString& InSavePath);
+	static bool ExportMaterialInstance(UDynamicMaterialModelBase* InMaterialModel, const FString& InSavePath);
 
-	static bool ExportGeneratedMaterial(UDynamicMaterialModel* InMaterialModel, const FString& InSavePath);
+	static bool ExportGeneratedMaterial(UDynamicMaterialModelBase* InMaterialModel, const FString& InSavePath);
 
 	static void SetStageInputToRenderer(UDMMaterialStage* InStage, TSubclassOf<UDMRenderTargetRenderer> InRendererClass, int32 InInputIndex);
 };
diff --git a/Engine/Plugins/Experimental/DynamicMaterial/Source/DynamicMaterialEditor/Public/Components/DMMaterialEffect.h b/Engine/Plugins/Experimental/DynamicMaterial/Source/DynamicMaterialEditor/Public/Components/DMMaterialEffect.h
--- a/Engine/Plugins/Experimental/DynamicMaterial/Source/DynamicMaterialEditor/Public/Components/DMMaterialEffect.h
+++ b/Engine/Plugins/Experimental/DynamicMaterial/Source/DynamicMaterialEditor/Public/Components/DMMaterialEffect.h
@@ -25,17 +25,17 @@ enum class EDMMaterialEffectTarget : uint8
 	Slot      = 1 << 3
 };
 
-UCLASS(Abstract, BlueprintType, Blueprintable, ClassGroup = "Material Designer", meta = (DisplayName = "Material Designer Effect"))
-class DYNAMICMATERIALEDITOR_API UDMMaterialEffect : public UDMMaterialComponent, public IDMJsonSerializable
+UCLASS(MinimalAPI, Abstract, BlueprintType, Blueprintable, ClassGroup = "Material Designer", meta = (DisplayName = "Material Designer Effect"))
+class UDMMaterialEffect : public UDMMaterialComponent, public IDMJsonSerializable
 {
 	GENERATED_BODY()
 
 public:
 	UFUNCTION(BlueprintPure, Category = "Material Designer")
-	static EDMMaterialEffectTarget StageTypeToEffectType(EDMMaterialLayerStage InStageType);
+	static DYNAMICMATERIALEDITOR_API EDMMaterialEffectTarget StageTypeToEffectType(EDMMaterialLayerStage InStageType);
 
 	UFUNCTION(BlueprintCallable, Category = "Material Designer")
-	static UDMMaterialEffect* CreateEffect(UDMMaterialEffectStack* InEffectStack, TSubclassOf<UDMMaterialEffect> InEffectClass);
+	static DYNAMICMATERIALEDITOR_API UDMMaterialEffect* CreateEffect(UDMMaterialEffectStack* InEffectStack, TSubclassOf<UDMMaterialEffect> InEffectClass);
 
 	template<typename InEffectClass>
 	static InEffectClass* CreateEffect(UDMMaterialEffectStack* InEffectStack)
@@ -43,44 +43,48 @@ public:
 		return Cast<InEffectClass>(CreateEffect(InEffectStack, InEffectClass::StaticClass()));
 	}
 
-	UDMMaterialEffect();
+	DYNAMICMATERIALEDITOR_API UDMMaterialEffect();
 
 	UFUNCTION(BlueprintPure, Category = "Material Designer")
-	UDMMaterialEffectStack* GetEffectStack() const;
+	DYNAMICMATERIALEDITOR_API UDMMaterialEffectStack* GetEffectStack() const;
 
+	/** Retrieves the index of this effect in the effect stack. */
 	UFUNCTION(BlueprintCallable, Category = "Material Designer")
-	int32 FindIndex() const;
+	DYNAMICMATERIALEDITOR_API int32 FindIndex() const;
 
 	UFUNCTION(BlueprintPure, Category = "Material Designer")
-	bool IsEnabled() const;
+	DYNAMICMATERIALEDITOR_API bool IsEnabled() const;
 
 	UFUNCTION(BlueprintCallable, Category = "Material Designer")
-	bool SetEnabled(bool bInIsEnabled);
+	DYNAMICMATERIALEDITOR_API bool SetEnabled(bool bInIsEnabled);
 
+	/** Returns the type of nodes which this effect targets. */
 	UFUNCTION(BlueprintPure, Category = "Material Designer")
-	EDMMaterialEffectTarget GetEffectTarget() const;
+	DYNAMICMATERIALEDITOR_API EDMMaterialEffectTarget GetEffectTarget() const;
 
 	UFUNCTION(BlueprintPure, Category = "Material Designer")
-	virtual FText GetEffectName() const;
+	DYNAMICMATERIALEDITOR_API virtual FText GetEffectName() const;
 
 	UFUNCTION(BlueprintPure, Category = "Material Designer")
 	virtual FText GetEffectDescription() const PURE_VIRTUAL(UDMMaterialEffect::ApplyTo, return FText::GetEmpty();)
 
+	/** Test whether this effect is compatible with another effect. */
 	virtual bool IsCompatibleWith(UDMMaterialEffect* InEffect) const { return true; }
 
+	/** Apply this effect to the output of something, such as a stage, slot or texture. */
 	virtual void ApplyTo(const TSharedRef<FDMMaterialBuildState>& InBuildState, TArray<UMaterialExpression*>& InOutExpressions, 
 		int32& InOutLastExpressionOutputChannel, int32& InLastExpressionOutputIndex) const PURE_VIRTUAL(UDMMaterialEffect::ApplyTo)
 
 	//~ Begin UDMMaterialComponent
-	virtual UDMMaterialComponent* GetParentComponent() const override;
-	virtual FString GetComponentPathComponent() const override;
-	virtual FText GetComponentDescription() const override;
-	virtual void Update(EDMUpdateType InUpdateType) override;
-	virtual void PostEditorDuplicate(UDynamicMaterialModel* InMaterialModel, UDMMaterialComponent* InParent) override;
+	DYNAMICMATERIALEDITOR_API virtual UDMMaterialComponent* GetParentComponent() const override;
+	DYNAMICMATERIALEDITOR_API virtual FString GetComponentPathComponent() const override;
+	DYNAMICMATERIALEDITOR_API virtual FText GetComponentDescription() const override;
+	DYNAMICMATERIALEDITOR_API virtual void Update(EDMUpdateType InUpdateType) override;
+	DYNAMICMATERIALEDITOR_API virtual void PostEditorDuplicate(UDynamicMaterialModel* InMaterialModel, UDMMaterialComponent* InParent) override;
 	//~ End UDMMaterialComponent
 
 	//~ Begin UObject
-	virtual void PostEditUndo() override;
+	DYNAMICMATERIALEDITOR_API virtual void PostEditUndo() override;
 	//~ End UObject
 
 protected:
diff --git a/Engine/Plugins/Experimental/DynamicMaterial/Source/DynamicMaterialEditor/Public/Components/DMMaterialEffectFunction.h b/Engine/Plugins/Experimental/DynamicMaterial/Source/DynamicMaterialEditor/Public/Components/DMMaterialEffectFunction.h
--- a/Engine/Plugins/Experimental/DynamicMaterial/Source/DynamicMaterialEditor/Public/Components/DMMaterialEffectFunction.h
+++ b/Engine/Plugins/Experimental/DynamicMaterial/Source/DynamicMaterialEditor/Public/Components/DMMaterialEffectFunction.h
@@ -7,53 +7,54 @@
 
 class UMaterialFunctionInterface;
 
-UCLASS(BlueprintType, ClassGroup = "Material Designer", Meta = (DisplayName = "Material Designer Effect Function"))
-class DYNAMICMATERIALEDITOR_API UDMMaterialEffectFunction : public UDMMaterialEffect
+UCLASS(MinimalAPI, BlueprintType, ClassGroup = "Material Designer", Meta = (DisplayName = "Material Designer Effect Function"))
+class UDMMaterialEffectFunction : public UDMMaterialEffect
 {
 	GENERATED_BODY()
 
 public:
-	static const FString InputsPathToken;
+	DYNAMICMATERIALEDITOR_API static const FString InputsPathToken;
 
-	UDMMaterialEffectFunction();
+	DYNAMICMATERIALEDITOR_API UDMMaterialEffectFunction();
 
 	UFUNCTION(BlueprintPure, Category = "Material Designer")
-	UMaterialFunctionInterface* GetMaterialFunction() const;
+	DYNAMICMATERIALEDITOR_API UMaterialFunctionInterface* GetMaterialFunction() const;
 
 	UFUNCTION(BlueprintCallable, Category = "Material Designer")
-	bool SetMaterialFunction(UMaterialFunctionInterface* InFunction);
+	DYNAMICMATERIALEDITOR_API bool SetMaterialFunction(UMaterialFunctionInterface* InFunction);
 
+	/** Returns the value used as the function input. */
 	UFUNCTION(BlueprintPure, Category = "Material Designer")
-	UDMMaterialValue* GetInputValue(int32 InIndex) const;
+	DYNAMICMATERIALEDITOR_API UDMMaterialValue* GetInputValue(int32 InIndex) const;
 
 	UFUNCTION(BlueprintCallable, Category = "Material Designer", Meta = (DisplayName = "Get Input Values"))
-	TArray<UDMMaterialValue*> BP_GetInputValues() const;
+	DYNAMICMATERIALEDITOR_API TArray<UDMMaterialValue*> BP_GetInputValues() const;
 
-	const TArray<TObjectPtr<UDMMaterialValue>>& GetInputValues() const;
+	DYNAMICMATERIALEDITOR_API const TArray<TObjectPtr<UDMMaterialValue>>& GetInputValues() const;
 
 #if WITH_EDITOR
 	//~ Begin IDMJsonSerializable
-	virtual TSharedPtr<FJsonValue> JsonSerialize() const override;
-	virtual bool JsonDeserialize(const TSharedPtr<FJsonValue>& InJsonValue) override;
+	DYNAMICMATERIALEDITOR_API virtual TSharedPtr<FJsonValue> JsonSerialize() const override;
+	DYNAMICMATERIALEDITOR_API virtual bool JsonDeserialize(const TSharedPtr<FJsonValue>& InJsonValue) override;
 	//~ End IDMJsonSerializable
 #endif
 
 	//~ Begin UDMMaterialEffect
-	virtual FText GetEffectName() const override;
-	virtual FText GetEffectDescription() const override;
-	virtual bool IsCompatibleWith(UDMMaterialEffect* InEffect) const override;
-	virtual void ApplyTo(const TSharedRef<FDMMaterialBuildState>& InBuildState, TArray<UMaterialExpression*>& InOutExpressions,
+	DYNAMICMATERIALEDITOR_API virtual FText GetEffectName() const override;
+	DYNAMICMATERIALEDITOR_API virtual FText GetEffectDescription() const override;
+	DYNAMICMATERIALEDITOR_API virtual bool IsCompatibleWith(UDMMaterialEffect* InEffect) const override;
+	DYNAMICMATERIALEDITOR_API virtual void ApplyTo(const TSharedRef<FDMMaterialBuildState>& InBuildState, TArray<UMaterialExpression*>& InOutExpressions,
 		int32& InOutLastExpressionOutputChannel, int32& InOutLastExpressionOutputIndex) const override;
 	//~ End UDMMaterialEffect
 
 	//~ Begin UDMMaterialComponent
-	virtual FText GetComponentDescription() const override;
-	virtual void PostEditorDuplicate(UDynamicMaterialModel* InMaterialModel, UDMMaterialComponent* InParent) override;
+	DYNAMICMATERIALEDITOR_API virtual FText GetComponentDescription() const override;
+	DYNAMICMATERIALEDITOR_API virtual void PostEditorDuplicate(UDynamicMaterialModel* InMaterialModel, UDMMaterialComponent* InParent) override;
 	//~ End UDMMaterialComponent
 
 	//~ Begin UObject
-	virtual bool Modify(bool bInAlwaysMarkDirty = true) override;
-	virtual void PostLoad() override;
+	DYNAMICMATERIALEDITOR_API virtual bool Modify(bool bInAlwaysMarkDirty = true) override;
+	DYNAMICMATERIALEDITOR_API virtual void PostLoad() override;
 	//~ End UObject
 
 protected:
@@ -63,17 +64,21 @@ protected:
 	UPROPERTY(VisibleInstanceOnly, BlueprintReadOnly, Instanced, Category = "Material Designer")
 	TArray<TObjectPtr<UDMMaterialValue>> InputValues;
 
+	/** Re-inits the function. */
 	void OnMaterialFunctionChanged();
 
+	/** Adds required inputs. */
 	void InitFunction();
 
+	/** Removes inputs. */
 	void DeinitFunction();
 
+	/** Returns true if the input values do not match the function's input pins. */
 	bool NeedsFunctionInit() const;
 
 	//~ Begin UDMMaterialComponent
-	virtual UDMMaterialComponent* GetSubComponentByPath(FDMComponentPath& InPath, const FDMComponentPathSegment& InPathSegment) const override;
-	virtual void OnComponentAdded() override;
-	virtual void OnComponentRemoved() override;
+	DYNAMICMATERIALEDITOR_API virtual UDMMaterialComponent* GetSubComponentByPath(FDMComponentPath& InPath, const FDMComponentPathSegment& InPathSegment) const override;
+	DYNAMICMATERIALEDITOR_API virtual void OnComponentAdded() override;
+	DYNAMICMATERIALEDITOR_API virtual void OnComponentRemoved() override;
 	//~ End UDMMaterialComponent
 };
