Material Designer

diff --git a/Engine/Plugins/Experimental/DynamicMaterial/Source/DynamicMaterialEditor/Public/Components/DMMaterialEffectStack.h b/Engine/Plugins/Experimental/DynamicMaterial/Source/DynamicMaterialEditor/Public/Components/DMMaterialEffectStack.h
--- a/Engine/Plugins/Experimental/DynamicMaterial/Source/DynamicMaterialEditor/Public/Components/DMMaterialEffectStack.h
+++ b/Engine/Plugins/Experimental/DynamicMaterial/Source/DynamicMaterialEditor/Public/Components/DMMaterialEffectStack.h
@@ -34,13 +34,13 @@ struct FDMMaterialEffectStackJson
 /**
  * Container for effects. Effects can be applied to either layers (on a per stage basis) or to slots.
  */
-UCLASS(BlueprintType, ClassGroup = "Material Designer", Meta = (DisplayName = "Material Designer Effect Stack"))
-class DYNAMICMATERIALEDITOR_API UDMMaterialEffectStack : public UDMMaterialComponent
+UCLASS(MinimalAPI, BlueprintType, ClassGroup = "Material Designer", Meta = (DisplayName = "Material Designer Effect Stack"))
+class UDMMaterialEffectStack : public UDMMaterialComponent
 {
 	GENERATED_BODY()
 
 public:
-	static const FString EffectsPathToken;
+	DYNAMICMATERIALEDITOR_API static const FString EffectsPathToken;
 
 	using FEffectCallbackFunc = TFunctionRef<void(UDMMaterialEffect*)>;
 
@@ -56,39 +56,39 @@ public:
 		return CreateEffectStack(InLayer);
 	}
 
-	static UDMMaterialEffectStack* CreateEffectStack(UDMMaterialSlot* InSlot);
-	static UDMMaterialEffectStack* CreateEffectStack(UDMMaterialLayerObject* InLayer);
+	DYNAMICMATERIALEDITOR_API static UDMMaterialEffectStack* CreateEffectStack(UDMMaterialSlot* InSlot);
+	DYNAMICMATERIALEDITOR_API static UDMMaterialEffectStack* CreateEffectStack(UDMMaterialLayerObject* InLayer);
 
-	UDMMaterialEffectStack();
+	DYNAMICMATERIALEDITOR_API UDMMaterialEffectStack();
 
 	UFUNCTION(BlueprintPure, Category = "Material Designer")
-	UDMMaterialSlot* GetSlot() const;
+	DYNAMICMATERIALEDITOR_API UDMMaterialSlot* GetSlot() const;
 
 	UFUNCTION(BlueprintPure, Category = "Material Designer")
-	UDMMaterialLayerObject* GetLayer() const;
+	DYNAMICMATERIALEDITOR_API UDMMaterialLayerObject* GetLayer() const;
 
 	UFUNCTION(BlueprintCallable, Category = "Material Designer")
-	bool IsEnabled() const;
+	DYNAMICMATERIALEDITOR_API bool IsEnabled() const;
 
 	UFUNCTION(BlueprintCallable, Category = "Material Designer")
-	bool SetEnabled(bool bInIsEnabled);
+	DYNAMICMATERIALEDITOR_API bool SetEnabled(bool bInIsEnabled);
 
 	UFUNCTION(BlueprintPure, Category = "Material Designer")
-	UDMMaterialEffect* GetEffect(int32 InIndex) const;
+	DYNAMICMATERIALEDITOR_API UDMMaterialEffect* GetEffect(int32 InIndex) const;
 
 	UFUNCTION(BlueprintCallable, Category = "Material Designer", Meta = (DisplayName = "Get Effects"))
-	TArray<UDMMaterialEffect*> BP_GetEffects() const;
+	DYNAMICMATERIALEDITOR_API TArray<UDMMaterialEffect*> BP_GetEffects() const;
 
-	const TArray<TObjectPtr<UDMMaterialEffect>>& GetEffects() const;
+	DYNAMICMATERIALEDITOR_API const TArray<TObjectPtr<UDMMaterialEffect>>& GetEffects() const;
 
 	UFUNCTION(BlueprintCallable, Category = "Material Designer")
-	bool HasEffect(const UDMMaterialEffect* InEffect) const;
+	DYNAMICMATERIALEDITOR_API bool HasEffect(const UDMMaterialEffect* InEffect) const;
 
 	UFUNCTION(BlueprintCallable, Category = "Material Designer")
-	bool AddEffect(UDMMaterialEffect* InEffect);
+	DYNAMICMATERIALEDITOR_API bool AddEffect(UDMMaterialEffect* InEffect);
 
 	UFUNCTION(BlueprintCallable, Category = "Material Designer")
-	UDMMaterialEffect* SetEffect(int32 InIndex, UDMMaterialEffect* InEffect);
+	DYNAMICMATERIALEDITOR_API UDMMaterialEffect* SetEffect(int32 InIndex, UDMMaterialEffect* InEffect);
 
 	UFUNCTION(BlueprintCallable, Category = "Material Designer", Meta = (DisplayName = "Move Effect (By Index)"))
 	bool BP_MoveEffectByIndex(int32 InIndex, int32 InNewIndex)
@@ -96,7 +96,7 @@ public:
 		return MoveEffect(InIndex, InNewIndex);
 	}
 
-	bool MoveEffect(int32 InIndex, int32 InNewIndex);
+	DYNAMICMATERIALEDITOR_API bool MoveEffect(int32 InIndex, int32 InNewIndex);
 
 	UFUNCTION(BlueprintCallable, Category = "Material Designer", Meta = (DisplayName = "Move Effect (By Value)"))
 	bool BP_MoveEffectByValue(UDMMaterialEffect* InEffect, int32 InNewIndex)
@@ -104,7 +104,7 @@ public:
 		return MoveEffect(InEffect, InNewIndex);
 	}
 
-	bool MoveEffect(UDMMaterialEffect* InEffect, int32 InNewIndex);
+	DYNAMICMATERIALEDITOR_API bool MoveEffect(UDMMaterialEffect* InEffect, int32 InNewIndex);
 
 	UFUNCTION(BlueprintCallable, Category = "Material Designer", Meta = (DisplayName = "Remove Effect (By Index)"))
 	UDMMaterialEffect* BP_RemoveEffectByIndex(int32 InIndex)
@@ -112,7 +112,7 @@ public:
 		return RemoveEffect(InIndex);
 	}
 
-	UDMMaterialEffect* RemoveEffect(int32 InIndex);
+	DYNAMICMATERIALEDITOR_API UDMMaterialEffect* RemoveEffect(int32 InIndex);
 
 	UFUNCTION(BlueprintCallable, Category = "Material Designer", Meta = (DisplayName = "Remove Effect (By Value)"))
 	bool BP_RemoveEffectByValue(UDMMaterialEffect* InEffect)
@@ -120,28 +120,31 @@ public:
 		return RemoveEffect(InEffect);
 	}
 
-	bool RemoveEffect(UDMMaterialEffect* InEffect);
+	DYNAMICMATERIALEDITOR_API bool RemoveEffect(UDMMaterialEffect* InEffect);
 
-	bool ApplyEffects(const TSharedRef<FDMMaterialBuildState>& InBuildState, EDMMaterialEffectTarget InEffectTarget,
+	/** Apply all matching effect types to the expressions and add them to the array. */
+	DYNAMICMATERIALEDITOR_API bool ApplyEffects(const TSharedRef<FDMMaterialBuildState>& InBuildState, EDMMaterialEffectTarget InEffectTarget,
 		TArray<UMaterialExpression*>& InOutStageExpressions, int32& InOutLastExpressionOutputChannel, int32& InOutLastExpressionOutputIndex) const;
 
+	/** Creates a preset based on the current stage. */
 	UFUNCTION(BlueprintCallable, Category = "Material Designer")
-	FDMMaterialEffectStackJson CreatePreset();
+	DYNAMICMATERIALEDITOR_API FDMMaterialEffectStackJson CreatePreset();
 
+	/** Apply the given preset to this stack. Does not remove old effects. */
 	UFUNCTION(BlueprintCallable, Category = "Material Designer")
-	void ApplyPreset(const FDMMaterialEffectStackJson& InPreset);
+	DYNAMICMATERIALEDITOR_API void ApplyPreset(const FDMMaterialEffectStackJson& InPreset);
 
 	//~ Begin UDMMaterialComponent
-	virtual UDMMaterialComponent* GetParentComponent() const override;
-	virtual FString GetComponentPathComponent() const override;
-	virtual FText GetComponentDescription() const override;
-	virtual void Update(EDMUpdateType InUpdateType) override;
-	virtual void PostEditorDuplicate(UDynamicMaterialModel* InMaterialModel, UDMMaterialComponent* InParent) override;
+	DYNAMICMATERIALEDITOR_API virtual UDMMaterialComponent* GetParentComponent() const override;
+	DYNAMICMATERIALEDITOR_API virtual FString GetComponentPathComponent() const override;
+	DYNAMICMATERIALEDITOR_API virtual FText GetComponentDescription() const override;
+	DYNAMICMATERIALEDITOR_API virtual void Update(EDMUpdateType InUpdateType) override;
+	DYNAMICMATERIALEDITOR_API virtual void PostEditorDuplicate(UDynamicMaterialModel* InMaterialModel, UDMMaterialComponent* InParent) override;
 	//~ End UDMMaterialComponent
 
 	//~ Begin UObject
-	virtual bool Modify(bool bInAlwaysMarkDirty = true) override;
-	virtual void PostEditUndo() override;
+	DYNAMICMATERIALEDITOR_API virtual bool Modify(bool bInAlwaysMarkDirty = true) override;
+	DYNAMICMATERIALEDITOR_API virtual void PostEditUndo() override;
 	//~ End UObject
 
 protected:
@@ -151,13 +154,14 @@ protected:
 	UPROPERTY(VisibleInstanceOnly, BlueprintReadOnly, Instanced, Category = "Material Designer")
 	TArray<TObjectPtr<UDMMaterialEffect>> Effects;
 
-	TArray<UDMMaterialEffect*> GetIncompatibleEffects(UDMMaterialEffect* InEffect);
+	DYNAMICMATERIALEDITOR_API TArray<UDMMaterialEffect*> GetIncompatibleEffects(UDMMaterialEffect* InEffect);
 
-	TArray<UDMMaterialEffect*> RemoveIncompatibleEffects(UDMMaterialEffect* InEffect);
+	/** Returns the removed effects. */
+	DYNAMICMATERIALEDITOR_API TArray<UDMMaterialEffect*> RemoveIncompatibleEffects(UDMMaterialEffect* InEffect);
 
 	//~ Begin UDMMaterialComponent
-	virtual UDMMaterialComponent* GetSubComponentByPath(FDMComponentPath& InPath, const FDMComponentPathSegment& InPathSegment) const override;
-	virtual void OnComponentAdded() override;
-	virtual void OnComponentRemoved() override;
+	DYNAMICMATERIALEDITOR_API virtual UDMMaterialComponent* GetSubComponentByPath(FDMComponentPath& InPath, const FDMComponentPathSegment& InPathSegment) const override;
+	DYNAMICMATERIALEDITOR_API virtual void OnComponentAdded() override;
+	DYNAMICMATERIALEDITOR_API virtual void OnComponentRemoved() override;
 	//~ End UDMMaterialComponent
 };
diff --git a/Engine/Plugins/Experimental/DynamicMaterial/Source/DynamicMaterialEditor/Public/Components/DMMaterialLayer.h b/Engine/Plugins/Experimental/DynamicMaterial/Source/DynamicMaterialEditor/Public/Components/DMMaterialLayer.h
--- a/Engine/Plugins/Experimental/DynamicMaterial/Source/DynamicMaterialEditor/Public/Components/DMMaterialLayer.h
+++ b/Engine/Plugins/Experimental/DynamicMaterial/Source/DynamicMaterialEditor/Public/Components/DMMaterialLayer.h
@@ -12,65 +12,70 @@ class UDMMaterialSlot;
 class UDMMaterialStage;
 class UMaterialExpression;
 
-UCLASS(BlueprintType, ClassGroup = "Material Designer", Meta = (DisplayName = "Material Designer Layer"))
-class DYNAMICMATERIALEDITOR_API UDMMaterialLayerObject : public UDMMaterialComponent
+/** A collection of stages. */
+UCLASS(MinimalAPI, BlueprintType, ClassGroup = "Material Designer", Meta = (DisplayName = "Material Designer Layer"))
+class UDMMaterialLayerObject : public UDMMaterialComponent
 {
 	GENERATED_BODY()
 
 public:
-	static const FString StagesPathToken;
-	static const FString BasePathToken;
-	static const FString MaskPathToken;
-	static const FString EffectStackPathToken;
+	DYNAMICMATERIALEDITOR_API static const FString StagesPathToken;
+	DYNAMICMATERIALEDITOR_API static const FString BasePathToken;
+	DYNAMICMATERIALEDITOR_API static const FString MaskPathToken;
+	DYNAMICMATERIALEDITOR_API static const FString EffectStackPathToken;
 
 	using FStageCallbackFunc = TFunctionRef<void(UDMMaterialStage*)>;
 
 	UFUNCTION(BlueprintCallable, Category = "Material Designer")
-	static UDMMaterialLayerObject* CreateLayer(UDMMaterialSlot* InSlot, EDMMaterialPropertyType InMaterialProperty, 
+	static DYNAMICMATERIALEDITOR_API UDMMaterialLayerObject* CreateLayer(UDMMaterialSlot* InSlot, EDMMaterialPropertyType InMaterialProperty,
 		const TArray<UDMMaterialStage*>& InStages);
 	
-	static UDMMaterialLayerObject* DeserializeFromString(UDMMaterialSlot* InOuter, const FString& InSerializedString);
+	DYNAMICMATERIALEDITOR_API static UDMMaterialLayerObject* DeserializeFromString(UDMMaterialSlot* InOuter, const FString& InSerializedString);
 
-	UDMMaterialLayerObject();
+	DYNAMICMATERIALEDITOR_API UDMMaterialLayerObject();
 
 	UFUNCTION(BlueprintPure, Category = "Material Designer")
-	UDMMaterialSlot* GetSlot() const;
+	DYNAMICMATERIALEDITOR_API UDMMaterialSlot* GetSlot() const;
 
+	/** Find the index of this layer in the slot. */
 	UFUNCTION(BlueprintCallable, Category = "Material Designer")
-	int32 FindIndex() const;
+	DYNAMICMATERIALEDITOR_API int32 FindIndex() const;
 
 	UFUNCTION(BlueprintCallable, Category = "Material Designer")
 	const FText& GetLayerName() const { return LayerName; }
 
 	UFUNCTION(BlueprintCallable, Category = "Material Designer")
-	void SetLayerName(const FText& InName);
+	DYNAMICMATERIALEDITOR_API void SetLayerName(const FText& InName);
 
 	UFUNCTION(BlueprintCallable, Category = "Material Designer")
-	bool IsEnabled() const;
+	DYNAMICMATERIALEDITOR_API bool IsEnabled() const;
 
 	UFUNCTION(BlueprintCallable, Category = "Material Designer")
-	bool SetEnabled(bool bInIsEnabled);
+	DYNAMICMATERIALEDITOR_API bool SetEnabled(bool bInIsEnabled);
 
 	UFUNCTION(BlueprintPure, Category = "Material Designer")
-	EDMMaterialPropertyType GetMaterialProperty() const;
+	DYNAMICMATERIALEDITOR_API EDMMaterialPropertyType GetMaterialProperty() const;
 
 	UFUNCTION(BlueprintCallable, Category = "Material Designer")
-	bool SetMaterialProperty(EDMMaterialPropertyType InMaterialProperty);
+	DYNAMICMATERIALEDITOR_API bool SetMaterialProperty(EDMMaterialPropertyType InMaterialProperty);
 
+	/** Texture UV Link means that all stages use the same Texture UV from the base stage, if available. */
 	UFUNCTION(BlueprintPure, Category = "Material Designer")
-	bool IsTextureUVLinkEnabled() const;
+	DYNAMICMATERIALEDITOR_API bool IsTextureUVLinkEnabled() const;
 
+	/** Texture UV Link means that all stages use the same Texture UV from the base stage, if available. */
 	UFUNCTION(BlueprintCallable, Category = "Material Designer")
-	bool SetTextureUVLinkEnabled(bool bInValue);
+	DYNAMICMATERIALEDITOR_API bool SetTextureUVLinkEnabled(bool bInValue);
 
+	/** Texture UV Link means that all stages use the same Texture UV from the base stage, if available. */
 	UFUNCTION(BlueprintCallable, Category = "Material Designer")
-	bool ToggleTextureUVLinkEnabled();
+	DYNAMICMATERIALEDITOR_API bool ToggleTextureUVLinkEnabled();
 
 	UFUNCTION(BlueprintCallable, Category = "Material Designer")
-	UDMMaterialLayerObject* GetPreviousLayer(EDMMaterialPropertyType InUsingProperty, EDMMaterialLayerStage InSearchFor) const;
+	DYNAMICMATERIALEDITOR_API UDMMaterialLayerObject* GetPreviousLayer(EDMMaterialPropertyType InUsingProperty, EDMMaterialLayerStage InSearchFor) const;
 
 	UFUNCTION(BlueprintCallable, Category = "Material Designer")
-	UDMMaterialLayerObject* GetNextLayer(EDMMaterialPropertyType InUsingProperty, EDMMaterialLayerStage InSearchFor) const;
+	DYNAMICMATERIALEDITOR_API UDMMaterialLayerObject* GetNextLayer(EDMMaterialPropertyType InUsingProperty, EDMMaterialLayerStage InSearchFor) const;
 
 	UFUNCTION(BlueprintCallable, Category = "Material Designer")
 	bool CanMoveLayerAbove(UDMMaterialLayerObject* InLayer) const;
@@ -79,80 +84,83 @@ public:
 	bool CanMoveLayerBelow(UDMMaterialLayerObject* InLayer) const;
 
 	UFUNCTION(BlueprintCallable, Category = "Material Designer")
-	UDMMaterialStage* GetStage(EDMMaterialLayerStage InStageType = EDMMaterialLayerStage::All, bool bInCheckEnabled = false) const;
+	DYNAMICMATERIALEDITOR_API UDMMaterialStage* GetStage(EDMMaterialLayerStage InStageType = EDMMaterialLayerStage::All, bool bInCheckEnabled = false) const;
 
 	UFUNCTION(BlueprintCallable, Category = "Material Designer")
-	TArray<UDMMaterialStage*> GetStages(EDMMaterialLayerStage InStageType = EDMMaterialLayerStage::All, bool bInCheckEnabled = false) const;
+	DYNAMICMATERIALEDITOR_API TArray<UDMMaterialStage*> GetStages(EDMMaterialLayerStage InStageType = EDMMaterialLayerStage::All, bool bInCheckEnabled = false) const;
 
-	const TArray<TObjectPtr<UDMMaterialStage>>& GetAllStages() const;
+	DYNAMICMATERIALEDITOR_API const TArray<TObjectPtr<UDMMaterialStage>>& GetAllStages() const;
 
 	UFUNCTION(BlueprintCallable, Category = "Material Designer")
-	EDMMaterialLayerStage GetStageType(const UDMMaterialStage* InStage) const;
+	DYNAMICMATERIALEDITOR_API EDMMaterialLayerStage GetStageType(const UDMMaterialStage* InStage) const;
 
 	UFUNCTION(BlueprintCallable, Category = "Material Designer")
-	UDMMaterialStage* GetFirstValidStage(EDMMaterialLayerStage InStageScope) const;
+	DYNAMICMATERIALEDITOR_API UDMMaterialStage* GetFirstValidStage(EDMMaterialLayerStage InStageScope) const;
 
 	UFUNCTION(BlueprintCallable, Category = "Material Designer")
-	UDMMaterialStage* GetLastValidStage(EDMMaterialLayerStage InStageScope) const;
+	DYNAMICMATERIALEDITOR_API UDMMaterialStage* GetLastValidStage(EDMMaterialLayerStage InStageScope) const;
 
 	UFUNCTION(BlueprintCallable, Category = "Material Designer")
-	bool HasValidStage(const UDMMaterialStage* InStage) const;
+	DYNAMICMATERIALEDITOR_API bool HasValidStage(const UDMMaterialStage* InStage) const;
 
 	UFUNCTION(BlueprintCallable, Category = "Material Designer")
-	bool HasValidStageOfType(EDMMaterialLayerStage InStageScope = EDMMaterialLayerStage::All) const;
+	DYNAMICMATERIALEDITOR_API bool HasValidStageOfType(EDMMaterialLayerStage InStageScope = EDMMaterialLayerStage::All) const;
 
 	UFUNCTION(BlueprintPure, Category = "Material Designer")
-	bool IsStageEnabled(EDMMaterialLayerStage InStageScope = EDMMaterialLayerStage::All) const;
+	DYNAMICMATERIALEDITOR_API bool IsStageEnabled(EDMMaterialLayerStage InStageScope = EDMMaterialLayerStage::All) const;
 
 	/** Checks for the first enabled and valid stage. */
 	UFUNCTION(BlueprintCallable, Category = "Material Designer")
-	UDMMaterialStage* GetFirstEnabledStage(EDMMaterialLayerStage InStageScope) const;
+	DYNAMICMATERIALEDITOR_API UDMMaterialStage* GetFirstEnabledStage(EDMMaterialLayerStage InStageScope) const;
 
 	/** Checks for the last enabled and valid stage. */
 	UFUNCTION(BlueprintCallable, Category = "Material Designer")
-	UDMMaterialStage* GetLastEnabledStage(EDMMaterialLayerStage InStageScope) const;
+	DYNAMICMATERIALEDITOR_API UDMMaterialStage* GetLastEnabledStage(EDMMaterialLayerStage InStageScope) const;
 
+	/** Replace the specified stage. */
 	UFUNCTION(BlueprintCallable, Category = "Material Designer")
-	bool SetStage(EDMMaterialLayerStage InStageType, UDMMaterialStage* InStage);
+	DYNAMICMATERIALEDITOR_API bool SetStage(EDMMaterialLayerStage InStageType, UDMMaterialStage* InStage);
 
 	UFUNCTION(BlueprintCallable, Category = "Material Designer")
-	bool AreAllStagesValid(EDMMaterialLayerStage InStageScope) const;
+	DYNAMICMATERIALEDITOR_API bool AreAllStagesValid(EDMMaterialLayerStage InStageScope) const;
 
 	/** Checks if both stages are enabled and valid */
 	UFUNCTION(BlueprintCallable, Category = "Material Designer")
-	bool AreAllStagesEnabled(EDMMaterialLayerStage InStageScope) const;
+	DYNAMICMATERIALEDITOR_API bool AreAllStagesEnabled(EDMMaterialLayerStage InStageScope) const;
 
-	UFUNCTION(BlueprintPure, Category = "Material Designer")
 	bool IsStageBeingEdited(EDMMaterialLayerStage InStageScope = EDMMaterialLayerStage::All) const;
 
-	UFUNCTION(BlueprintCallable, Category = "Material Designer")
 	UDMMaterialStage* GetFirstStageBeingEdited(EDMMaterialLayerStage InStageScope) const;
 
-	void ForEachValidStage(EDMMaterialLayerStage InStageScope, FStageCallbackFunc InCallback) const;
+	/** Iterate over all the valid stages, whether they are enabled or not. */
+	DYNAMICMATERIALEDITOR_API void ForEachValidStage(EDMMaterialLayerStage InStageScope, FStageCallbackFunc InCallback) const;
 
-	void ForEachEnabledStage(EDMMaterialLayerStage InStageScope, FStageCallbackFunc InCallback) const;
+	/** Iterate over only the enabled stages. */
+	DYNAMICMATERIALEDITOR_API void ForEachEnabledStage(EDMMaterialLayerStage InStageScope, FStageCallbackFunc InCallback) const;
 
 	UFUNCTION(BlueprintPure, Category = "Material Designer")
-	UDMMaterialEffectStack* GetEffectStack() const;
+	DYNAMICMATERIALEDITOR_API UDMMaterialEffectStack* GetEffectStack() const;
 
+	/** Used for copy+pasting. */
 	FString SerializeToString() const;
 
-	void GenerateExpressions(const TSharedRef<FDMMaterialBuildState>& InBuildState) const;
+	DYNAMICMATERIALEDITOR_API void GenerateExpressions(const TSharedRef<FDMMaterialBuildState>& InBuildState) const;
 
-	bool ApplyEffects(const TSharedRef<FDMMaterialBuildState>& InBuildState, const UDMMaterialStage* InStage,
+	/** Apply the effects from this layer's effect stack to the given expressions based on the type of stage. */
+	DYNAMICMATERIALEDITOR_API bool ApplyEffects(const TSharedRef<FDMMaterialBuildState>& InBuildState, const UDMMaterialStage* InStage,
 		TArray<UMaterialExpression*>& InOutStageExpressions, int32& InOutLastExpressionOutputChannel, int32& InOutLastExpressionOutputIndex) const;
 
 	//~ Begin UDMMaterialComponent
-	virtual UDMMaterialComponent* GetParentComponent() const override;
-	virtual FString GetComponentPathComponent() const override;
-	virtual FText GetComponentDescription() const override;
-	virtual void Update(EDMUpdateType InUpdateType) override;
-	virtual void PostEditorDuplicate(UDynamicMaterialModel* InMaterialModel, UDMMaterialComponent* InParent) override;
+	DYNAMICMATERIALEDITOR_API virtual UDMMaterialComponent* GetParentComponent() const override;
+	DYNAMICMATERIALEDITOR_API virtual FString GetComponentPathComponent() const override;
+	DYNAMICMATERIALEDITOR_API virtual FText GetComponentDescription() const override;
+	DYNAMICMATERIALEDITOR_API virtual void Update(EDMUpdateType InUpdateType) override;
+	DYNAMICMATERIALEDITOR_API virtual void PostEditorDuplicate(UDynamicMaterialModel* InMaterialModel, UDMMaterialComponent* InParent) override;
 	//~ End UDMMaterialComponent
 
 	//~ Begin UObject
-	virtual bool Modify(bool bInAlwaysMarkDirty = true) override;
-	virtual void PostEditUndo() override;
+	DYNAMICMATERIALEDITOR_API virtual bool Modify(bool bInAlwaysMarkDirty = true) override;
+	DYNAMICMATERIALEDITOR_API virtual void PostEditUndo() override;
 	//~ End UObject
 
 protected:
@@ -175,8 +183,8 @@ protected:
 	bool bLinkedUVs;
 
 	//~ Begin UDMMaterialComponent
-	virtual UDMMaterialComponent* GetSubComponentByPath(FDMComponentPath& InPath, const FDMComponentPathSegment& InPathSegment) const override;
-	virtual void OnComponentAdded() override;
-	virtual void OnComponentRemoved() override;
+	DYNAMICMATERIALEDITOR_API virtual UDMMaterialComponent* GetSubComponentByPath(FDMComponentPath& InPath, const FDMComponentPathSegment& InPathSegment) const override;
+	DYNAMICMATERIALEDITOR_API virtual void OnComponentAdded() override;
+	DYNAMICMATERIALEDITOR_API virtual void OnComponentRemoved() override;
 	//~ End UDMMaterialComponent
 };
diff --git a/Engine/Plugins/Experimental/DynamicMaterial/Source/DynamicMaterialEditor/Public/Components/DMMaterialProperty.h b/Engine/Plugins/Experimental/DynamicMaterial/Source/DynamicMaterialEditor/Public/Components/DMMaterialProperty.h
--- a/Engine/Plugins/Experimental/DynamicMaterial/Source/DynamicMaterialEditor/Public/Components/DMMaterialProperty.h
+++ b/Engine/Plugins/Experimental/DynamicMaterial/Source/DynamicMaterialEditor/Public/Components/DMMaterialProperty.h
@@ -21,13 +21,13 @@ struct FDMMaterialBuildState;
 /**
  * Base Color, Specular, Opacity, etc
  */
-UCLASS(BlueprintType, ClassGroup = "Material Designer", Meta = (DisplayName = "Material Designer Property"))
-class DYNAMICMATERIALEDITOR_API UDMMaterialProperty : public UDMMaterialComponent
+UCLASS(MinimalAPI, BlueprintType, ClassGroup = "Material Designer", Meta = (DisplayName = "Material Designer Property"))
+class UDMMaterialProperty : public UDMMaterialComponent
 {
 	GENERATED_BODY()
 
 public:
-	static const FString ComponentsPathToken;
+	DYNAMICMATERIALEDITOR_API static const FString ComponentsPathToken;
 
 	/**
 	 * Craetes a material property as a default subobject of the EditorOnlyData object.
@@ -36,11 +36,13 @@ public:
 	 * @param InSubObjName The subobject name.
 	 * @return The newly created default subobject.
 	 */
-	static UDMMaterialProperty* CreateCustomMaterialPropertyDefaultSubobject(UDynamicMaterialModelEditorOnlyData* InModelEditorOnlyData, 
+	DYNAMICMATERIALEDITOR_API static UDMMaterialProperty* CreateCustomMaterialPropertyDefaultSubobject(UDynamicMaterialModelEditorOnlyData* InModelEditorOnlyData,
 		EDMMaterialPropertyType InMaterialProperty, const FName& InSubObjName);
 
+	UDMMaterialProperty();
+
 	UFUNCTION(BlueprintPure, Category = "Material Designer")
-	UDynamicMaterialModelEditorOnlyData* GetMaterialModelEditorOnlyData() const;
+	DYNAMICMATERIALEDITOR_API UDynamicMaterialModelEditorOnlyData* GetMaterialModelEditorOnlyData() const;
 
 	UFUNCTION(BlueprintPure, Category = "Material Designer")
 	EDMMaterialPropertyType GetMaterialProperty() const { return MaterialProperty; }
@@ -49,13 +51,14 @@ public:
 	 * The description of this property based on the EDMMaterialPropertyType enum.
 	 */
 	UFUNCTION(BlueprintPure, Category = "Material Designer")
-	FText GetDescription() const;
+	DYNAMICMATERIALEDITOR_API FText GetDescription() const;
 
 	/**
 	 * Whether this property represents a material pin. Only materials properties with a pin will
 	 * be generated by default.
 	 */
-	bool IsMaterialPin() const;
+	UFUNCTION(BlueprintPure, Category = "Material Designer")
+	DYNAMICMATERIALEDITOR_API bool IsMaterialPin() const;
 
 	/**
 	 * @return The value type of this property. Will be either VT_Float1, VT_Float3_RGB or VT_Float3_XYZ.
@@ -86,32 +89,32 @@ public:
 	 * @param InFunction The new output processor.
 	 */
 	UFUNCTION(BlueprintCallable, Category = "Material Designer")
-	void SetOutputProcessor(UMaterialFunctionInterface* InFunction);
+	DYNAMICMATERIALEDITOR_API void SetOutputProcessor(UMaterialFunctionInterface* InFunction);
 
 	/**
 	 * Check whether the current model settings are valid for this property.
 	 * @param InMaterialModel The model with the settings.
 	 * @return True if it is valid.
 	 */
-	virtual bool IsValidForModel(UDynamicMaterialModelEditorOnlyData& InMaterialModel) const;
+	DYNAMICMATERIALEDITOR_API virtual bool IsValidForModel(UDynamicMaterialModelEditorOnlyData& InMaterialModel) const;
 
 	/**
 	 * Empty the connection map and recalculate the input types.
 	 */
 	UFUNCTION(BlueprintCallable, Category = "Material Designer")
-	virtual void ResetInputConnectionMap();
+	DYNAMICMATERIALEDITOR_API virtual void ResetInputConnectionMap();
 
 	/**
 	 * Generates a default input node for this property if the layer stack is empty.
 	 * @param InBuildState The current build state.
 	 * @return The last in the chain of any expressions created (for direct connection to the property).
 	 */
-	virtual UMaterialExpression* GetDefaultInput(const TSharedRef<FDMMaterialBuildState>& InBuildState) const;
+	DYNAMICMATERIALEDITOR_API virtual UMaterialExpression* GetDefaultInput(const TSharedRef<FDMMaterialBuildState>& InBuildState) const;
 
 	/**
 	 * @return The texture sample type to use for any texture samplers connected to this property.
 	 */
-	virtual TEnumAsByte<EMaterialSamplerType> GetTextureSamplerType() const;
+	DYNAMICMATERIALEDITOR_API virtual TEnumAsByte<EMaterialSamplerType> GetTextureSamplerType() const;
 
 	/**
 	 * Triggered when a slot is added to this property. Should create useful, default layer stack.
@@ -135,21 +138,21 @@ public:
 	 * @param InComponent The component to add.
 	 * @return The replaced component from the map.
 	 */
-	virtual UDMMaterialComponent* AddComponent(FName InName, UDMMaterialComponent* InComponent);
+	DYNAMICMATERIALEDITOR_API virtual UDMMaterialComponent* AddComponent(FName InName, UDMMaterialComponent* InComponent);
 
 	/**
 	 * Checks to see if a component with this name exists.
 	 * @param InName The name of the component.
 	 * @return True if it exists, even if it isn't valid.
 	 */
-	virtual bool HasComponent(FName InName) const;
+	DYNAMICMATERIALEDITOR_API virtual bool HasComponent(FName InName) const;
 
 	/**
 	 * Fetches a component from the component list.
 	 * @param InName The name of the component.
 	 * @return The component in the map with that name.
 	 */
-	virtual UDMMaterialComponent* GetComponent(FName InName) const;
+	DYNAMICMATERIALEDITOR_API virtual UDMMaterialComponent* GetComponent(FName InName) const;
 
 	template<typename InClassName>
 	InClassName* GetTypedComponent(FName InName) const
@@ -162,17 +165,17 @@ public:
 	 * @param InName The name of the component.
 	 * @return The removed component.
 	 */
-	virtual UDMMaterialComponent* RemoveComponent(FName InName);
+	DYNAMICMATERIALEDITOR_API virtual UDMMaterialComponent* RemoveComponent(FName InName);
 
 	/**
 	 * Generates the base expressions for this slot.
 	 */
-	virtual void GenerateExpressions(const TSharedRef<FDMMaterialBuildState>& InBuildState) const;
+	DYNAMICMATERIALEDITOR_API virtual void GenerateExpressions(const TSharedRef<FDMMaterialBuildState>& InBuildState) const;
 
 	/**
 	 * Adds expressions to opacity properties based on color channel.
 	 */
-	static void GenerateOpacityExpressions(const TSharedRef<FDMMaterialBuildState>& InBuildState, UDMMaterialSlot* InFromSlot,
+	DYNAMICMATERIALEDITOR_API static void GenerateOpacityExpressions(const TSharedRef<FDMMaterialBuildState>& InBuildState, UDMMaterialSlot* InFromSlot,
 		EDMMaterialPropertyType InFromProperty, UMaterialExpression*& OutExpression, int32& OutOutputIndex, int32& OutOutputChannel);
 
 	/**
@@ -186,21 +189,20 @@ public:
 	virtual void AddOutputProcessor(const TSharedRef<FDMMaterialBuildState>& InBuildState) const;
 
 	//~ Begin UDMMaterialComponent
-	virtual void Update(EDMUpdateType InUpdateType) override;
-	virtual bool IsRootComponent() const override { return true; }
-	virtual FString GetComponentPathComponent() const override;
-	virtual void PostEditorDuplicate(UDynamicMaterialModel* InMaterialModel, UDMMaterialComponent* InParent) override;
-	virtual void PreEditChange(FProperty* InPropertyAboutToChange) override;
-	virtual void PostEditChangeProperty(FPropertyChangedEvent& InPropertyChangedEvent) override;
+	DYNAMICMATERIALEDITOR_API virtual void Update(EDMUpdateType InUpdateType) override;
+	DYNAMICMATERIALEDITOR_API virtual FString GetComponentPathComponent() const override;
+	DYNAMICMATERIALEDITOR_API virtual void PostEditorDuplicate(UDynamicMaterialModel* InMaterialModel, UDMMaterialComponent* InParent) override;
+	DYNAMICMATERIALEDITOR_API virtual void PreEditChange(FProperty* InPropertyAboutToChange) override;
+	DYNAMICMATERIALEDITOR_API virtual void PostEditChangeProperty(FPropertyChangedEvent& InPropertyChangedEvent) override;
 	//~ End UDMMaterialComponent
 
 	void LoadDeprecatedModelData(UDMMaterialProperty* InOldProperty);
 
 protected:
-	static UMaterialExpression* CreateConstant(const TSharedRef<FDMMaterialBuildState>& InBuildState, float InDefaultValue);
-	static UMaterialExpression* CreateConstant(const TSharedRef<FDMMaterialBuildState>& InBuildState, const FVector2d& InDefaultValue);
-	static UMaterialExpression* CreateConstant(const TSharedRef<FDMMaterialBuildState>& InBuildState, const FVector3d& InDefaultValue);
-	static UMaterialExpression* CreateConstant(const TSharedRef<FDMMaterialBuildState>& InBuildState, const FVector4d& InDefaultValue);
+	DYNAMICMATERIALEDITOR_API static UMaterialExpression* CreateConstant(const TSharedRef<FDMMaterialBuildState>& InBuildState, float InDefaultValue);
+	DYNAMICMATERIALEDITOR_API static UMaterialExpression* CreateConstant(const TSharedRef<FDMMaterialBuildState>& InBuildState, const FVector2d& InDefaultValue);
+	DYNAMICMATERIALEDITOR_API static UMaterialExpression* CreateConstant(const TSharedRef<FDMMaterialBuildState>& InBuildState, const FVector3d& InDefaultValue);
+	DYNAMICMATERIALEDITOR_API static UMaterialExpression* CreateConstant(const TSharedRef<FDMMaterialBuildState>& InBuildState, const FVector4d& InDefaultValue);
 
 	/**
 	 * The property type of this property.
@@ -226,7 +228,7 @@ protected:
 	UPROPERTY(VisibleInstanceOnly, BlueprintReadOnly, Category = "Material Designer")
 	TObjectPtr<UMaterialFunctionInterface> OutputProcessor;
 
-	UPROPERTY(Transient, DuplicateTransient, TextExportTransient)
+	UPROPERTY(Transient, TextExportTransient)
 	TObjectPtr<UMaterialFunctionInterface> OutputProcessor_PreUpdate;
 
 	/**
@@ -235,8 +237,7 @@ protected:
 	UPROPERTY(VisibleInstanceOnly, BlueprintReadOnly, Category = "Material Designer")
 	TMap<FName, TObjectPtr<UDMMaterialComponent>> Components;
 
-	UDMMaterialProperty();
-	UDMMaterialProperty(EDMMaterialPropertyType InMaterialProperty, EDMValueType InInputConnectorType);
+	DYNAMICMATERIALEDITOR_API UDMMaterialProperty(EDMMaterialPropertyType InMaterialProperty, EDMValueType InInputConnectorType);
 
 	/**
 	 * Triggered when the output process is updated either through PostEditChangeProperty or via the setter.
@@ -245,8 +246,8 @@ protected:
 	void OnOutputProcessorUpdated();
 
 	//~ Begin UDMMaterialComponent
-	virtual UDMMaterialComponent* GetSubComponentByPath(FDMComponentPath& InPath, const FDMComponentPathSegment& InPathSegment) const override;
-	virtual void OnComponentAdded() override;
-	virtual void OnComponentRemoved() override;
+	DYNAMICMATERIALEDITOR_API virtual UDMMaterialComponent* GetSubComponentByPath(FDMComponentPath& InPath, const FDMComponentPathSegment& InPathSegment) const override;
+	DYNAMICMATERIALEDITOR_API virtual void OnComponentAdded() override;
+	DYNAMICMATERIALEDITOR_API virtual void OnComponentRemoved() override;
 	//~ End UDMMaterialComponent
 };
diff --git a/Engine/Plugins/Experimental/DynamicMaterial/Source/DynamicMaterialEditor/Public/Components/DMMaterialSlot.h b/Engine/Plugins/Experimental/DynamicMaterial/Source/DynamicMaterialEditor/Public/Components/DMMaterialSlot.h
--- a/Engine/Plugins/Experimental/DynamicMaterial/Source/DynamicMaterialEditor/Public/Components/DMMaterialSlot.h
+++ b/Engine/Plugins/Experimental/DynamicMaterial/Source/DynamicMaterialEditor/Public/Components/DMMaterialSlot.h
@@ -22,8 +22,8 @@ DECLARE_MULTICAST_DELEGATE_OneParam(FDMOnMaterialSlotLayersUpdated, UDMMaterialS
 /**
  * A list of operations/inputs daisy chained together to produce an output.
  */
-UCLASS(BlueprintType, ClassGroup = "Material Designer", Meta = (DisplayName = "Material Designer Slot"))
-class DYNAMICMATERIALEDITOR_API UDMMaterialSlot : public UDMMaterialComponent
+UCLASS(MinimalAPI, BlueprintType, ClassGroup = "Material Designer", Meta = (DisplayName = "Material Designer Slot"))
+class UDMMaterialSlot : public UDMMaterialComponent
 {
 	friend class SDMMaterialSlot;
 	friend class SDMSlot;
@@ -31,80 +31,100 @@ class DYNAMICMATERIALEDITOR_API UDMMaterialSlot : public UDMMaterialComponent
 	GENERATED_BODY()
 
 public:
-	static const FString LayersPathToken;
+	DYNAMICMATERIALEDITOR_API static const FString LayersPathToken;
 
-	UDMMaterialSlot();
+	DYNAMICMATERIALEDITOR_API UDMMaterialSlot();
 
 	UFUNCTION(BlueprintPure, Category = "Material Designer")
-	UDynamicMaterialModelEditorOnlyData* GetMaterialModelEditorOnlyData() const;
+	DYNAMICMATERIALEDITOR_API UDynamicMaterialModelEditorOnlyData* GetMaterialModelEditorOnlyData() const;
 
+	/** Returns the index of this slot in the model. */
 	UFUNCTION(BlueprintPure, Category = "Material Designer")
 	int32 GetIndex() const { return Index; }
 
 	void SetIndex(int32 InNewIndex) { Index = InNewIndex; }
 
 	UFUNCTION(BlueprintPure, Category = "Material Designer")
-	FText GetDescription() const;
+	DYNAMICMATERIALEDITOR_API FText GetDescription() const;
 
+	/** Returns the output types for the last layer with the given material property. */
 	UFUNCTION(BlueprintPure, Category = "Material Designer")
-	const TArray<EDMValueType>& GetOutputConnectorTypesForMaterialProperty(EDMMaterialPropertyType InMaterialProperty) const;
+	DYNAMICMATERIALEDITOR_API const TArray<EDMValueType>& GetOutputConnectorTypesForMaterialProperty(EDMMaterialPropertyType InMaterialProperty) const;
 
+	/** Returns all possible output connector types. */
 	UFUNCTION(BlueprintPure, Category = "Material Designer")
-	TSet<EDMValueType> GetAllOutputConnectorTypes() const;
+	DYNAMICMATERIALEDITOR_API TSet<EDMValueType> GetAllOutputConnectorTypes() const;
 
 	UFUNCTION(BlueprintPure, Category = "Material Designer", meta = (DisplayName = "Get Layer"))
-	UDMMaterialLayerObject* GetLayer(int32 InLayerIndex) const;
+	DYNAMICMATERIALEDITOR_API UDMMaterialLayerObject* GetLayer(int32 InLayerIndex) const;
 
 	UFUNCTION(BlueprintCallable, Category = "Material Designer")
-	UDMMaterialLayerObject* FindLayer(const UDMMaterialStage* InBaseOrMask) const;
+	DYNAMICMATERIALEDITOR_API UDMMaterialLayerObject* FindLayer(const UDMMaterialStage* InBaseOrMask) const;
 
 	UFUNCTION(BlueprintPure, Category = "Material Designer", meta = (DisplayName = "Get Layers"))
-	TArray<UDMMaterialLayerObject*> BP_GetLayers() const;
+	DYNAMICMATERIALEDITOR_API TArray<UDMMaterialLayerObject*> BP_GetLayers() const;
 
 	const TArray<TObjectPtr<UDMMaterialLayerObject>>& GetLayers() const { return LayerObjects; }
 
+	/** Adds the default layer type for this slot based on the given material property. */
 	UFUNCTION(BlueprintCallable, Category = "Material Designer")
-	UDMMaterialLayerObject* AddDefaultLayer(EDMMaterialPropertyType InMaterialProperty);
+	DYNAMICMATERIALEDITOR_API UDMMaterialLayerObject* AddDefaultLayer(EDMMaterialPropertyType InMaterialProperty);
 
+	/** Adds the default layer (with specified base) based on the given material property. */
 	UFUNCTION(BlueprintCallable, Category = "Material Designer")
-	UDMMaterialLayerObject* AddLayer(EDMMaterialPropertyType InMaterialProperty, UDMMaterialStage* InNewBase);
+	DYNAMICMATERIALEDITOR_API UDMMaterialLayerObject* AddLayer(EDMMaterialPropertyType InMaterialProperty, UDMMaterialStage* InNewBase);
 
+	/** Adds a new layer with the specified base and mask layers. */
 	UFUNCTION(BlueprintCallable, Category = "Material Designer")
-	UDMMaterialLayerObject* AddLayerWithMask(EDMMaterialPropertyType InMaterialProperty, UDMMaterialStage* InNewBase, UDMMaterialStage* InNewMask);
+	DYNAMICMATERIALEDITOR_API UDMMaterialLayerObject* AddLayerWithMask(EDMMaterialPropertyType InMaterialProperty, UDMMaterialStage* InNewBase,
+		UDMMaterialStage* InNewMask);
 
-	// Adds the specified layer to the end of the layer list
-	UFUNCTION(BlueprintCallable, Category = "Material Designer")
+	/** Adds the specified layer to the end of the layer list. */
 	bool PasteLayer(UDMMaterialLayerObject* InLayer);
 
+	/** Can't be removed if it is the last remaining layer. */
 	UFUNCTION(BlueprintCallable, Category = "Material Designer")
-	bool CanRemoveLayer(const UDMMaterialLayerObject* InLayer) const;
+	DYNAMICMATERIALEDITOR_API bool CanRemoveLayer(const UDMMaterialLayerObject* InLayer) const;
 
+	/** Removes the layer, if possible. */
 	UFUNCTION(BlueprintCallable, Category = "Material Designer")
-	bool RemoveLayer(UDMMaterialLayerObject* InLayer);
+	DYNAMICMATERIALEDITOR_API bool RemoveLayer(UDMMaterialLayerObject* InLayer);
 
 	UFUNCTION(BlueprintCallable, Category = "Material Designer")
-	bool MoveLayer(UDMMaterialLayerObject* InLayer, int32 InNewIndex);
+	DYNAMICMATERIALEDITOR_API bool MoveLayer(UDMMaterialLayerObject* InLayer, int32 InNewIndex);
 
 	UFUNCTION(BlueprintCallable, Category = "Material Designer")
-	bool MoveLayerBefore(UDMMaterialLayerObject* InLayer, UDMMaterialLayerObject* InBeforeLayer = nullptr);
+	DYNAMICMATERIALEDITOR_API bool MoveLayerBefore(UDMMaterialLayerObject* InLayer, UDMMaterialLayerObject* InBeforeLayer = nullptr);
 
 	UFUNCTION(BlueprintCallable, Category = "Material Designer")
-	bool MoveLayerAfter(UDMMaterialLayerObject* InLayer, UDMMaterialLayerObject* InAfterLayer = nullptr);
+	DYNAMICMATERIALEDITOR_API bool MoveLayerAfter(UDMMaterialLayerObject* InLayer, UDMMaterialLayerObject* InAfterLayer = nullptr);
 
+	/** Useful for determining output types. */
 	UFUNCTION(BlueprintCallable, Category = "Material Designer")
-	UDMMaterialLayerObject* GetLastLayerForMaterialProperty(EDMMaterialPropertyType InMaterialProperty) const;
+	DYNAMICMATERIALEDITOR_API UDMMaterialLayerObject* GetLastLayerForMaterialProperty(EDMMaterialPropertyType InMaterialProperty) const;
 
 	void UpdateOutputConnectorTypes();
+
 	void UpdateMaterialProperties();
 
-	FDMOnMaterialSlotConnectorsUpdated& GetOnConnectorsUpdateDelegate() { return OnConnectorsUpdateDelegate; }
-	FDMOnMaterialSlotPropertiesUpdated& GetOnPropertiesUpdateDelegate() { return OnPropertiesUpdateDelegate; }
-	FDMOnMaterialSlotLayersUpdated& GetOnLayersUpdateDelegate() { return OnLayersUpdateDelegate; }
+	/** Called when the output connectors for this slot change. */
+	FDMOnMaterialSlotConnectorsUpdated::RegistrationType& GetOnConnectorsUpdateDelegate() { return OnConnectorsUpdateDelegate; }
+
+	/** Called when properties of this slot change. */
+	FDMOnMaterialSlotPropertiesUpdated::RegistrationType& GetOnPropertiesUpdateDelegate() { return OnPropertiesUpdateDelegate; }
+
+	/** Called whenever the properties of a layer change or when one is added, removed or moved. */
+	FDMOnMaterialSlotLayersUpdated::RegistrationType& GetOnLayersUpdateDelegate() { return OnLayersUpdateDelegate; }
+
+	/** Calls OnPropertiesUpdateDelegate when the property for this slot is updated. */
+	void OnPropertiesUpdated();
 
-	void GenerateExpressions(const TSharedRef<FDMMaterialBuildState>& InBuildState) const;
+	DYNAMICMATERIALEDITOR_API void GenerateExpressions(const TSharedRef<FDMMaterialBuildState>& InBuildState) const;
 
+	/** Return a map of the slots referencing this slot and how many times that reference exists. */
 	const TMap<TWeakObjectPtr<UDMMaterialSlot>, int32>& GetSlotsReferencedBy() const { return SlotsReferencedBy; }
 
+	/** Returns an array of the slots referencing this slot. */
 	UFUNCTION(BlueprintCallable, Category = "Material Designer", Meta = (DisplayName="Get Slots Referenced By"))
 	TArray<UDMMaterialSlot*> K2_GetSlotsReferencedBy() const;
 
@@ -115,31 +135,32 @@ public:
 	bool UnreferencedBySlot(UDMMaterialSlot* InOtherSlot);
 
 	bool IsEditingLayers() const { return bIsEditingLayers; }
+
 	void SetEditingLayers(bool bInIsEditing) { bIsEditingLayers = bInIsEditing; }
 
-	UFUNCTION(BlueprintPure, Category = "Material Designer")
-	UMaterial* GetPreviewMaterial(EDMMaterialLayerStage InLayerStage);
+	UFUNCTION(BlueprintCallable, Category = "Material Designer")
+	DYNAMICMATERIALEDITOR_API void GeneralPreviewMaterial(UMaterial* InPreviewMaterial, EDMMaterialLayerStage InLayerStage);
 
+	/** Sets the material property of the given layer and changes all other layers matching that property to a different one. */
 	UFUNCTION(BlueprintCallable, Category = "Material Designer")
-	bool SetLayerMaterialPropertyAndReplaceOthers(UDMMaterialLayerObject* InLayer, EDMMaterialPropertyType InPropertyFrom,
+	DYNAMICMATERIALEDITOR_API bool SetLayerMaterialPropertyAndReplaceOthers(UDMMaterialLayerObject* InLayer, EDMMaterialPropertyType InPropertyFrom,
 		EDMMaterialPropertyType InPropertyTo);
 
+	/** Changes the material property of all matching layers to another. */
 	UFUNCTION(BlueprintCallable, Category = "Material Designer")
-	bool ChangeMaterialProperty(EDMMaterialPropertyType InPropertyFrom, EDMMaterialPropertyType InPropertyTo);
+	DYNAMICMATERIALEDITOR_API bool ChangeMaterialProperty(EDMMaterialPropertyType InPropertyFrom, EDMMaterialPropertyType InPropertyTo);
 
 	//~ Begin UDMMaterialComponent
-	virtual void Update(EDMUpdateType InUpdateType) override;
-	virtual void DoClean() override;
-	virtual bool IsRootComponent() const override { return true; }
-	virtual FString GetComponentPathComponent() const override;
+	DYNAMICMATERIALEDITOR_API virtual void Update(EDMUpdateType InUpdateType) override;
+	DYNAMICMATERIALEDITOR_API virtual FString GetComponentPathComponent() const override;
 	virtual UDMMaterialComponent* GetParentComponent() const override { return nullptr; }
-	virtual void PostEditorDuplicate(UDynamicMaterialModel* InMaterialModel, UDMMaterialComponent* InParent) override;
+	DYNAMICMATERIALEDITOR_API virtual void PostEditorDuplicate(UDynamicMaterialModel* InMaterialModel, UDMMaterialComponent* InParent) override;
 	//~ End UDMMaterialComponent
 
 	//~ Begin UObject
-	virtual bool Modify(bool bInAlwaysMarkDirty = true) override;
-	virtual void PostEditUndo() override;
-	virtual void PostLoad() override;
+	DYNAMICMATERIALEDITOR_API virtual bool Modify(bool bInAlwaysMarkDirty = true) override;
+	DYNAMICMATERIALEDITOR_API virtual void PostEditUndo() override;
+	DYNAMICMATERIALEDITOR_API virtual void PostLoad() override;
 	//~ End UObject
 
 protected:
@@ -155,33 +176,25 @@ protected:
 	UPROPERTY()
 	TMap<TWeakObjectPtr<UDMMaterialSlot>, int32> SlotsReferencedBy;
 
-	UPROPERTY(VisibleInstanceOnly, BlueprintReadOnly, Transient, DuplicateTransient, TextExportTransient, Category = "Material Designer")
+	UPROPERTY(VisibleInstanceOnly, BlueprintReadOnly, Transient, TextExportTransient, Category = "Material Designer")
 	bool bIsEditingLayers;
 
-	UPROPERTY(VisibleInstanceOnly, BlueprintReadOnly, Transient, DuplicateTransient, TextExportTransient, Category = "Material Designer")
-	TObjectPtr<UMaterial> BasePreviewMaterial;
-
-	UPROPERTY(VisibleInstanceOnly, BlueprintReadOnly, Transient, DuplicateTransient, TextExportTransient, Category = "Material Designer")
-	TObjectPtr<UMaterial> MaskPreviewMaterial;
-
 	FDMOnMaterialSlotConnectorsUpdated OnConnectorsUpdateDelegate;
 	FDMOnMaterialSlotPropertiesUpdated OnPropertiesUpdateDelegate;
 	FDMOnMaterialSlotLayersUpdated OnLayersUpdateDelegate;
 
-	void CreatePreviewMaterial(EDMMaterialLayerStage InLayerStage);
-	void UpdatePreviewMaterial(EDMMaterialLayerStage InLayerStage);
 	void UpdateBasePreviewMaterial(const TSharedRef<FDMMaterialBuildState>& InBuildState);
 	void UpdateBasePreviewMaterialProperty(const TSharedRef<FDMMaterialBuildState>& InBuildState, EDMMaterialPropertyType InBaseProperty);
 	void UpdateBasePreviewMaterialFull(const TSharedRef<FDMMaterialBuildState>& InBuildState);
 	void UpdateMaskPreviewMaterial(const TSharedRef<FDMMaterialBuildState>& InBuildState);
 	void UpdateMaskPreviewMaterialProperty(const TSharedRef<FDMMaterialBuildState>& InBuildState, EDMMaterialPropertyType InMaskProperty);
 	void UpdateMaskPreviewMaterialMaskCombination(const TSharedRef<FDMMaterialBuildState>& InBuildState, EDMMaterialPropertyType InMaskProperty);
-	void UpdatePreviewMaterialProperty(const TSharedRef<FDMMaterialBuildState>& InBuildState, UMaterial* InPreviewMaterial, EDMMaterialPropertyType InProperty);
+	void UpdatePreviewMaterialProperty(const TSharedRef<FDMMaterialBuildState>& InBuildState, EDMMaterialPropertyType InProperty);
 
 	//~ Begin UDMMaterialComponent
-	virtual void OnComponentAdded() override;
-	virtual void OnComponentRemoved() override;
-	virtual UDMMaterialComponent* GetSubComponentByPath(FDMComponentPath& InPath, const FDMComponentPathSegment& InPathSegment) const override;
+	DYNAMICMATERIALEDITOR_API virtual void OnComponentAdded() override;
+	DYNAMICMATERIALEDITOR_API virtual void OnComponentRemoved() override;
+	DYNAMICMATERIALEDITOR_API virtual UDMMaterialComponent* GetSubComponentByPath(FDMComponentPath& InPath, const FDMComponentPathSegment& InPathSegment) const override;
 	//~ End UDMMaterialComponent
 
 private:
diff --git a/Engine/Plugins/Experimental/DynamicMaterial/Source/DynamicMaterialEditor/Public/Components/DMMaterialStage.h b/Engine/Plugins/Experimental/DynamicMaterial/Source/DynamicMaterialEditor/Public/Components/DMMaterialStage.h
--- a/Engine/Plugins/Experimental/DynamicMaterial/Source/DynamicMaterialEditor/Public/Components/DMMaterialStage.h
+++ b/Engine/Plugins/Experimental/DynamicMaterial/Source/DynamicMaterialEditor/Public/Components/DMMaterialStage.h
@@ -37,24 +37,26 @@ struct FDMMaterialBuildState;
 struct FDMTextureUV;
 
 /**
- * A node which handles a specific operation and manages its inputs and outputs.
+ * A component which wraps a source and its inputs.
  */
-UCLASS(BlueprintType, ClassGroup = "Material Designer", meta = (DisplayName = "Material Designer Stage"))
-class DYNAMICMATERIALEDITOR_API UDMMaterialStage : public UDMMaterialComponent
+UCLASS(MinimalAPI, BlueprintType, ClassGroup = "Material Designer", meta = (DisplayName = "Material Designer Stage"))
+class UDMMaterialStage : public UDMMaterialComponent
 {
 	GENERATED_BODY()
 
 	friend class FDMThroughputPropertyRowGenerator;
 
 public:
-	static const FString SourcePathToken;
-	static const FString InputsPathToken;
+	DYNAMICMATERIALEDITOR_API static const FString SourcePathToken;
+	DYNAMICMATERIALEDITOR_API static const FString InputsPathToken;
 
 	UFUNCTION(BlueprintCallable, Category = "Material Designer")
-	static UDMMaterialStage* CreateMaterialStage(UDMMaterialLayerObject* InLayer = nullptr);
+	static DYNAMICMATERIALEDITOR_API UDMMaterialStage* CreateMaterialStage(UDMMaterialLayerObject* InLayer = nullptr);
+
+	DYNAMICMATERIALEDITOR_API UDMMaterialStage();
 
 	UFUNCTION(BlueprintPure, Category = "Material Designer")
-	UDMMaterialLayerObject* GetLayer() const;
+	DYNAMICMATERIALEDITOR_API UDMMaterialLayerObject* GetLayer() const;
 
 	UFUNCTION(BlueprintPure, Category = "Material Designer")
 	UDMMaterialStageSource* GetSource() const { return Source; }
@@ -63,7 +65,7 @@ public:
 	bool IsEnabled() const { return bEnabled; }
 
 	UFUNCTION(BlueprintCallable, Category = "Material Designer")
-	bool SetEnabled(bool bInEnabled);
+	DYNAMICMATERIALEDITOR_API bool SetEnabled(bool bInEnabled);
 
 	UFUNCTION(BlueprintPure, Category = "Material Designer")
 	bool CanChangeSource() const { return bCanChangeSource; }
@@ -72,11 +74,7 @@ public:
 	void SetCanChangeSource(bool bInCanChangeSource) { bCanChangeSource = bInCanChangeSource; }
 
 	UFUNCTION(BlueprintCallable, Category = "Material Designer")
-	void SetSource(UDMMaterialStageSource* InSource);
-
-	//~ Begin UDMMaterialComponent
-	virtual FText GetComponentDescription() const override;
-	//~ End UDMMaterialComponent
+	DYNAMICMATERIALEDITOR_API void SetSource(UDMMaterialStageSource* InSource);
 
 	UFUNCTION(BlueprintPure, Category = "Material Designer")
 	const TArray<UDMMaterialStageInput*>& GetInputs() const { return Inputs; }
@@ -88,32 +86,41 @@ public:
 	TArray<FDMMaterialStageConnection>& GetInputConnectionMap() { return InputConnectionMap; }
 
 	UFUNCTION(BlueprintCallable, Category = "Material Designer")
-	EDMValueType GetSourceType(const FDMMaterialStageConnectorChannel& InChannel) const;
+	DYNAMICMATERIALEDITOR_API EDMValueType GetSourceType(const FDMMaterialStageConnectorChannel& InChannel) const;
 
+	/** Returns true if the given source's input is mapped to an input (or the previous stage). */
 	UFUNCTION(BlueprintCallable, Category = "Material Designer")
-	bool IsInputMapped(int32 InputIndex) const;
+	DYNAMICMATERIALEDITOR_API bool IsInputMapped(int32 InputIndex) const;
 
+	/**
+	 * Returns true if the output of the previous stage can connect to this stage.
+	 * It is now up to the user to sort this particular problem out because it would do more harm than good
+	 * to force correctness in "transition states" while the user is changing settings.
+	 */
 	UFUNCTION(BlueprintPure, Category = "Material Designer")
-	virtual bool IsCompatibleWithPreviousStage(const UDMMaterialStage* InPreviousStage) const;
+	DYNAMICMATERIALEDITOR_API virtual bool IsCompatibleWithPreviousStage(const UDMMaterialStage* InPreviousStage) const;
 
+	/* @see IsCompatibleWithPreviousStage */
 	UFUNCTION(BlueprintPure, Category = "Material Designer")
-	virtual bool IsCompatibleWithNextStage(const UDMMaterialStage* InNextStage) const;
+	DYNAMICMATERIALEDITOR_API virtual bool IsCompatibleWithNextStage(const UDMMaterialStage* InNextStage) const;
 
 	UFUNCTION(BlueprintCallable, Category = "Material Designer")
-	void AddInput(UDMMaterialStageInput* InNewInput);
+	DYNAMICMATERIALEDITOR_API void AddInput(UDMMaterialStageInput* InNewInput);
 
 	UFUNCTION(BlueprintCallable, Category = "Material Designer")
-	void RemoveInput(UDMMaterialStageInput* InInput);
+	DYNAMICMATERIALEDITOR_API void RemoveInput(UDMMaterialStageInput* InInput);
 
 	UFUNCTION(BlueprintCallable, Category = "Material Designer")
-	void RemoveAllInputs();
+	DYNAMICMATERIALEDITOR_API void RemoveAllInputs();
 
+	/** Called when one of the inputs triggers it's Update event. */
 	virtual void InputUpdated(UDMMaterialStageInput* InInput, EDMUpdateType InUpdateType);
 
+	/** Verifies the entire input connection map. */
 	UFUNCTION(BlueprintCallable, Category = "Material Designer")
-	virtual void ResetInputConnectionMap();
+	DYNAMICMATERIALEDITOR_API virtual void ResetInputConnectionMap();
 
-	void GenerateExpressions(const TSharedRef<FDMMaterialBuildState>& InBuildState) const;
+	DYNAMICMATERIALEDITOR_API void GenerateExpressions(const TSharedRef<FDMMaterialBuildState>& InBuildState) const;
 
 	UFUNCTION(BlueprintPure, Category = "Material Designer")
 	bool IsBeingEdited() const { return bIsBeingEdited; }
@@ -121,8 +128,11 @@ public:
 	UFUNCTION(BlueprintCallable, Category = "Material Designer")
 	bool SetBeingEdited(bool bInBeingEdited);
 
-	TMap<EDMMaterialPropertyType, UDMMaterialLayerObject*> GetPreviousStagesPropertyMap();
-	TMap<EDMMaterialPropertyType, UDMMaterialLayerObject*> GetPropertyMap();
+	/** Get the last layer for each property type from the previous stages. */
+	DYNAMICMATERIALEDITOR_API TMap<EDMMaterialPropertyType, UDMMaterialLayerObject*> GetPreviousStagesPropertyMap();
+
+	/** Get the last layer for each property type from all stages. */
+	DYNAMICMATERIALEDITOR_API TMap<EDMMaterialPropertyType, UDMMaterialLayerObject*> GetPropertyMap();
 	 
 	using FSourceInitFunctionPtr = TFunction<void(UDMMaterialStage*, UDMMaterialStageSource*)>;
 
@@ -132,7 +142,8 @@ public:
 		return ChangeSource(InSourceClass, nullptr);
 	}
 
-	UDMMaterialStageSource* ChangeSource(TSubclassOf<UDMMaterialStageSource> InSourceClass, FSourceInitFunctionPtr InPreInit);
+	DYNAMICMATERIALEDITOR_API UDMMaterialStageSource* ChangeSource(TSubclassOf<UDMMaterialStageSource> InSourceClass, 
+		FSourceInitFunctionPtr InPreInit);
 
 	template<typename InSourceClass>
 	InSourceClass* ChangeSource(FSourceInitFunctionPtr InPreInit = nullptr)
@@ -155,7 +166,15 @@ public:
 		return ChangeInput(InInputClass, InInputIdx, InInputChannel, InOutputIdx, InOutputChannel, nullptr);
 	}
 
-	UDMMaterialStageInput* ChangeInput(TSubclassOf<UDMMaterialStageInput> InInputClass, int32 InInputIdx, int32 InInputChannel,
+	/**
+	 * Creates a new input value and maps it to a specific source input.
+	 * @param InInputIdx Index of the source input.
+	 * @param InInputChannel The channel of the input that the input connects to.
+	 * @param InOutputIdx The output index of the new input.
+	 * @param InOutputChannel The channel of the output to connect.
+	 * @param InPreInit Called on the new input before initialisation.
+	 */
+	DYNAMICMATERIALEDITOR_API UDMMaterialStageInput* ChangeInput(TSubclassOf<UDMMaterialStageInput> InInputClass, int32 InInputIdx, int32 InInputChannel,
 		int32 InOutputIdx, int32 InOutputChannel, FInputInitFunctionPtr InPreInit);
 
 	template<typename InInputClass>
@@ -174,58 +193,62 @@ public:
 			InPreInit));
 	}
 
+	/** Changes the input of the given input index to the output of the previous stage with the given material property. */
 	UFUNCTION(BlueprintCallable, Category = "Material Designer")
-	UDMMaterialStageSource* ChangeInput_PreviousStage(int32 InInputIdx, int32 InInputChannel, EDMMaterialPropertyType InPreviousStageProperty, 
+	DYNAMICMATERIALEDITOR_API UDMMaterialStageSource* ChangeInput_PreviousStage(int32 InInputIdx, int32 InInputChannel, EDMMaterialPropertyType InPreviousStageProperty,
 		int32 InOutputIdx, int32 InOutputChannel);
 
 	UFUNCTION(BlueprintCallable, Category = "Material Designer")
-	void RemoveUnusedInputs();
+	DYNAMICMATERIALEDITOR_API void RemoveUnusedInputs();
 
 	/** Returns true if any changes were made */
 	UFUNCTION(BlueprintCallable, Category = "Material Designer")
-	bool VerifyAllInputMaps();
+	DYNAMICMATERIALEDITOR_API bool VerifyAllInputMaps();
 
 	/** Returns true if any changes were made */
 	UFUNCTION(BlueprintCallable, Category = "Material Designer")
-	bool VerifyInputMap(int32 InInputIdx);
-
-	UFUNCTION(BlueprintPure, Category = "Material Designer")
-	UMaterialInterface* GetPreviewMaterial();
-
-	const FDMMaterialStageConnectorChannel* FindInputChannel(UDMMaterialStageInput* InStageInput);
-
-	void UpdateInputMap(int32 InInputIdx, int32 InSourceIndex, int32 InInputChannel, int32 InOutputIdx, int32 InOutputChannel, EDMMaterialPropertyType InStageProperty);
+	DYNAMICMATERIALEDITOR_API bool VerifyInputMap(int32 InInputIdx);
 
 	UFUNCTION(BlueprintCallable, Category = "Material Designer")
-	int32 FindIndex() const;
-
+	DYNAMICMATERIALEDITOR_API void GeneratePreviewMaterial(UMaterial* InPreviewMaterial);
+
+	DYNAMICMATERIALEDITOR_API const FDMMaterialStageConnectorChannel* FindInputChannel(UDMMaterialStageInput* InStageInput);
+
+	/**
+	 * Changes the input mapping.
+	 * @param InInputIdx Index of the source input.
+	 * @param InInputChannel The channel of the input that the input connects to.
+	 * @param InOutputIdx The output index of the new input.
+	 * @param InOutputChannel The channel of the output to connect.
+	 * @param InStageProperty The property for previous stage connections.
+	 */
+	DYNAMICMATERIALEDITOR_API void UpdateInputMap(int32 InInputIdx, int32 InSourceIndex, int32 InInputChannel, int32 InOutputIdx, 
+		int32 InOutputChannel, EDMMaterialPropertyType InStageProperty);
+
+	/** Returns the index of this stage in the layer. */
 	UFUNCTION(BlueprintCallable, Category = "Material Designer")
-	UDMMaterialStage* GetPreviousStage() const;
+	DYNAMICMATERIALEDITOR_API int32 FindIndex() const;
 
 	UFUNCTION(BlueprintCallable, Category = "Material Designer")
-	UDMMaterialStage* GetNextStage() const;
+	DYNAMICMATERIALEDITOR_API UDMMaterialStage* GetPreviousStage() const;
 
-	UFUNCTION(BlueprintPure, Category = "Material Designer")
-	virtual bool IsRootStage() const;
+	UFUNCTION(BlueprintCallable, Category = "Material Designer")
+	DYNAMICMATERIALEDITOR_API UDMMaterialStage* GetNextStage() const;
 
 	//~ Begin UDMMaterialComponent
-	virtual void Update(EDMUpdateType InUpdateType) override;
-	virtual void DoClean() override;
-	virtual FString GetComponentPathComponent() const override;
-	virtual UDMMaterialComponent* GetParentComponent() const override;
-	virtual void PostEditorDuplicate(UDynamicMaterialModel* InMaterialModel, UDMMaterialComponent* InParent) override;
+	DYNAMICMATERIALEDITOR_API virtual FText GetComponentDescription() const override;
+	DYNAMICMATERIALEDITOR_API virtual void Update(EDMUpdateType InUpdateType) override;
+	DYNAMICMATERIALEDITOR_API virtual FString GetComponentPathComponent() const override;
+	DYNAMICMATERIALEDITOR_API virtual UDMMaterialComponent* GetParentComponent() const override;
+	DYNAMICMATERIALEDITOR_API virtual void PostEditorDuplicate(UDynamicMaterialModel* InMaterialModel, UDMMaterialComponent* InParent) override;
 	//~ End UDMMaterialComponent
 
 	//~ Begin UObject
-	virtual bool Modify(bool bInAlwaysMarkDirty = true) override;
-	virtual void PostEditUndo() override;
-	virtual void PostLoad() override;
-	virtual void PostEditImport() override;
+	DYNAMICMATERIALEDITOR_API virtual bool Modify(bool bInAlwaysMarkDirty = true) override;
+	DYNAMICMATERIALEDITOR_API virtual void PostEditUndo() override;
 	//~ End UObject
 
 protected:
-	UDMMaterialStage();
-
 	UPROPERTY(VisibleInstanceOnly, BlueprintReadOnly, Instanced, Category = "Material Designer")
 	TObjectPtr<UDMMaterialStageSource> Source;
 
@@ -242,28 +265,13 @@ protected:
 	UPROPERTY(VisibleInstanceOnly, BlueprintReadOnly, Category = "Material Designer")
 	bool bCanChangeSource;
 
-	UPROPERTY(VisibleInstanceOnly, BlueprintReadOnly, Transient, DuplicateTransient, TextExportTransient, Category = "Material Designer")
+	UPROPERTY(VisibleInstanceOnly, BlueprintReadOnly, Transient, TextExportTransient, Category = "Material Designer")
 	bool bIsBeingEdited;
 
-	UPROPERTY(VisibleInstanceOnly, BlueprintReadOnly, Transient, DuplicateTransient, TextExportTransient, Category = "Material Designer")
-	TObjectPtr<UMaterial> PreviewMaterialBase;
-
-	UPROPERTY(VisibleInstanceOnly, BlueprintReadOnly, Transient, DuplicateTransient, TextExportTransient, Category = "Material Designer")
-	TObjectPtr<UMaterialInstanceDynamic> PreviewMaterialDynamic;
-
-	void CreatePreviewMaterial();
-	void UpdatePreviewMaterial();
-
 	//~ Begin UDMMaterialComponent
-	virtual void OnComponentAdded() override;
-	virtual void OnComponentRemoved() override;
-	virtual void GetComponentPathInternal(TArray<FString>& OutChildComponentPathComponents) const override;
-	virtual UDMMaterialComponent* GetSubComponentByPath(FDMComponentPath& InPath, const FDMComponentPathSegment& InPathSegment) const override;
+	DYNAMICMATERIALEDITOR_API virtual void OnComponentAdded() override;
+	DYNAMICMATERIALEDITOR_API virtual void OnComponentRemoved() override;
+	DYNAMICMATERIALEDITOR_API virtual void GetComponentPathInternal(TArray<FString>& OutChildComponentPathComponents) const override;
+	DYNAMICMATERIALEDITOR_API virtual UDMMaterialComponent* GetSubComponentByPath(FDMComponentPath& InPath, const FDMComponentPathSegment& InPathSegment) const override;
 	//~ End UDMMaterialComponent
-
-	virtual void AddDelegates();
-	virtual void RemoveDelegates();
-
-	void OnValueUpdated(UDynamicMaterialModel* InMaterialModel, UDMMaterialValue* InValue);
-	void OnTextureUVUpdated(UDynamicMaterialModel* InMaterialModel, UDMTextureUV* InTextureUV);
 };
diff --git a/Engine/Plugins/Experimental/DynamicMaterial/Source/DynamicMaterialEditor/Public/Components/DMMaterialStageBlend.h b/Engine/Plugins/Experimental/DynamicMaterial/Source/DynamicMaterialEditor/Public/Components/DMMaterialStageBlend.h
--- a/Engine/Plugins/Experimental/DynamicMaterial/Source/DynamicMaterialEditor/Public/Components/DMMaterialStageBlend.h
+++ b/Engine/Plugins/Experimental/DynamicMaterial/Source/DynamicMaterialEditor/Public/Components/DMMaterialStageBlend.h
@@ -17,8 +17,8 @@ struct FDMMaterialBuildState;
 /**
  * A node which represents a blend operation.
  */
-UCLASS(Abstract, BlueprintType, Blueprintable, ClassGroup = "Material Designer", meta = (DisplayName = "Material Designer Stage Blend"))
-class DYNAMICMATERIALEDITOR_API UDMMaterialStageBlend : public UDMMaterialStageThroughput
+UCLASS(MinimalAPI, Abstract, BlueprintType, Blueprintable, ClassGroup = "Material Designer", meta = (DisplayName = "Material Designer Stage Blend"))
+class UDMMaterialStageBlend : public UDMMaterialStageThroughput
 {
 	GENERATED_BODY()
 
@@ -27,55 +27,56 @@ public:
 	static constexpr int32 InputA = 1;
 	static constexpr int32 InputB = 2;
 
-	static UDMMaterialStage* CreateStage(TSubclassOf<UDMMaterialStageBlend> InMaterialStageBlendClass, UDMMaterialLayerObject* InLayer = nullptr);
+	DYNAMICMATERIALEDITOR_API static UDMMaterialStage* CreateStage(TSubclassOf<UDMMaterialStageBlend> InMaterialStageBlendClass, UDMMaterialLayerObject* InLayer = nullptr);
 
 	static const TArray<TStrongObjectPtr<UClass>>& GetAvailableBlends();
 
-	UDMMaterialValueFloat1* GetInputAlpha() const;
-	UDMMaterialStageInput* GetInputB() const;
+	UDMMaterialStageBlend();
+
+	DYNAMICMATERIALEDITOR_API UDMMaterialValueFloat1* GetInputAlpha() const;
+	DYNAMICMATERIALEDITOR_API UDMMaterialStageInput* GetInputB() const;
 
 	UFUNCTION(BlueprintCallable, Category = "Material Designer")
-	EAvaColorChannel GetBaseChannelOverride() const;
+	DYNAMICMATERIALEDITOR_API EAvaColorChannel GetBaseChannelOverride() const;
 
 	UFUNCTION(BlueprintCallable, Category = "Material Designer")
-	void SetBaseChannelOverride(EAvaColorChannel InMaskChannel);
+	DYNAMICMATERIALEDITOR_API void SetBaseChannelOverride(EAvaColorChannel InMaskChannel);
 
 	//~ Begin UDMMaterialStageThroughput
-	virtual bool CanInputAcceptType(int32 InputIndex, EDMValueType ValueType) const override;
-	virtual void AddDefaultInput(int32 InInputIndex) const override;
-	virtual bool CanChangeInput(int32 InputIndex) const override;
-	virtual bool CanChangeInputType(int32 InputIndex) const override;
-	virtual bool IsInputVisible(int32 InputIndex) const override;
-	virtual int32 ResolveInput(const TSharedRef<FDMMaterialBuildState>& InBuildState, int32 InputIndex, FDMMaterialStageConnectorChannel& OutChannel,
-		TArray<UMaterialExpression*>& OutExpressions) const override;
+	DYNAMICMATERIALEDITOR_API virtual bool CanInputAcceptType(int32 InputIndex, EDMValueType ValueType) const override;
+	DYNAMICMATERIALEDITOR_API virtual void AddDefaultInput(int32 InInputIndex) const override;
+	DYNAMICMATERIALEDITOR_API virtual bool CanChangeInput(int32 InputIndex) const override;
+	DYNAMICMATERIALEDITOR_API virtual bool CanChangeInputType(int32 InputIndex) const override;
+	DYNAMICMATERIALEDITOR_API virtual bool IsInputVisible(int32 InputIndex) const override;
+	DYNAMICMATERIALEDITOR_API virtual int32 ResolveInput(const TSharedRef<FDMMaterialBuildState>& InBuildState, int32 InputIndex, 
+		FDMMaterialStageConnectorChannel& OutChannel, TArray<UMaterialExpression*>& OutExpressions) const override;
 	//~ End UDMMaterialStageThroughput
 
 	//~ Begin UDMMaterialStageSource
-	virtual FText GetStageDescription() const override;
+	DYNAMICMATERIALEDITOR_API virtual FText GetStageDescription() const override;
 	virtual bool SupportsLayerMaskTextureUVLink() const override { return true; }
-	virtual FDMExpressionInput GetLayerMaskLinkTextureUVInputExpressions(const TSharedRef<FDMMaterialBuildState>& InBuildState) const override;
-	virtual void GetMaskAlphaBlendNode(const TSharedRef<FDMMaterialBuildState>& InBuildState, UMaterialExpression*& OutExpression,
-		int32& OutOutputIndex, int32& OutOutputChannel) const override;
-	virtual bool UpdateStagePreviewMaterial(UDMMaterialStage* InStage, UMaterial* InPreviewMaterial, UMaterialExpression*& OutMaterialExpression,
-		int32& OutputIndex) override;
+	DYNAMICMATERIALEDITOR_API virtual FDMExpressionInput GetLayerMaskLinkTextureUVInputExpressions(const TSharedRef<FDMMaterialBuildState>& InBuildState) const override;
+	DYNAMICMATERIALEDITOR_API virtual void GetMaskAlphaBlendNode(const TSharedRef<FDMMaterialBuildState>& InBuildState, 
+		UMaterialExpression*& OutExpression, int32& OutOutputIndex, int32& OutOutputChannel) const override;
+	DYNAMICMATERIALEDITOR_API virtual bool GenerateStagePreviewMaterial(UDMMaterialStage* InStage, UMaterial* InPreviewMaterial, 
+		UMaterialExpression*& OutMaterialExpression, int32& OutputIndex) override;
 	//~ End UDMMaterialStageSource
 
 	//~ Begin FNotifyHook
-	virtual void NotifyPostChange(const FPropertyChangedEvent& InPropertyChangedEvent, class FEditPropertyChain* InPropertyThatChanged) override;
+	DYNAMICMATERIALEDITOR_API virtual void NotifyPostChange(const FPropertyChangedEvent& InPropertyChangedEvent, class FEditPropertyChain* InPropertyThatChanged) override;
 	//~ End FNotifyHook
 
 protected:
 	static TArray<TStrongObjectPtr<UClass>> Blends;
 
-	UPROPERTY(EditInstanceOnly, BlueprintReadOnly, Getter = GetBaseChannelOverride, Setter = SetBaseChannelOverride, BlueprintGetter = GetBaseChannelOverride,
+	UPROPERTY(EditInstanceOnly, BlueprintReadOnly, Getter, Setter, BlueprintGetter = GetBaseChannelOverride,
 		Category = "Material Designer", DisplayName = "Channel Mask",
 		meta = (NotKeyframeable, ToolTip = "Changes the output channel of the base input."))
 	mutable EAvaColorChannel BaseChannelOverride;
 
 	static void GenerateBlendList();
 
-	UDMMaterialStageBlend();
-	UDMMaterialStageBlend(const FText& InName);
+	DYNAMICMATERIALEDITOR_API UDMMaterialStageBlend(const FText& InName);
 
 	/* Returns true if there are any outputs on the base input that have more than 1 channel. */
 	bool CanUseBaseChannelOverride() const;
@@ -92,5 +93,7 @@ protected:
 	/* Takes the override setting and applies it to the input map. */
 	void PushBaseChannelOverride();
 
-	virtual void UpdatePreviewMaterial(UMaterial* InPreviewMaterial = nullptr) override;
-};
\ No newline at end of file
+	//~ Begin UDMMaterialStageThroughput
+	DYNAMICMATERIALEDITOR_API virtual void GeneratePreviewMaterial(UMaterial* InPreviewMaterial) override;
+	//~ End UDMMaterialStageThroughput
+};
diff --git a/Engine/Plugins/Experimental/DynamicMaterial/Source/DynamicMaterialEditor/Public/Components/DMMaterialStageBlendFunction.h b/Engine/Plugins/Experimental/DynamicMaterial/Source/DynamicMaterialEditor/Public/Components/DMMaterialStageBlendFunction.h
--- a/Engine/Plugins/Experimental/DynamicMaterial/Source/DynamicMaterialEditor/Public/Components/DMMaterialStageBlendFunction.h
+++ b/Engine/Plugins/Experimental/DynamicMaterial/Source/DynamicMaterialEditor/Public/Components/DMMaterialStageBlendFunction.h
@@ -9,21 +9,30 @@ class UMaterialExpression;
 class UMaterialFunctionInterface;
 struct FDMMaterialBuildState;
 
-UCLASS(Abstract, BlueprintType, Blueprintable, ClassGroup = "Material Designer")
-class DYNAMICMATERIALEDITOR_API UDMMaterialStageBlendFunction : public UDMMaterialStageBlend
+/** A blending stage based on a material function. */
+UCLASS(MinimalAPI, Abstract, BlueprintType, Blueprintable, ClassGroup = "Material Designer")
+class UDMMaterialStageBlendFunction : public UDMMaterialStageBlend
 {
 	GENERATED_BODY()
 
 public:
-	virtual void GenerateExpressions(const TSharedRef<FDMMaterialBuildState>& InBuildState) const override;
-	virtual void ConnectOutputToInput(const TSharedRef<FDMMaterialBuildState>& InBuildState, int32 InInputIdx, int32 InExpressionInputIndex,
-		UMaterialExpression* InSourceExpression, int32 InSourceOutputIndex, int32 InSourceOutputChannel) override;
+	UDMMaterialStageBlendFunction();
+
+	//~ Begin UDMMaterialStageThroughput
+	DYNAMICMATERIALEDITOR_API virtual void ConnectOutputToInput(const TSharedRef<FDMMaterialBuildState>& InBuildState, int32 InInputIdx,
+		int32 InExpressionInputIndex, UMaterialExpression* InSourceExpression, int32 InSourceOutputIndex, int32 InSourceOutputChannel) override;
+	//~ End UDMMaterialStageThroughput
+
+	//~ Begin UDMMaterialStageSource
+	DYNAMICMATERIALEDITOR_API virtual void GenerateExpressions(const TSharedRef<FDMMaterialBuildState>& InBuildState) const override;
+	//~ End UDMMaterialStageSource
 
 protected:
 	UPROPERTY()
 	TObjectPtr<UMaterialFunctionInterface> MaterialFunction;
 
-	UDMMaterialStageBlendFunction();
-	UDMMaterialStageBlendFunction(const FText& InName, UMaterialFunctionInterface* InMaterialFunction);
-	UDMMaterialStageBlendFunction(const FText& InName, const FName& InFunctionName, const FString& InFunctionPath);
+	DYNAMICMATERIALEDITOR_API UDMMaterialStageBlendFunction(const FText& InName, UMaterialFunctionInterface* InMaterialFunction);
+
+	/** Loads the function asset. */
+	DYNAMICMATERIALEDITOR_API UDMMaterialStageBlendFunction(const FText& InName, const FName& InFunctionName, const FString& InFunctionPath);
 };
diff --git a/Engine/Plugins/Experimental/DynamicMaterial/Source/DynamicMaterialEditor/Public/Components/DMMaterialStageExpression.h b/Engine/Plugins/Experimental/DynamicMaterial/Source/DynamicMaterialEditor/Public/Components/DMMaterialStageExpression.h
--- a/Engine/Plugins/Experimental/DynamicMaterial/Source/DynamicMaterialEditor/Public/Components/DMMaterialStageExpression.h
+++ b/Engine/Plugins/Experimental/DynamicMaterial/Source/DynamicMaterialEditor/Public/Components/DMMaterialStageExpression.h
@@ -30,25 +30,28 @@ enum class EDMExpressionMenu : uint8
 };
 
 /**
- * A node which directly represents an material expression (or function).
+ * A node which directly represents a material expression (or function).
  */
-UCLASS(Abstract, BlueprintType, Blueprintable, ClassGroup = "Material Designer", meta = (DisplayName = "Material Designer Stage Expression"))
-class DYNAMICMATERIALEDITOR_API UDMMaterialStageExpression : public UDMMaterialStageThroughput
+UCLASS(MinimalAPI, Abstract, BlueprintType, Blueprintable, ClassGroup = "Material Designer", meta = (DisplayName = "Material Designer Stage Expression"))
+class UDMMaterialStageExpression : public UDMMaterialStageThroughput
 {
 	GENERATED_BODY()
 
 public:
 	static TSet<FName> BlockedMaterialExpressionClasses;
 
-	static UDMMaterialStage* CreateStage(TSubclassOf<UDMMaterialStageExpression> InMaterialStageExpressionClass, UDMMaterialLayerObject* InLayer = nullptr);
+	DYNAMICMATERIALEDITOR_API static UDMMaterialStage* CreateStage(TSubclassOf<UDMMaterialStageExpression> InMaterialStageExpressionClass, 
+		UDMMaterialLayerObject* InLayer = nullptr);
 
 	/** Some expressions are either exported or not at random. Use this to find ones that aren't. */
-	static TSubclassOf<UMaterialExpression> FindClass(FString InClassName);
+	DYNAMICMATERIALEDITOR_API static TSubclassOf<UMaterialExpression> FindClass(FString InClassName);
+
+	UDMMaterialStageExpression();
 
 	static const TArray<TStrongObjectPtr<UClass>>& GetAvailableSourceExpressions();
 
 	UFUNCTION(BlueprintCallable, Category = "Material Designer")
-	static UDMMaterialStageExpression* ChangeStageSource_Expression(UDMMaterialStage* InStage,
+	static DYNAMICMATERIALEDITOR_API UDMMaterialStageExpression* ChangeStageSource_Expression(UDMMaterialStage* InStage,
 		TSubclassOf<UDMMaterialStageExpression> InExpressionClass);
 
 	template<typename InExpressionClass>
@@ -63,21 +66,22 @@ public:
 	UFUNCTION(BlueprintPure, Category = "Material Designer")
 	const TArray<EDMExpressionMenu>& GetMenus() const { return Menus; }
 
-	virtual void GenerateExpressions(const TSharedRef<FDMMaterialBuildState>& InBuildState) const override;
+	//~ Begin UDMMaterialStageSource
+	DYNAMICMATERIALEDITOR_API virtual void GenerateExpressions(const TSharedRef<FDMMaterialBuildState>& InBuildState) const override;
+	//~ End UDMMaterialStageSource
 
 protected:
 	static TArray<TStrongObjectPtr<UClass>> SourceExpressions;
 
 	static void GenerateExpressionList();
 
-	UDMMaterialStageExpression();
-	UDMMaterialStageExpression(const FText& InName, TSubclassOf<UMaterialExpression> InClass);
-
 	UPROPERTY(VisibleInstanceOnly, BlueprintReadOnly, Category = "Material Designer")
 	TSubclassOf<UMaterialExpression> MaterialExpressionClass;
 
 	UPROPERTY(VisibleInstanceOnly, BlueprintReadOnly, Category = "Material Designer")
 	TArray<EDMExpressionMenu> Menus;
+
+	DYNAMICMATERIALEDITOR_API UDMMaterialStageExpression(const FText& InName, TSubclassOf<UMaterialExpression> InClass);
 };
 
 
diff --git a/Engine/Plugins/Experimental/DynamicMaterial/Source/DynamicMaterialEditor/Public/Components/DMMaterialStageFunction.h b/Engine/Plugins/Experimental/DynamicMaterial/Source/DynamicMaterialEditor/Public/Components/DMMaterialStageFunction.h
--- a/Engine/Plugins/Experimental/DynamicMaterial/Source/DynamicMaterialEditor/Public/Components/DMMaterialStageFunction.h
+++ b/Engine/Plugins/Experimental/DynamicMaterial/Source/DynamicMaterialEditor/Public/Components/DMMaterialStageFunction.h
@@ -8,21 +8,22 @@
 class UDMMaterialLayerObject;
 class UMaterialFunctionInterface;
 
-UCLASS(BlueprintType, ClassGroup = "Material Designer")
-class DYNAMICMATERIALEDITOR_API UDMMaterialStageFunction : public UDMMaterialStageThroughput
+/** Represents a material function which can be added directly to a stage. */
+UCLASS(MinimalAPI, BlueprintType, ClassGroup = "Material Designer")
+class UDMMaterialStageFunction : public UDMMaterialStageThroughput
 {
 	GENERATED_BODY()
 
 public:
 	static constexpr int32 InputPreviousStage = 0;
 
-	static UDMMaterialStage* CreateStage(UDMMaterialLayerObject* InLayer = nullptr);
+	DYNAMICMATERIALEDITOR_API static UDMMaterialStage* CreateStage(UDMMaterialLayerObject* InLayer = nullptr);
 
 	UFUNCTION(BlueprintCallable, Category = "Material Designer")
-	static UDMMaterialStageFunction* ChangeStageSource_Function(UDMMaterialStage* InStage,
+	static DYNAMICMATERIALEDITOR_API UDMMaterialStageFunction* ChangeStageSource_Function(UDMMaterialStage* InStage,
 		UMaterialFunctionInterface* InMaterialFunction);
 
-	static UMaterialFunctionInterface* GetNoOpFunction();
+	DYNAMICMATERIALEDITOR_API static UMaterialFunctionInterface* GetNoOpFunction();
 
 	UDMMaterialStageFunction();
 
@@ -30,43 +31,43 @@ public:
 	UMaterialFunctionInterface* GetMaterialFunction() const { return MaterialFunction.Get(); }
 
 	UFUNCTION(BlueprintCallable, Category = "Material Designer")
-	void SetMaterialFunction(UMaterialFunctionInterface* InMaterialFunction);
+	DYNAMICMATERIALEDITOR_API void SetMaterialFunction(UMaterialFunctionInterface* InMaterialFunction);
 
 	UFUNCTION(BlueprintPure, Category = "Material Designer")
-	UDMMaterialValue* GetInputValue(int32 InIndex) const;
+	DYNAMICMATERIALEDITOR_API UDMMaterialValue* GetInputValue(int32 InIndex) const;
 
 	UFUNCTION(BlueprintCallable, Category = "Material Designer")
-	TArray<UDMMaterialValue*> GetInputValues() const;
+	DYNAMICMATERIALEDITOR_API TArray<UDMMaterialValue*> GetInputValues() const;
 
 	//~ Begin UDMMaterialStageThroughput
-	virtual void AddDefaultInput(int32 InInputIndex) const override;
-	virtual bool CanChangeInput(int32 InputIndex) const override;
-	virtual bool CanChangeInputType(int32 InputIndex) const override;
-	virtual bool IsInputVisible(int32 InputIndex) const override;
+	DYNAMICMATERIALEDITOR_API virtual void AddDefaultInput(int32 InInputIndex) const override;
+	DYNAMICMATERIALEDITOR_API virtual bool CanChangeInput(int32 InputIndex) const override;
+	DYNAMICMATERIALEDITOR_API virtual bool CanChangeInputType(int32 InputIndex) const override;
+	DYNAMICMATERIALEDITOR_API virtual bool IsInputVisible(int32 InputIndex) const override;
 	//~ End UDMMaterialStageThroughput
 
 	//~ Begin UDMMaterialStageSource
-	virtual void GenerateExpressions(const TSharedRef<FDMMaterialBuildState>& InBuildState) const override;
+	DYNAMICMATERIALEDITOR_API virtual void GenerateExpressions(const TSharedRef<FDMMaterialBuildState>& InBuildState) const override;
 	//~ End UDMMaterialStageSource
 
 	//~ Begin UDMMaterialComponent
-	virtual FText GetComponentDescription() const override;
+	DYNAMICMATERIALEDITOR_API virtual FText GetComponentDescription() const override;
 	//~ End UDMMaterialComponent
 
 	//~ Begin UObject
-	virtual void PreEditChange(FEditPropertyChain& PropertyAboutToChange) override;
-	virtual void PostEditChangeProperty(FPropertyChangedEvent& PropertyChangedEvent) override;
-	virtual void PostLoad() override;
+	DYNAMICMATERIALEDITOR_API virtual void PreEditChange(FEditPropertyChain& PropertyAboutToChange) override;
+	DYNAMICMATERIALEDITOR_API virtual void PostEditChangeProperty(FPropertyChangedEvent& InPropertyChangedEvent) override;
+	DYNAMICMATERIALEDITOR_API virtual void PostLoad() override;
 	//~ End UObject
 
 protected:
 	static TSoftObjectPtr<UMaterialFunctionInterface> NoOp;
 
-	UPROPERTY(EditAnywhere, BlueprintReadWrite, Getter = GetMaterialFunction, Setter = SetMaterialFunction, BlueprintSetter = SetMaterialFunction, Category = "Material Designer",
+	UPROPERTY(EditAnywhere, BlueprintReadWrite, Getter, Setter, BlueprintSetter = SetMaterialFunction, Category = "Material Designer",
 		meta = (DisplayThumbnail = true, AllowPrivateAccess = "true", HighPriority, NotKeyframeable, NoCreate))
 	TObjectPtr<UMaterialFunctionInterface> MaterialFunction;
 
-	UPROPERTY(Transient, DuplicateTransient, TextExportTransient)
+	UPROPERTY(Transient, TextExportTransient)
 	TObjectPtr<UMaterialFunctionInterface> MaterialFunction_PreEdit;
 
 	void OnMaterialFunctionChanged();
@@ -78,7 +79,7 @@ protected:
 	bool NeedsFunctionInit() const;
 
 	//~ Begin UDMMaterialComponent
-	virtual void OnComponentAdded() override;
-	virtual void OnComponentRemoved() override;
+	DYNAMICMATERIALEDITOR_API virtual void OnComponentAdded() override;
+	DYNAMICMATERIALEDITOR_API virtual void OnComponentRemoved() override;
 	//~ End UDMMaterialComponent
 };
diff --git a/Engine/Plugins/Experimental/DynamicMaterial/Source/DynamicMaterialEditor/Public/Components/DMMaterialStageGradient.h b/Engine/Plugins/Experimental/DynamicMaterial/Source/DynamicMaterialEditor/Public/Components/DMMaterialStageGradient.h
--- a/Engine/Plugins/Experimental/DynamicMaterial/Source/DynamicMaterialEditor/Public/Components/DMMaterialStageGradient.h
+++ b/Engine/Plugins/Experimental/DynamicMaterial/Source/DynamicMaterialEditor/Public/Components/DMMaterialStageGradient.h
@@ -14,18 +14,19 @@ struct FDMMaterialBuildState;
 /**
  * A node which represents UV-based gradient.
  */
-UCLASS(Abstract, BlueprintType, Blueprintable, ClassGroup = "Material Designer", meta = (DisplayName = "Material Designer Stage Gradient"))
-class DYNAMICMATERIALEDITOR_API UDMMaterialStageGradient : public UDMMaterialStageThroughput
+UCLASS(MinimalAPI, Abstract, BlueprintType, Blueprintable, ClassGroup = "Material Designer", meta = (DisplayName = "Material Designer Stage Gradient"))
+class UDMMaterialStageGradient : public UDMMaterialStageThroughput
 {
 	GENERATED_BODY()
 
 public:
-	static UDMMaterialStage* CreateStage(TSubclassOf<UDMMaterialStageGradient> InMaterialStageGradientClass, UDMMaterialLayerObject* InLayer = nullptr);
+	UFUNCTION(BlueprintCallable, Category = "Material Designer")
+	static DYNAMICMATERIALEDITOR_API UDMMaterialStage* CreateStage(TSubclassOf<UDMMaterialStageGradient> InMaterialStageGradientClass, UDMMaterialLayerObject* InLayer = nullptr);
 
 	static const TArray<TStrongObjectPtr<UClass>>& GetAvailableGradients();
 
 	UFUNCTION(BlueprintCallable, Category = "Material Designer")
-	static UDMMaterialStageGradient* ChangeStageSource_Gradient(UDMMaterialStage* InStage,
+	static DYNAMICMATERIALEDITOR_API UDMMaterialStageGradient* ChangeStageSource_Gradient(UDMMaterialStage* InStage,
 		TSubclassOf<UDMMaterialStageGradient> InGradientClass);
 
 	template<typename InGradientClass>
@@ -34,8 +35,10 @@ public:
 		return ChangeStageSource_Gradient(InStage, InGradientClass::StaticClass());
 	}
 
+	UDMMaterialStageGradient();
+
 	//~ Begin UDMMaterialStageThroughput
-	virtual bool CanChangeInputType(int32 InputIndex) const override;
+	DYNAMICMATERIALEDITOR_API virtual bool CanChangeInputType(int32 InputIndex) const override;
 	//~ End UDMMaterialStageThroughput
 	
 	//~ Begin UDMMaterialStageSource
@@ -48,7 +51,5 @@ protected:
 
 	static void GenerateGradientList();
 
-	UDMMaterialStageGradient();
-	UDMMaterialStageGradient(const FText& InName);
-
+	DYNAMICMATERIALEDITOR_API UDMMaterialStageGradient(const FText& InName);
 };
\ No newline at end of file
diff --git a/Engine/Plugins/Experimental/DynamicMaterial/Source/DynamicMaterialEditor/Public/Components/DMMaterialStageInput.h b/Engine/Plugins/Experimental/DynamicMaterial/Source/DynamicMaterialEditor/Public/Components/DMMaterialStageInput.h
--- a/Engine/Plugins/Experimental/DynamicMaterial/Source/DynamicMaterialEditor/Public/Components/DMMaterialStageInput.h
+++ b/Engine/Plugins/Experimental/DynamicMaterial/Source/DynamicMaterialEditor/Public/Components/DMMaterialStageInput.h
@@ -12,34 +12,31 @@ struct FDMMaterialStageConnectorChannel;
 /**
  * A node which produces an output (e.g. Texture coordinate.)
  */
-UCLASS(Abstract, BlueprintType, ClassGroup = "Material Designer", meta = (DisplayName = "Material Designer Stage Input"))
-class DYNAMICMATERIALEDITOR_API UDMMaterialStageInput : public UDMMaterialStageSource
+UCLASS(MinimalAPI, Abstract, BlueprintType, ClassGroup = "Material Designer", meta = (DisplayName = "Material Designer Stage Input"))
+class UDMMaterialStageInput : public UDMMaterialStageSource
 {
 	GENERATED_BODY()
 
 public:
-	static const FString StageInputPrefixStr;
+	DYNAMICMATERIALEDITOR_API static const FString StageInputPrefixStr;
 
 	virtual FText GetChannelDescription(const FDMMaterialStageConnectorChannel& Channel)
 		PURE_VIRTUAL(UDMMaterialStageInput::GetChannelDescription, return FText::GetEmpty();)
 		
-	virtual void Update(EDMUpdateType InUpdateType) override;
-
 	//~ Begin UDMMaterialComponent
-	virtual FString GetComponentPathComponent() const override;
+	DYNAMICMATERIALEDITOR_API virtual FString GetComponentPathComponent() const override;
+	DYNAMICMATERIALEDITOR_API virtual void Update(EDMUpdateType InUpdateType) override;
 	//~ End UDMMaterialComponent
 
 protected:
-	UDMMaterialStageInput()
-	{
-	}
+	/** Updates the output connectors based on the type of input. */
+	virtual void UpdateOutputConnectors() {}
 
 	//~ Begin UDMMaterialStageSource
-	virtual void UpdateOutputConnectors() { }
-	virtual void UpdatePreviewMaterial(UMaterial* InPreviewMaterial = nullptr) override;
+	DYNAMICMATERIALEDITOR_API virtual void GeneratePreviewMaterial(UMaterial* InPreviewMaterial) override;
 	//~ End UDMMaterialStageSource
 
 	//~ Begin UDMMaterialComponent
-	virtual void GetComponentPathInternal(TArray<FString>& OutChildComponentPathComponents) const override;
+	DYNAMICMATERIALEDITOR_API virtual void GetComponentPathInternal(TArray<FString>& OutChildComponentPathComponents) const override;
 	//~ End UDMMaterialComponent
 };
diff --git a/Engine/Plugins/Experimental/DynamicMaterial/Source/DynamicMaterialEditor/Public/Components/DMMaterialStageSource.h b/Engine/Plugins/Experimental/DynamicMaterial/Source/DynamicMaterialEditor/Public/Components/DMMaterialStageSource.h
--- a/Engine/Plugins/Experimental/DynamicMaterial/Source/DynamicMaterialEditor/Public/Components/DMMaterialStageSource.h
+++ b/Engine/Plugins/Experimental/DynamicMaterial/Source/DynamicMaterialEditor/Public/Components/DMMaterialStageSource.h
@@ -15,8 +15,8 @@ struct FDMMaterialBuildState;
 /**
  * A node which produces an output.
  */
-UCLASS(Abstract, BlueprintType, ClassGroup = "Material Designer", meta = (DisplayName = "Material Designer Stage Source"))
-class DYNAMICMATERIALEDITOR_API UDMMaterialStageSource : public UDMMaterialComponent
+UCLASS(MinimalAPI, Abstract, BlueprintType, ClassGroup = "Material Designer", meta = (DisplayName = "Material Designer Stage Source"))
+class UDMMaterialStageSource : public UDMMaterialComponent
 {
 	GENERATED_BODY()
 
@@ -26,7 +26,7 @@ public:
 	static const TArray<TStrongObjectPtr<UClass>>& GetAvailableSourceClasses();
 
 	UFUNCTION(BlueprintPure, Category = "Material Designer")
-	UDMMaterialStage* GetStage() const;
+	DYNAMICMATERIALEDITOR_API UDMMaterialStage* GetStage() const;
 
 	/* Returns a description of the stage for which this is the source. */
 	UFUNCTION(BlueprintPure, Category = "Material Designer")
@@ -35,24 +35,28 @@ public:
 	UFUNCTION(BlueprintPure, Category = "Material Designer")
 	const TArray<FDMMaterialStageConnector>& GetOutputConnectors() const { return OutputConnectors; }
 
-	virtual void GetMaskAlphaBlendNode(const TSharedRef<FDMMaterialBuildState>& InBuildState, UMaterialExpression*& OutExpression, 
-		int32& OutOutputIndex, int32& OutOutputChannel) const;
+	/** For pre-multiplied alpha, this retrieves the node which should be multiplied against. */
+	DYNAMICMATERIALEDITOR_API virtual void GetMaskAlphaBlendNode(const TSharedRef<FDMMaterialBuildState>& InBuildState, 
+		UMaterialExpression*& OutExpression, int32& OutOutputIndex, int32& OutOutputChannel) const;
 
 	virtual void GenerateExpressions(const TSharedRef<FDMMaterialBuildState>& InBuildState) const
 		PURE_VIRTUAL(UDMMaterialStageSource::GenerateExpressions)
 
+	/**
+	 * Adds non-expression based node properties, such as Clamp Texture.
+	 * Is called after expressions are generated.
+	 */
 	virtual void AddExpressionProperties(const TArray<UMaterialExpression*>& Expressions) const {}
 
-	UFUNCTION(BlueprintPure, Category = "Material Designer")
-	UMaterial* GetPreviewMaterial();
-
-	void UpdateOntoPreviewMaterial(UMaterial* ExternalPreviewMaterial);
+	/** Generates a material representing just this node. */
+	UFUNCTION(BlueprintCallable, Category = "Material Designer")
+	DYNAMICMATERIALEDITOR_API virtual void GeneratePreviewMaterial(UMaterial* InPreviewMaterial);
 
 	/** 
 	 * Returns the output index (channel WHOLE_CHANNEL) if this expression has pre-masked outputs.
 	 * Returns INDEX_NONE if it is not supported.
 	 */
-	virtual int32 GetInnateMaskOutput(int32 OutputIndex, int32 OutputChannels) const;
+	DYNAMICMATERIALEDITOR_API virtual int32 GetInnateMaskOutput(int32 OutputIndex, int32 OutputChannels) const;
 
 	/**
 	 * Given an output index, may return an override for output channels on that output.
@@ -61,22 +65,21 @@ public:
 	 */
 	virtual int32 GetOutputChannelOverride(int32 InOutputIndex) const { return INDEX_NONE; }
 
-	virtual bool UpdateStagePreviewMaterial(UDMMaterialStage* InStage, UMaterial* InPreviewMaterial, UMaterialExpression*& OutMaterialExpression, 
-		int32& OutputIndex);
+	/** Generates a preview material based on this source, as owned by the given stage. */
+	DYNAMICMATERIALEDITOR_API virtual bool GenerateStagePreviewMaterial(UDMMaterialStage* InStage, UMaterial* InPreviewMaterial, 
+		UMaterialExpression*& OutMaterialExpression, int32& OutputIndex);
 
 	//~ Begin FNotifyHook
-	virtual void NotifyPostChange(const FPropertyChangedEvent& InPropertyChangedEvent, class FEditPropertyChain* InPropertyThatChanged) override;
+	DYNAMICMATERIALEDITOR_API virtual void NotifyPostChange(const FPropertyChangedEvent& InPropertyChangedEvent, class FEditPropertyChain* InPropertyThatChanged) override;
 	//~ End FNotifyHook
 
 	//~ Begin UDMMaterialComponent
-	virtual void Update(EDMUpdateType InUpdateType) override;
-	virtual void DoClean() override;
-	virtual UDMMaterialComponent* GetParentComponent() const override;
-	virtual void PostEditorDuplicate(UDynamicMaterialModel* InMaterialModel, UDMMaterialComponent* InParent) override;
+	DYNAMICMATERIALEDITOR_API virtual void Update(EDMUpdateType InUpdateType) override;
+	DYNAMICMATERIALEDITOR_API virtual UDMMaterialComponent* GetParentComponent() const override;
 	//~ End UDMMaterialComponent
 
 	//~ Begin UObject
-	virtual void PostEditUndo() override;
+	DYNAMICMATERIALEDITOR_API virtual void PostEditUndo() override;
 	//~ End UObject
 
 protected:
@@ -85,17 +88,9 @@ protected:
 	static void GenerateClassList();
 
 	UPROPERTY(VisibleInstanceOnly, BlueprintReadOnly, Category = "Material Designer")
-	TArray<FDMMaterialStageConnector> OutputConnectors;
-
-	UPROPERTY(VisibleInstanceOnly, BlueprintReadOnly, Transient, DuplicateTransient, TextExportTransient, Category = "Material Designer")
-	TObjectPtr<UMaterial> PreviewMaterial;
-
-	UDMMaterialStageSource();
-
-	void CreatePreviewMaterial();
-	virtual void UpdatePreviewMaterial(UMaterial* InPreviewMaterial = nullptr);
+	TArray<FDMMaterialStageConnector> OutputConnectors;	
 
 	//~ Begin UDMMaterialComponent
-	virtual void OnComponentAdded() override;
+	DYNAMICMATERIALEDITOR_API virtual void OnComponentAdded() override;
 	//~ End UDMMaterialComponent
 };
