Material Designer

diff --git a/Engine/Plugins/Experimental/DynamicMaterial/Source/DynamicMaterialEditor/Public/Components/DMMaterialStageThroughput.h b/Engine/Plugins/Experimental/DynamicMaterial/Source/DynamicMaterialEditor/Public/Components/DMMaterialStageThroughput.h
--- a/Engine/Plugins/Experimental/DynamicMaterial/Source/DynamicMaterialEditor/Public/Components/DMMaterialStageThroughput.h
+++ b/Engine/Plugins/Experimental/DynamicMaterial/Source/DynamicMaterialEditor/Public/Components/DMMaterialStageThroughput.h
@@ -18,7 +18,7 @@ struct FDMExpressionInput
 	int32 OutputIndex = INDEX_NONE;
 	int32 OutputChannel = INDEX_NONE;
 
-	bool IsValid()
+	bool IsValid() const
 	{
 		return OutputExpressions.IsEmpty() == false && OutputIndex != INDEX_NONE && OutputChannel != INDEX_NONE;
 	}
@@ -27,22 +27,24 @@ struct FDMExpressionInput
 /**
  * A node which take one or more inputs and produces an output (e.g. Multiply)
  */
-UCLASS(Abstract, BlueprintType, ClassGroup = "Material Designer", meta = (DisplayName = "Material Designer Stage Throughput"))
-class DYNAMICMATERIALEDITOR_API UDMMaterialStageThroughput : public UDMMaterialStageSource
+UCLASS(MinimalAPI, Abstract, BlueprintType, ClassGroup = "Material Designer", meta = (DisplayName = "Material Designer Stage Throughput"))
+class UDMMaterialStageThroughput : public UDMMaterialStageSource
 {
 	GENERATED_BODY()
 
 public:
 	static const TArray<TStrongObjectPtr<UClass>>& GetAvailableThroughputs();
 
-	virtual FText GetComponentDescription() const override { return GetDescription(); }
+	DYNAMICMATERIALEDITOR_API UDMMaterialStageThroughput();
 
 	UFUNCTION(BlueprintPure, Category = "Material Designer")
 	const FText& GetDescription() const { return Name; }
 
+	/** Returns true if input is required to successfully compile this node. */
 	UFUNCTION(BlueprintPure, Category = "Material Designer")
 	bool IsInputRequired() const { return bInputRequired; }
 
+	/** Returns true if this node's inputs can have their own inputs. */
 	UFUNCTION(BlueprintPure, Category = "Material Designer")
 	bool AllowsNestedInputs() const { return bAllowNestedInputs; }
 
@@ -50,61 +52,86 @@ public:
 	const TArray<FDMMaterialStageConnector>& GetInputConnectors() const { return InputConnectors; }
 
 	UFUNCTION(BlueprintCallable, Category = "Material Designer")
-	virtual bool CanInputAcceptType(int32 InThroughputInputIndex, EDMValueType InValueType) const;
+	DYNAMICMATERIALEDITOR_API virtual bool CanInputAcceptType(int32 InThroughputInputIndex, EDMValueType InValueType) const;
 
+	/**
+	 * Whether the given output connector can connect to this node.
+	 * @param bInCheckSingleFloat If the initial compatibility check fails, it will again check against a single float.
+	 */
 	UFUNCTION(BlueprintCallable, Category = "Material Designer")
-	virtual bool CanInputConnectTo(int32 InThroughputInputIndex, const FDMMaterialStageConnector& InOutputConnector, int32 InOutputChannel,
-		bool bInCheckSingleFloat = false);
+	DYNAMICMATERIALEDITOR_API virtual bool CanInputConnectTo(int32 InThroughputInputIndex, const FDMMaterialStageConnector& InOutputConnector, 
+		int32 InOutputChannel, bool bInCheckSingleFloat = false);
+
+	/** Whether the value of the given input change. */
+	DYNAMICMATERIALEDITOR_API virtual bool CanChangeInput(int32 InThroughputInputIndex) const;
 
-	virtual bool CanChangeInput(int32 InThroughputInputIndex) const;
-	virtual bool CanChangeInputType(int32 InThroughputInputIndex) const;
+	/** Whether you can change the type of the given input. */
+	DYNAMICMATERIALEDITOR_API virtual bool CanChangeInputType(int32 InThroughputInputIndex) const;
 
-	virtual bool IsInputVisible(int32 InThroughputInputIndex) const;
+	/** Whether this input will show up in the Material Designer editor. */
+	DYNAMICMATERIALEDITOR_API virtual bool IsInputVisible(int32 InThroughputInputIndex) const;
 
-	virtual void ConnectOutputToInput(const TSharedRef<FDMMaterialBuildState>& InBuildState, int32 InThroughputInputIndex, int32 InExpressionInputIndex,
-		UMaterialExpression* InSourceExpression, int32 InSourceOutputIndex, int32 InSourceOutputChannel);
+	/**
+	 * Connect the output of a node to the given input of this node.
+	 * @param InExpressionInputIndex The input index of this node.
+	 * @param InSourceExpression The node to take the input from.
+	 * @param InSourceOutputIndex The output index of the source expression.
+	 * @param InSourceOutputChannel The channel of the output (RGBA).
+	 */
+	DYNAMICMATERIALEDITOR_API virtual void ConnectOutputToInput(const TSharedRef<FDMMaterialBuildState>& InBuildState, int32 InThroughputInputIndex, 
+		int32 InExpressionInputIndex, UMaterialExpression* InSourceExpression, int32 InSourceOutputIndex, int32 InSourceOutputChannel);
 
 	/** Returns true if the layer and mask can have their Texture UV linked. */
 	virtual bool SupportsLayerMaskTextureUVLink() const { return false; }
 
 	/** Returns the input index for the default implementation of the below method. */
-	virtual int32 GetLayerMaskTextureUVLinkInputIndex() const;
+	DYNAMICMATERIALEDITOR_API virtual int32 GetLayerMaskTextureUVLinkInputIndex() const;
 
 	/** 
 	 * Returns all the material nodes requires to create this node's Texture UV input. 
 	 * If you override this method, you do not need to override GetLayerMaskTextureUVLinkInputIndex.
 	 */
-	virtual FDMExpressionInput GetLayerMaskLinkTextureUVInputExpressions(const TSharedRef<FDMMaterialBuildState>& InBuildState) const;
+	DYNAMICMATERIALEDITOR_API virtual FDMExpressionInput GetLayerMaskLinkTextureUVInputExpressions(const TSharedRef<FDMMaterialBuildState>& InBuildState) const;
 
 	/**
 	 * Override this to redirect inputs to other nodes.
 	 * Returns the first node in the array by default
 	 * --> In [ ]-[ ]-[ ] Out -->
 	 */
-	virtual UMaterialExpression* GetExpressionForInput(const TArray<UMaterialExpression*>& InStageSourceExpressions, int32 InThroughputInputIndex, 
-		int32 InExpressionInputIndex);
+	DYNAMICMATERIALEDITOR_API virtual UMaterialExpression* GetExpressionForInput(const TArray<UMaterialExpression*>& InStageSourceExpressions, 
+		int32 InThroughputInputIndex, int32 InExpressionInputIndex);
 
-	virtual void AddDefaultInput(int32 InInputIndex) const;
+	/** When the node is instantiated, this method adds default input values based on type. */
+	DYNAMICMATERIALEDITOR_API virtual void AddDefaultInput(int32 InInputIndex) const;
 
-	/** Returns the actual output index of the material expression */
-	virtual int32 ResolveInput(const TSharedRef<FDMMaterialBuildState>& InBuildState, int32 InThroughputInputIndex, FDMMaterialStageConnectorChannel& OutChannel,
-		TArray<UMaterialExpression*>& OutExpressions) const;
-
-	virtual int32 ResolveLayerMaskTextureUVLinkInput(const TSharedRef<FDMMaterialBuildState>& InBuildState, int32 InThroughputInputIndex, 
+	/**
+	 * Generates (or retrieves) expressions that produce this input for the node.
+	 * @return the actual output index of the material expression
+	 */
+	DYNAMICMATERIALEDITOR_API virtual int32 ResolveInput(const TSharedRef<FDMMaterialBuildState>& InBuildState, int32 InThroughputInputIndex, 
 		FDMMaterialStageConnectorChannel& OutChannel, TArray<UMaterialExpression*>& OutExpressions) const;
 
-	virtual void InputUpdated(int32 InThroughputInputIndex, EDMUpdateType InUpdateType) { }
+	/** If this is on a Mask stage and it is the UV input index, this method is used to retrieve the base stage's UV input. */
+	DYNAMICMATERIALEDITOR_API virtual int32 ResolveLayerMaskTextureUVLinkInput(const TSharedRef<FDMMaterialBuildState>& InBuildState, 
+		int32 InThroughputInputIndex, FDMMaterialStageConnectorChannel& OutChannel, TArray<UMaterialExpression*>& OutExpressions) const;
+
+	/** This is called when the input value of this node's stage is updated. */
+	virtual void InputUpdated(int32 InThroughputInputIndex, EDMUpdateType InUpdateType) {}
+
+	//~ Begin UDMMaterialComponent
+	virtual FText GetComponentDescription() const override { return GetDescription(); }
+	//~ End UDMMaterialComponent
 
 protected:
 	static TArray<TStrongObjectPtr<UClass>> Throughputs;
 
 	static void GenerateThroughputList();
 
-	static int32 ResolveLayerMaskTextureUVLinkInputImpl(const TSharedRef<FDMMaterialBuildState>& InBuildState, const UDMMaterialStageSource* InStageSource,
-		FDMMaterialStageConnectorChannel& OutChannel, TArray<UMaterialExpression*>& OutExpressions);
+	/** @See ResolveLayerMaskTextureUVLinkInput */
+	DYNAMICMATERIALEDITOR_API static int32 ResolveLayerMaskTextureUVLinkInputImpl(const TSharedRef<FDMMaterialBuildState>& InBuildState, 
+		const UDMMaterialStageSource* InStageSource, FDMMaterialStageConnectorChannel& OutChannel, TArray<UMaterialExpression*>& OutExpressions);
 
-	UDMMaterialStageThroughput();
-	UDMMaterialStageThroughput(const FText& InName);
+	DYNAMICMATERIALEDITOR_API UDMMaterialStageThroughput(const FText& InName);
 
 	UPROPERTY(VisibleInstanceOnly, BlueprintReadOnly, Category = "Material Designer")
 	FText Name;
@@ -118,17 +145,21 @@ protected:
 	UPROPERTY(VisibleInstanceOnly, BlueprintReadOnly, Category = "Material Designer")
 	TArray<FDMMaterialStageConnector> InputConnectors;
 
-	virtual bool ShouldKeepInput(int32 InThroughputInputIndex);
+	/** When the stage's source is changed, whether the given of input from the previous source should be kept. */
+	DYNAMICMATERIALEDITOR_API virtual bool ShouldKeepInput(int32 InThroughputInputIndex);
 
-	void ConnectOutputToInput_Internal(const TSharedRef<FDMMaterialBuildState>& InBuildState, UMaterialExpression* InTargetExpression,
+	/** @See ConnectOutputToInput */
+	DYNAMICMATERIALEDITOR_API void ConnectOutputToInput_Internal(const TSharedRef<FDMMaterialBuildState>& InBuildState, UMaterialExpression* InTargetExpression,
 		int32 InExpressionInputIndex, UMaterialExpression* InSourceExpression, int32 InSourceOutputIndex, int32 InSourceOutputChannel) const;
 
-	virtual int32 ResolveInputChannel(const TSharedRef<FDMMaterialBuildState>& InBuildState, int32 InThroughputInputIndex, int32 InChannelIndex, 
-		FDMMaterialStageConnectorChannel& OutChannel, TArray<UMaterialExpression*>& OutExpressions) const;
+	/** Finds the input for the individual channel. @see ResolveInput */
+	DYNAMICMATERIALEDITOR_API virtual int32 ResolveInputChannel(const TSharedRef<FDMMaterialBuildState>& InBuildState, int32 InThroughputInputIndex,
+		int32 InChannelIndex, FDMMaterialStageConnectorChannel& OutChannel, TArray<UMaterialExpression*>& OutExpressions) const;
 
-	virtual void UpdatePreviewMaterial(UMaterial* InPreviewMaterial = nullptr);
+	/** Generates a material based on the output of just this node. */
+	DYNAMICMATERIALEDITOR_API virtual void GeneratePreviewMaterial(UMaterial* InPreviewMaterial);
 
 	//~ Begin UDMMaterialComponent
-	virtual void OnComponentAdded() override;
+	DYNAMICMATERIALEDITOR_API virtual void OnComponentAdded() override;
 	//~ End UDMMaterialComponent
 };
diff --git a/Engine/Plugins/Experimental/DynamicMaterial/Source/DynamicMaterialEditor/Public/Components/DMMaterialStageThroughputLayerBlend.h b/Engine/Plugins/Experimental/DynamicMaterial/Source/DynamicMaterialEditor/Public/Components/DMMaterialStageThroughputLayerBlend.h
--- a/Engine/Plugins/Experimental/DynamicMaterial/Source/DynamicMaterialEditor/Public/Components/DMMaterialStageThroughputLayerBlend.h
+++ b/Engine/Plugins/Experimental/DynamicMaterial/Source/DynamicMaterialEditor/Public/Components/DMMaterialStageThroughputLayerBlend.h
@@ -10,8 +10,9 @@ class UMaterial;
 class UMaterialExpression;
 enum class EAvaColorChannel : uint8;
 
-UCLASS(BlueprintType, ClassGroup = "Material Designer")
-class DYNAMICMATERIALEDITOR_API UDMMaterialStageThroughputLayerBlend : public UDMMaterialStageThroughput
+/** Used as the source for mask stages. */
+UCLASS(MinimalAPI, BlueprintType, ClassGroup = "Material Designer")
+class UDMMaterialStageThroughputLayerBlend : public UDMMaterialStageThroughput
 {
 	GENERATED_BODY()
 
@@ -20,57 +21,65 @@ public:
 	static constexpr int32 InputBaseStage = 1;
 	static constexpr int32 InputMaskSource = 2;
 
-	static UDMMaterialStage* CreateStage(UDMMaterialLayerObject* InLayer = nullptr);
+	UFUNCTION(BlueprintCallable, Category = "Material Designer")
+	static DYNAMICMATERIALEDITOR_API UDMMaterialStage* CreateStage(UDMMaterialLayerObject* InLayer = nullptr);
 
-	UDMMaterialStageThroughputLayerBlend();
+	DYNAMICMATERIALEDITOR_API UDMMaterialStageThroughputLayerBlend();
 
-	UDMMaterialStageInput* GetInputMask() const;
+	/** Returns the input connected to the Mask input. */
+	UFUNCTION(BlueprintCallable, Category = "Material Designer")
+	DYNAMICMATERIALEDITOR_API UDMMaterialStageInput* GetInputMask() const;
 
+	/** Filters the output of the mask input node with the given channel. */
 	UFUNCTION(BlueprintCallable, Category = "Material Designer")
-	EAvaColorChannel GetMaskChannelOverride() const;
+	DYNAMICMATERIALEDITOR_API EAvaColorChannel GetMaskChannelOverride() const;
 
+	/** Filters the output of the mask input node with the given channel. */
 	UFUNCTION(BlueprintCallable, Category="Material Designer")
-	void SetMaskChannelOverride(EAvaColorChannel InMaskChannel);
+	DYNAMICMATERIALEDITOR_API void SetMaskChannelOverride(EAvaColorChannel InMaskChannel);
 
 	//~ Begin UDMMaterialComponent
-	virtual void OnComponentAdded() override;
-	virtual void Update(EDMUpdateType InUpdateType) override;
-	virtual void GenerateExpressions(const TSharedRef<FDMMaterialBuildState>& InBuildState) const override;
-	virtual void PostEditorDuplicate(UDynamicMaterialModel* InMaterialModel, UDMMaterialComponent* InParent) override;
-	virtual bool IsPropertyVisible(FName InProperty) const override;
+	DYNAMICMATERIALEDITOR_API virtual void OnComponentAdded() override;
+	DYNAMICMATERIALEDITOR_API virtual void Update(EDMUpdateType InUpdateType) override;
+	DYNAMICMATERIALEDITOR_API virtual void GenerateExpressions(const TSharedRef<FDMMaterialBuildState>& InBuildState) const override;
+	DYNAMICMATERIALEDITOR_API virtual void PostEditorDuplicate(UDynamicMaterialModel* InMaterialModel, UDMMaterialComponent* InParent) override;
+	DYNAMICMATERIALEDITOR_API virtual bool IsPropertyVisible(FName InProperty) const override;
 	//~ End UDMMaterialComponent
 
 	//~ Begin UDMMaterialStageSource
-	virtual FText GetStageDescription() const override;
-	virtual bool UpdateStagePreviewMaterial(UDMMaterialStage* InStage, UMaterial* InPreviewMaterial, UMaterialExpression*& OutMaterialExpression,
-		int32& OutputIndex) override;
+	DYNAMICMATERIALEDITOR_API virtual FText GetStageDescription() const override;
+	DYNAMICMATERIALEDITOR_API virtual bool GenerateStagePreviewMaterial(UDMMaterialStage* InStage, UMaterial* InPreviewMaterial, 
+		UMaterialExpression*& OutMaterialExpression, int32& OutputIndex) override;
 	//~ End UDMMaterialStageSource
 
 	//~ Begin FNotifyHook
-	virtual void NotifyPostChange(const FPropertyChangedEvent& InPropertyChangedEvent, class FEditPropertyChain* InPropertyThatChanged) override;
+	DYNAMICMATERIALEDITOR_API virtual void NotifyPostChange(const FPropertyChangedEvent& InPropertyChangedEvent, class FEditPropertyChain* InPropertyThatChanged) override;
 	//~ End FNotifyHook
 
 	//~ Begin UDMMaterialStageThroughput
-	virtual void AddDefaultInput(int32 InInputIndex) const override;
-	virtual bool CanChangeInput(int32 InputIndex) const override;
-	virtual bool IsInputVisible(int32 InputIndex) const override;
-	virtual int32 ResolveInput(const TSharedRef<FDMMaterialBuildState>& InBuildState, int32 InputIndex, FDMMaterialStageConnectorChannel& OutChannel,
-		TArray<UMaterialExpression*>& OutExpressions) const override;
-	virtual void ConnectOutputToInput(const TSharedRef<FDMMaterialBuildState>& InBuildState, int32 InInputIdx, int32 InExpressionInputIndex,
-		UMaterialExpression* InSourceExpression, int32 InSourceOutputIndex, int32 InSourceOutputChannel) override;
+	DYNAMICMATERIALEDITOR_API virtual void AddDefaultInput(int32 InInputIndex) const override;
+	DYNAMICMATERIALEDITOR_API virtual bool CanChangeInput(int32 InputIndex) const override;
+	DYNAMICMATERIALEDITOR_API virtual bool IsInputVisible(int32 InputIndex) const override;
+	DYNAMICMATERIALEDITOR_API virtual int32 ResolveInput(const TSharedRef<FDMMaterialBuildState>& InBuildState, int32 InputIndex, 
+		FDMMaterialStageConnectorChannel& OutChannel, TArray<UMaterialExpression*>& OutExpressions) const override;
+	DYNAMICMATERIALEDITOR_API virtual void ConnectOutputToInput(const TSharedRef<FDMMaterialBuildState>& InBuildState, int32 InInputIdx, 
+		int32 InExpressionInputIndex, UMaterialExpression* InSourceExpression, int32 InSourceOutputIndex, int32 InSourceOutputChannel) override;
 	//~ End UDMMaterialStageThroughput
 
-	void GetMaskOutput(const TSharedRef<FDMMaterialBuildState>& InBuildState, UMaterialExpression*& OutExpression, int32& OutOutputIndex, int32& OutOutputChannel) const;
+	/** Resolves the Mask input. */
+	DYNAMICMATERIALEDITOR_API void GetMaskOutput(const TSharedRef<FDMMaterialBuildState>& InBuildState, UMaterialExpression*& OutExpression, int32& OutOutputIndex, int32& OutOutputChannel) const;
 
+	/* When true, the base stage's output will be multiplied by this stage (darkening it where it is translucent). */
 	UFUNCTION(BlueprintPure, Category = "Material Designer")
 	bool UsePremultiplyAlpha() const { return bPremultiplyAlpha; }
 
+	/* When true, the base stage's output will be multiplied by this stage (darkening it where it is translucent). */
 	UFUNCTION(BlueprintCallable, Category = "Material Designer")
-	void SetPremultiplyAlpha(bool bInValue);
+	DYNAMICMATERIALEDITOR_API void SetPremultiplyAlpha(bool bInValue);
 
 	//~ Begin UObject
-	virtual void PostLoad() override;
-	virtual void PostEditImport() override;
+	DYNAMICMATERIALEDITOR_API virtual void PostLoad() override;
+	DYNAMICMATERIALEDITOR_API virtual void PostEditImport() override;
 	//~ End UObject
 
 protected:
@@ -87,36 +96,45 @@ protected:
 		meta=(NotKeyframeable))
 	bool bIsAlphaOnlyBlend;
 
+	/** When doing an update, this will prevent recursive calls. */
 	bool bBlockUpdate;
 
-	virtual int32 ResolveMaskInput(const TSharedRef<FDMMaterialBuildState>& InBuildState, int32 InputIndex, FDMMaterialStageConnectorChannel& OutChannel,
-		TArray<UMaterialExpression*>& OutExpressions) const;
-	
-	virtual void UpdatePreviewMaterial(UMaterial* InPreviewMaterial = nullptr) override;
+	DYNAMICMATERIALEDITOR_API virtual int32 ResolveMaskInput(const TSharedRef<FDMMaterialBuildState>& InBuildState, int32 InputIndex, 
+		FDMMaterialStageConnectorChannel& OutChannel, TArray<UMaterialExpression*>& OutExpressions) const;
+
+	/** Updates the value of bIsAlphaOnlyBlend based on current conditions. */
+	DYNAMICMATERIALEDITOR_API virtual void UpdateAlphaOnlyMaskStatus();
+
+	/** Called when the owning stage is updated. */
+	DYNAMICMATERIALEDITOR_API virtual void OnStageUpdated(UDMMaterialComponent* InComponent, EDMUpdateType InUpdateType);
 
-	virtual void UpdateAlphaOnlyMaskStatus();
-	virtual void OnStageUpdated(UDMMaterialComponent* InComponent, EDMUpdateType InUpdateType);
-	virtual void UpdateAlphaOnlyMasks(EDMUpdateType InUpdateType);
+	/** Update this and all parent layer mask stages that are only-only (no based). */
+	DYNAMICMATERIALEDITOR_API virtual void UpdateAlphaOnlyMasks(EDMUpdateType InUpdateType);
 
-	void InitBlendStage();
+	DYNAMICMATERIALEDITOR_API void InitBlendStage();
 
-	void GenerateMainExpressions(const TSharedRef<FDMMaterialBuildState>& InBuildState) const;
-	void GeneratePreviewExpressions(const TSharedRef<FDMMaterialBuildState>& InBuildState) const;
+	DYNAMICMATERIALEDITOR_API void GenerateMainExpressions(const TSharedRef<FDMMaterialBuildState>& InBuildState) const;
+	DYNAMICMATERIALEDITOR_API void GeneratePreviewExpressions(const TSharedRef<FDMMaterialBuildState>& InBuildState) const;
 
-	void UpdateLinkedInputStage(EDMUpdateType InUpdateType);
+	/** If this shades any inputs with the base stage, update the base stage. */
+	DYNAMICMATERIALEDITOR_API void UpdateLinkedInputStage(EDMUpdateType InUpdateType);
 
 	/* Returns true if there are any outputs on the mask input that have more than 1 channel. */
-	bool CanUseMaskChannelOverride() const;
+	DYNAMICMATERIALEDITOR_API bool CanUseMaskChannelOverride() const;
 
 	/* Returns the first output on the mask input that has more than 1 channel. */
-	int32 GetDefaultMaskChannelOverrideOutputIndex() const;
+	DYNAMICMATERIALEDITOR_API int32 GetDefaultMaskChannelOverrideOutputIndex() const;
 
 	/* Returns true if the given mask output supports more than 1 channel. */
-	bool IsValidMaskChannelOverrideOutputIndex(int32 InIndex) const;
+	DYNAMICMATERIALEDITOR_API bool IsValidMaskChannelOverrideOutputIndex(int32 InIndex) const;
 
 	/* Reads the current output setting from the input map. */
 	void PullMaskChannelOverride() const;
 
 	/* Takes the override setting and applies it to the input map. */
 	void PushMaskChannelOverride();
+
+	//~ Begin UDMMaterialStageThroughput
+	DYNAMICMATERIALEDITOR_API virtual void GeneratePreviewMaterial(UMaterial* InPreviewMaterial) override;
+	//~ End UDMMaterialStageThroughput
 };
diff --git a/Engine/Plugins/Experimental/DynamicMaterial/Source/DynamicMaterialEditor/Public/Components/DMMaterialSubStage.h b/Engine/Plugins/Experimental/DynamicMaterial/Source/DynamicMaterialEditor/Public/Components/DMMaterialSubStage.h
--- a/Engine/Plugins/Experimental/DynamicMaterial/Source/DynamicMaterialEditor/Public/Components/DMMaterialSubStage.h
+++ b/Engine/Plugins/Experimental/DynamicMaterial/Source/DynamicMaterialEditor/Public/Components/DMMaterialSubStage.h
@@ -6,10 +6,10 @@
 #include "DMMaterialSubStage.generated.h"
 
 /**
- * A stage that is a subobject of another stage.
+ * A stage that is a subobject of another stage, such as when an input throughput has its own inputs.
  */
-UCLASS(BlueprintType, ClassGroup = "Material Designer", meta = (DisplayName = "Material Designer Sub Stage"))
-class DYNAMICMATERIALEDITOR_API UDMMaterialSubStage : public UDMMaterialStage
+UCLASS(MinimalAPI, BlueprintType, ClassGroup = "Material Designer", meta = (DisplayName = "Material Designer Sub Stage"))
+class UDMMaterialSubStage : public UDMMaterialStage
 {
 	GENERATED_BODY()
 
@@ -17,25 +17,31 @@ class DYNAMICMATERIALEDITOR_API UDMMaterialSubStage : public UDMMaterialStage
 
 public:
 	UFUNCTION(BlueprintCallable, Category = "Material Designer")
-	static UDMMaterialSubStage* CreateMaterialSubStage(UDMMaterialStage* InParentStage);
+	static DYNAMICMATERIALEDITOR_API UDMMaterialSubStage* CreateMaterialSubStage(UDMMaterialStage* InParentStage);
 
-	UDMMaterialStage* GetParentStage() const;
-	UDMMaterialStage* GetParentMostStage() const;
+	/** Returns the parent stage of this stage, which is probably not its direct parent. */
+	UFUNCTION(BlueprintPure, Category = "Material Designer")
+	DYNAMICMATERIALEDITOR_API UDMMaterialStage* GetParentStage() const;
+
+	/** Recursively calls GetParentStage() to find the outer stage. */
+	UFUNCTION(BlueprintPure, Category = "Material Designer")
+	DYNAMICMATERIALEDITOR_API UDMMaterialStage* GetParentMostStage() const;
+
+	/** Sets which object directly owns this component in the hierarchy. */
+	UFUNCTION(BlueprintCallable, Category = "Material Designer")
+	DYNAMICMATERIALEDITOR_API void SetParentComponent(UDMMaterialComponent* InParentComponent);
 
 	//~ Begin UDMMaterialStage
-	virtual bool IsCompatibleWithPreviousStage(const UDMMaterialStage* PreviousStage) const override;
-	virtual bool IsCompatibleWithNextStage(const UDMMaterialStage* NextStage) const override;
-	virtual bool IsRootStage() const override;
+	DYNAMICMATERIALEDITOR_API virtual bool IsCompatibleWithPreviousStage(const UDMMaterialStage* PreviousStage) const override;
+	DYNAMICMATERIALEDITOR_API virtual bool IsCompatibleWithNextStage(const UDMMaterialStage* NextStage) const override;
 	//~ End UDMMaterialStage
 
 	//~ Begin UDMMaterialComponent
-	virtual FString GetComponentPathComponent() const override;
-	virtual UDMMaterialComponent* GetParentComponent() const override;
-	virtual void PostEditorDuplicate(UDynamicMaterialModel* InMaterialModel, UDMMaterialComponent* InParent) override;
+	DYNAMICMATERIALEDITOR_API virtual FString GetComponentPathComponent() const override;
+	DYNAMICMATERIALEDITOR_API virtual UDMMaterialComponent* GetParentComponent() const override;
+	DYNAMICMATERIALEDITOR_API virtual void PostEditorDuplicate(UDynamicMaterialModel* InMaterialModel, UDMMaterialComponent* InParent) override;
 	//~ End UDMMaterialComponent
 
-	void SetParentComponent(UDMMaterialComponent* InParentComponent);
-
 protected:
 	UPROPERTY(VisibleInstanceOnly, BlueprintReadOnly, Category = "Material Designer")
 	TObjectPtr<UDMMaterialStage> ParentStage;
@@ -48,10 +54,4 @@ protected:
 		, ParentComponent(nullptr)
 	{
 	}
-
-	//~ Begin UDMMaterialStage
-	// Sub stages don't need these bound.
-	virtual void AddDelegates() override {}
-	virtual void RemoveDelegates() override {}
-	//~ End UDMMaterialStage
 };
diff --git a/Engine/Plugins/Experimental/DynamicMaterial/Source/DynamicMaterialEditor/Public/Components/MaterialStageBlends/DMMSBAdd.h b/Engine/Plugins/Experimental/DynamicMaterial/Source/DynamicMaterialEditor/Public/Components/MaterialStageBlends/DMMSBAdd.h
--- a/Engine/Plugins/Experimental/DynamicMaterial/Source/DynamicMaterialEditor/Public/Components/MaterialStageBlends/DMMSBAdd.h
+++ b/Engine/Plugins/Experimental/DynamicMaterial/Source/DynamicMaterialEditor/Public/Components/MaterialStageBlends/DMMSBAdd.h
@@ -5,8 +5,8 @@
 #include "Components/DMMaterialStageBlendFunction.h"
 #include "DMMSBAdd.generated.h"
 
-UCLASS(BlueprintType, ClassGroup = "Material Designer")
-class DYNAMICMATERIALEDITOR_API UDMMaterialStageBlendAdd : public UDMMaterialStageBlendFunction
+UCLASS(MinimalAPI, BlueprintType, ClassGroup = "Material Designer")
+class UDMMaterialStageBlendAdd : public UDMMaterialStageBlendFunction
 {
 	GENERATED_BODY()
 
diff --git a/Engine/Plugins/Experimental/DynamicMaterial/Source/DynamicMaterialEditor/Public/Components/MaterialStageBlends/DMMSBColor.h b/Engine/Plugins/Experimental/DynamicMaterial/Source/DynamicMaterialEditor/Public/Components/MaterialStageBlends/DMMSBColor.h
--- a/Engine/Plugins/Experimental/DynamicMaterial/Source/DynamicMaterialEditor/Public/Components/MaterialStageBlends/DMMSBColor.h
+++ b/Engine/Plugins/Experimental/DynamicMaterial/Source/DynamicMaterialEditor/Public/Components/MaterialStageBlends/DMMSBColor.h
@@ -5,8 +5,8 @@
 #include "Components/DMMaterialStageBlendFunction.h"
 #include "DMMSBColor.generated.h"
 
-UCLASS(BlueprintType, ClassGroup = "Material Designer")
-class DYNAMICMATERIALEDITOR_API UDMMaterialStageBlendColor : public UDMMaterialStageBlendFunction
+UCLASS(MinimalAPI, BlueprintType, ClassGroup = "Material Designer")
+class UDMMaterialStageBlendColor : public UDMMaterialStageBlendFunction
 {
 	GENERATED_BODY()
 
diff --git a/Engine/Plugins/Experimental/DynamicMaterial/Source/DynamicMaterialEditor/Public/Components/MaterialStageBlends/DMMSBColorBurn.h b/Engine/Plugins/Experimental/DynamicMaterial/Source/DynamicMaterialEditor/Public/Components/MaterialStageBlends/DMMSBColorBurn.h
--- a/Engine/Plugins/Experimental/DynamicMaterial/Source/DynamicMaterialEditor/Public/Components/MaterialStageBlends/DMMSBColorBurn.h
+++ b/Engine/Plugins/Experimental/DynamicMaterial/Source/DynamicMaterialEditor/Public/Components/MaterialStageBlends/DMMSBColorBurn.h
@@ -5,8 +5,8 @@
 #include "Components/DMMaterialStageBlendFunction.h"
 #include "DMMSBColorBurn.generated.h"
 
-UCLASS(BlueprintType, ClassGroup = "Material Designer")
-class DYNAMICMATERIALEDITOR_API UDMMaterialStageBlendColorBurn : public UDMMaterialStageBlendFunction
+UCLASS(MinimalAPI, BlueprintType, ClassGroup = "Material Designer")
+class UDMMaterialStageBlendColorBurn : public UDMMaterialStageBlendFunction
 {
 	GENERATED_BODY()
 
diff --git a/Engine/Plugins/Experimental/DynamicMaterial/Source/DynamicMaterialEditor/Public/Components/MaterialStageBlends/DMMSBColorDodge.h b/Engine/Plugins/Experimental/DynamicMaterial/Source/DynamicMaterialEditor/Public/Components/MaterialStageBlends/DMMSBColorDodge.h
--- a/Engine/Plugins/Experimental/DynamicMaterial/Source/DynamicMaterialEditor/Public/Components/MaterialStageBlends/DMMSBColorDodge.h
+++ b/Engine/Plugins/Experimental/DynamicMaterial/Source/DynamicMaterialEditor/Public/Components/MaterialStageBlends/DMMSBColorDodge.h
@@ -5,8 +5,8 @@
 #include "Components/DMMaterialStageBlendFunction.h"
 #include "DMMSBColorDodge.generated.h"
 
-UCLASS(BlueprintType, ClassGroup = "Material Designer")
-class DYNAMICMATERIALEDITOR_API UDMMaterialStageBlendColorDodge : public UDMMaterialStageBlendFunction
+UCLASS(MinimalAPI, BlueprintType, ClassGroup = "Material Designer")
+class UDMMaterialStageBlendColorDodge : public UDMMaterialStageBlendFunction
 {
 	GENERATED_BODY()
 
diff --git a/Engine/Plugins/Experimental/DynamicMaterial/Source/DynamicMaterialEditor/Public/Components/MaterialStageBlends/DMMSBContrastBase.h b/Engine/Plugins/Experimental/DynamicMaterial/Source/DynamicMaterialEditor/Public/Components/MaterialStageBlends/DMMSBContrastBase.h
--- a/Engine/Plugins/Experimental/DynamicMaterial/Source/DynamicMaterialEditor/Public/Components/MaterialStageBlends/DMMSBContrastBase.h
+++ b/Engine/Plugins/Experimental/DynamicMaterial/Source/DynamicMaterialEditor/Public/Components/MaterialStageBlends/DMMSBContrastBase.h
@@ -8,14 +8,14 @@
 class UMaterialExpression;
 struct FDMMaterialBuildState;
 
-UCLASS(BlueprintType, Abstract, ClassGroup = "Material Designer")
-class DYNAMICMATERIALEDITOR_API UDMMaterialStageBlendContrastBase : public UDMMaterialStageBlend
+UCLASS(MinimalAPI, BlueprintType, Abstract, ClassGroup = "Material Designer")
+class UDMMaterialStageBlendContrastBase : public UDMMaterialStageBlend
 {
 	GENERATED_BODY()
 
 public:
-	UDMMaterialStageBlendContrastBase(const FText& InName);
+	DYNAMICMATERIALEDITOR_API UDMMaterialStageBlendContrastBase(const FText& InName);
 
-protected:
+private:
 	UDMMaterialStageBlendContrastBase();
 };
diff --git a/Engine/Plugins/Experimental/DynamicMaterial/Source/DynamicMaterialEditor/Public/Components/MaterialStageBlends/DMMSBDarken.h b/Engine/Plugins/Experimental/DynamicMaterial/Source/DynamicMaterialEditor/Public/Components/MaterialStageBlends/DMMSBDarken.h
--- a/Engine/Plugins/Experimental/DynamicMaterial/Source/DynamicMaterialEditor/Public/Components/MaterialStageBlends/DMMSBDarken.h
+++ b/Engine/Plugins/Experimental/DynamicMaterial/Source/DynamicMaterialEditor/Public/Components/MaterialStageBlends/DMMSBDarken.h
@@ -5,8 +5,8 @@
 #include "Components/DMMaterialStageBlendFunction.h"
 #include "DMMSBDarken.generated.h"
 
-UCLASS(BlueprintType, ClassGroup = "Material Designer")
-class DYNAMICMATERIALEDITOR_API UDMMaterialStageBlendDarken : public UDMMaterialStageBlendFunction
+UCLASS(MinimalAPI, BlueprintType, ClassGroup = "Material Designer")
+class UDMMaterialStageBlendDarken : public UDMMaterialStageBlendFunction
 {
 	GENERATED_BODY()
 
diff --git a/Engine/Plugins/Experimental/DynamicMaterial/Source/DynamicMaterialEditor/Public/Components/MaterialStageBlends/DMMSBDarkenColor.h b/Engine/Plugins/Experimental/DynamicMaterial/Source/DynamicMaterialEditor/Public/Components/MaterialStageBlends/DMMSBDarkenColor.h
--- a/Engine/Plugins/Experimental/DynamicMaterial/Source/DynamicMaterialEditor/Public/Components/MaterialStageBlends/DMMSBDarkenColor.h
+++ b/Engine/Plugins/Experimental/DynamicMaterial/Source/DynamicMaterialEditor/Public/Components/MaterialStageBlends/DMMSBDarkenColor.h
@@ -5,8 +5,8 @@
 #include "Components/DMMaterialStageBlendFunction.h"
 #include "DMMSBDarkenColor.generated.h"
 
-UCLASS(BlueprintType, ClassGroup = "Material Designer")
-class DYNAMICMATERIALEDITOR_API UDMMaterialStageBlendDarkenColor : public UDMMaterialStageBlendFunction
+UCLASS(MinimalAPI, BlueprintType, ClassGroup = "Material Designer")
+class UDMMaterialStageBlendDarkenColor : public UDMMaterialStageBlendFunction
 {
 	GENERATED_BODY()
 
diff --git a/Engine/Plugins/Experimental/DynamicMaterial/Source/DynamicMaterialEditor/Public/Components/MaterialStageBlends/DMMSBDifference.h b/Engine/Plugins/Experimental/DynamicMaterial/Source/DynamicMaterialEditor/Public/Components/MaterialStageBlends/DMMSBDifference.h
--- a/Engine/Plugins/Experimental/DynamicMaterial/Source/DynamicMaterialEditor/Public/Components/MaterialStageBlends/DMMSBDifference.h
+++ b/Engine/Plugins/Experimental/DynamicMaterial/Source/DynamicMaterialEditor/Public/Components/MaterialStageBlends/DMMSBDifference.h
@@ -5,8 +5,8 @@
 #include "Components/DMMaterialStageBlendFunction.h"
 #include "DMMSBDifference.generated.h"
 
-UCLASS(BlueprintType, ClassGroup = "Material Designer")
-class DYNAMICMATERIALEDITOR_API UDMMaterialStageBlendDifference : public UDMMaterialStageBlendFunction
+UCLASS(MinimalAPI, BlueprintType, ClassGroup = "Material Designer")
+class UDMMaterialStageBlendDifference : public UDMMaterialStageBlendFunction
 {
 	GENERATED_BODY()
 
diff --git a/Engine/Plugins/Experimental/DynamicMaterial/Source/DynamicMaterialEditor/Public/Components/MaterialStageBlends/DMMSBDivide.h b/Engine/Plugins/Experimental/DynamicMaterial/Source/DynamicMaterialEditor/Public/Components/MaterialStageBlends/DMMSBDivide.h
--- a/Engine/Plugins/Experimental/DynamicMaterial/Source/DynamicMaterialEditor/Public/Components/MaterialStageBlends/DMMSBDivide.h
+++ b/Engine/Plugins/Experimental/DynamicMaterial/Source/DynamicMaterialEditor/Public/Components/MaterialStageBlends/DMMSBDivide.h
@@ -5,8 +5,8 @@
 #include "Components/DMMaterialStageBlendFunction.h"
 #include "DMMSBDivide.generated.h"
 
-UCLASS(BlueprintType, ClassGroup = "Material Designer")
-class DYNAMICMATERIALEDITOR_API UDMMaterialStageBlendDivide : public UDMMaterialStageBlendFunction
+UCLASS(MinimalAPI, BlueprintType, ClassGroup = "Material Designer")
+class UDMMaterialStageBlendDivide : public UDMMaterialStageBlendFunction
 {
 	GENERATED_BODY()
 
diff --git a/Engine/Plugins/Experimental/DynamicMaterial/Source/DynamicMaterialEditor/Public/Components/MaterialStageBlends/DMMSBExclusion.h b/Engine/Plugins/Experimental/DynamicMaterial/Source/DynamicMaterialEditor/Public/Components/MaterialStageBlends/DMMSBExclusion.h
--- a/Engine/Plugins/Experimental/DynamicMaterial/Source/DynamicMaterialEditor/Public/Components/MaterialStageBlends/DMMSBExclusion.h
+++ b/Engine/Plugins/Experimental/DynamicMaterial/Source/DynamicMaterialEditor/Public/Components/MaterialStageBlends/DMMSBExclusion.h
@@ -5,8 +5,8 @@
 #include "Components/DMMaterialStageBlendFunction.h"
 #include "DMMSBExclusion.generated.h"
 
-UCLASS(BlueprintType, ClassGroup = "Material Designer")
-class DYNAMICMATERIALEDITOR_API UDMMaterialStageBlendExclusion : public UDMMaterialStageBlendFunction
+UCLASS(MinimalAPI, BlueprintType, ClassGroup = "Material Designer")
+class UDMMaterialStageBlendExclusion : public UDMMaterialStageBlendFunction
 {
 	GENERATED_BODY()
 
diff --git a/Engine/Plugins/Experimental/DynamicMaterial/Source/DynamicMaterialEditor/Public/Components/MaterialStageBlends/DMMSBHardLight.h b/Engine/Plugins/Experimental/DynamicMaterial/Source/DynamicMaterialEditor/Public/Components/MaterialStageBlends/DMMSBHardLight.h
--- a/Engine/Plugins/Experimental/DynamicMaterial/Source/DynamicMaterialEditor/Public/Components/MaterialStageBlends/DMMSBHardLight.h
+++ b/Engine/Plugins/Experimental/DynamicMaterial/Source/DynamicMaterialEditor/Public/Components/MaterialStageBlends/DMMSBHardLight.h
@@ -5,8 +5,8 @@
 #include "Components/DMMaterialStageBlendFunction.h"
 #include "DMMSBHardLight.generated.h"
 
-UCLASS(BlueprintType, ClassGroup = "Material Designer")
-class DYNAMICMATERIALEDITOR_API UDMMaterialStageBlendHardLight : public UDMMaterialStageBlendFunction
+UCLASS(MinimalAPI, BlueprintType, ClassGroup = "Material Designer")
+class UDMMaterialStageBlendHardLight : public UDMMaterialStageBlendFunction
 {
 	GENERATED_BODY()
 
diff --git a/Engine/Plugins/Experimental/DynamicMaterial/Source/DynamicMaterialEditor/Public/Components/MaterialStageBlends/DMMSBHardMix.h b/Engine/Plugins/Experimental/DynamicMaterial/Source/DynamicMaterialEditor/Public/Components/MaterialStageBlends/DMMSBHardMix.h
--- a/Engine/Plugins/Experimental/DynamicMaterial/Source/DynamicMaterialEditor/Public/Components/MaterialStageBlends/DMMSBHardMix.h
+++ b/Engine/Plugins/Experimental/DynamicMaterial/Source/DynamicMaterialEditor/Public/Components/MaterialStageBlends/DMMSBHardMix.h
@@ -5,8 +5,8 @@
 #include "Components/DMMaterialStageBlendFunction.h"
 #include "DMMSBHardMix.generated.h"
 
-UCLASS(BlueprintType, ClassGroup = "Material Designer")
-class DYNAMICMATERIALEDITOR_API UDMMaterialStageBlendHardMix : public UDMMaterialStageBlendFunction
+UCLASS(MinimalAPI, BlueprintType, ClassGroup = "Material Designer")
+class UDMMaterialStageBlendHardMix : public UDMMaterialStageBlendFunction
 {
 	GENERATED_BODY()
 
diff --git a/Engine/Plugins/Experimental/DynamicMaterial/Source/DynamicMaterialEditor/Public/Components/MaterialStageBlends/DMMSBHue.h b/Engine/Plugins/Experimental/DynamicMaterial/Source/DynamicMaterialEditor/Public/Components/MaterialStageBlends/DMMSBHue.h
--- a/Engine/Plugins/Experimental/DynamicMaterial/Source/DynamicMaterialEditor/Public/Components/MaterialStageBlends/DMMSBHue.h
+++ b/Engine/Plugins/Experimental/DynamicMaterial/Source/DynamicMaterialEditor/Public/Components/MaterialStageBlends/DMMSBHue.h
@@ -5,8 +5,8 @@
 #include "Components/DMMaterialStageBlendFunction.h"
 #include "DMMSBHue.generated.h"
 
-UCLASS(BlueprintType, ClassGroup = "Material Designer")
-class DYNAMICMATERIALEDITOR_API UDMMaterialStageBlendHue : public UDMMaterialStageBlendFunction
+UCLASS(MinimalAPI, BlueprintType, ClassGroup = "Material Designer")
+class UDMMaterialStageBlendHue : public UDMMaterialStageBlendFunction
 {
 	GENERATED_BODY()
 
diff --git a/Engine/Plugins/Experimental/DynamicMaterial/Source/DynamicMaterialEditor/Public/Components/MaterialStageBlends/DMMSBLighten.h b/Engine/Plugins/Experimental/DynamicMaterial/Source/DynamicMaterialEditor/Public/Components/MaterialStageBlends/DMMSBLighten.h
--- a/Engine/Plugins/Experimental/DynamicMaterial/Source/DynamicMaterialEditor/Public/Components/MaterialStageBlends/DMMSBLighten.h
+++ b/Engine/Plugins/Experimental/DynamicMaterial/Source/DynamicMaterialEditor/Public/Components/MaterialStageBlends/DMMSBLighten.h
@@ -5,8 +5,8 @@
 #include "Components/DMMaterialStageBlendFunction.h"
 #include "DMMSBLighten.generated.h"
 
-UCLASS(BlueprintType, ClassGroup = "Material Designer")
-class DYNAMICMATERIALEDITOR_API UDMMaterialStageBlendLighten : public UDMMaterialStageBlendFunction
+UCLASS(MinimalAPI, BlueprintType, ClassGroup = "Material Designer")
+class UDMMaterialStageBlendLighten : public UDMMaterialStageBlendFunction
 {
 	GENERATED_BODY()
 
diff --git a/Engine/Plugins/Experimental/DynamicMaterial/Source/DynamicMaterialEditor/Public/Components/MaterialStageBlends/DMMSBLightenColor.h b/Engine/Plugins/Experimental/DynamicMaterial/Source/DynamicMaterialEditor/Public/Components/MaterialStageBlends/DMMSBLightenColor.h
--- a/Engine/Plugins/Experimental/DynamicMaterial/Source/DynamicMaterialEditor/Public/Components/MaterialStageBlends/DMMSBLightenColor.h
+++ b/Engine/Plugins/Experimental/DynamicMaterial/Source/DynamicMaterialEditor/Public/Components/MaterialStageBlends/DMMSBLightenColor.h
@@ -5,8 +5,8 @@
 #include "Components/DMMaterialStageBlendFunction.h"
 #include "DMMSBLightenColor.generated.h"
 
-UCLASS(BlueprintType, ClassGroup = "Material Designer")
-class DYNAMICMATERIALEDITOR_API UDMMaterialStageBlendLightenColor : public UDMMaterialStageBlendFunction
+UCLASS(MinimalAPI, BlueprintType, ClassGroup = "Material Designer")
+class UDMMaterialStageBlendLightenColor : public UDMMaterialStageBlendFunction
 {
 	GENERATED_BODY()
 
diff --git a/Engine/Plugins/Experimental/DynamicMaterial/Source/DynamicMaterialEditor/Public/Components/MaterialStageBlends/DMMSBLinearBurn.h b/Engine/Plugins/Experimental/DynamicMaterial/Source/DynamicMaterialEditor/Public/Components/MaterialStageBlends/DMMSBLinearBurn.h
--- a/Engine/Plugins/Experimental/DynamicMaterial/Source/DynamicMaterialEditor/Public/Components/MaterialStageBlends/DMMSBLinearBurn.h
+++ b/Engine/Plugins/Experimental/DynamicMaterial/Source/DynamicMaterialEditor/Public/Components/MaterialStageBlends/DMMSBLinearBurn.h
@@ -5,8 +5,8 @@
 #include "Components/DMMaterialStageBlendFunction.h"
 #include "DMMSBLinearBurn.generated.h"
 
-UCLASS(BlueprintType, ClassGroup = "Material Designer")
-class DYNAMICMATERIALEDITOR_API UDMMaterialStageBlendLinearBurn : public UDMMaterialStageBlendFunction
+UCLASS(MinimalAPI, BlueprintType, ClassGroup = "Material Designer")
+class UDMMaterialStageBlendLinearBurn : public UDMMaterialStageBlendFunction
 {
 	GENERATED_BODY()
 
diff --git a/Engine/Plugins/Experimental/DynamicMaterial/Source/DynamicMaterialEditor/Public/Components/MaterialStageBlends/DMMSBLinearDodge.h b/Engine/Plugins/Experimental/DynamicMaterial/Source/DynamicMaterialEditor/Public/Components/MaterialStageBlends/DMMSBLinearDodge.h
--- a/Engine/Plugins/Experimental/DynamicMaterial/Source/DynamicMaterialEditor/Public/Components/MaterialStageBlends/DMMSBLinearDodge.h
+++ b/Engine/Plugins/Experimental/DynamicMaterial/Source/DynamicMaterialEditor/Public/Components/MaterialStageBlends/DMMSBLinearDodge.h
@@ -5,8 +5,8 @@
 #include "Components/DMMaterialStageBlendFunction.h"
 #include "DMMSBLinearDodge.generated.h"
 
-UCLASS(BlueprintType, ClassGroup = "Material Designer")
-class DYNAMICMATERIALEDITOR_API UDMMaterialStageBlendLinearDodge : public UDMMaterialStageBlendFunction
+UCLASS(MinimalAPI, BlueprintType, ClassGroup = "Material Designer")
+class UDMMaterialStageBlendLinearDodge : public UDMMaterialStageBlendFunction
 {
 	GENERATED_BODY()
 
diff --git a/Engine/Plugins/Experimental/DynamicMaterial/Source/DynamicMaterialEditor/Public/Components/MaterialStageBlends/DMMSBLinearLight.h b/Engine/Plugins/Experimental/DynamicMaterial/Source/DynamicMaterialEditor/Public/Components/MaterialStageBlends/DMMSBLinearLight.h
--- a/Engine/Plugins/Experimental/DynamicMaterial/Source/DynamicMaterialEditor/Public/Components/MaterialStageBlends/DMMSBLinearLight.h
+++ b/Engine/Plugins/Experimental/DynamicMaterial/Source/DynamicMaterialEditor/Public/Components/MaterialStageBlends/DMMSBLinearLight.h
@@ -5,8 +5,8 @@
 #include "Components/DMMaterialStageBlendFunction.h"
 #include "DMMSBLinearLight.generated.h"
 
-UCLASS(BlueprintType, ClassGroup = "Material Designer")
-class DYNAMICMATERIALEDITOR_API UDMMaterialStageBlendLinearLight : public UDMMaterialStageBlendFunction
+UCLASS(MinimalAPI, BlueprintType, ClassGroup = "Material Designer")
+class UDMMaterialStageBlendLinearLight : public UDMMaterialStageBlendFunction
 {
 	GENERATED_BODY()
 
diff --git a/Engine/Plugins/Experimental/DynamicMaterial/Source/DynamicMaterialEditor/Public/Components/MaterialStageBlends/DMMSBLuminosity.h b/Engine/Plugins/Experimental/DynamicMaterial/Source/DynamicMaterialEditor/Public/Components/MaterialStageBlends/DMMSBLuminosity.h
--- a/Engine/Plugins/Experimental/DynamicMaterial/Source/DynamicMaterialEditor/Public/Components/MaterialStageBlends/DMMSBLuminosity.h
+++ b/Engine/Plugins/Experimental/DynamicMaterial/Source/DynamicMaterialEditor/Public/Components/MaterialStageBlends/DMMSBLuminosity.h
@@ -5,8 +5,8 @@
 #include "Components/DMMaterialStageBlendFunction.h"
 #include "DMMSBLuminosity.generated.h"
 
-UCLASS(BlueprintType, ClassGroup = "Material Designer")
-class DYNAMICMATERIALEDITOR_API UDMMaterialStageBlendLuminosity : public UDMMaterialStageBlendFunction
+UCLASS(MinimalAPI, BlueprintType, ClassGroup = "Material Designer")
+class UDMMaterialStageBlendLuminosity : public UDMMaterialStageBlendFunction
 {
 	GENERATED_BODY()
 
diff --git a/Engine/Plugins/Experimental/DynamicMaterial/Source/DynamicMaterialEditor/Public/Components/MaterialStageBlends/DMMSBMultiply.h b/Engine/Plugins/Experimental/DynamicMaterial/Source/DynamicMaterialEditor/Public/Components/MaterialStageBlends/DMMSBMultiply.h
--- a/Engine/Plugins/Experimental/DynamicMaterial/Source/DynamicMaterialEditor/Public/Components/MaterialStageBlends/DMMSBMultiply.h
+++ b/Engine/Plugins/Experimental/DynamicMaterial/Source/DynamicMaterialEditor/Public/Components/MaterialStageBlends/DMMSBMultiply.h
@@ -5,8 +5,8 @@
 #include "Components/DMMaterialStageBlendFunction.h"
 #include "DMMSBMultiply.generated.h"
 
-UCLASS(BlueprintType, ClassGroup = "Material Designer")
-class DYNAMICMATERIALEDITOR_API UDMMaterialStageBlendMultiply : public UDMMaterialStageBlendFunction
+UCLASS(MinimalAPI, BlueprintType, ClassGroup = "Material Designer")
+class UDMMaterialStageBlendMultiply : public UDMMaterialStageBlendFunction
 {
 	GENERATED_BODY()
 
diff --git a/Engine/Plugins/Experimental/DynamicMaterial/Source/DynamicMaterialEditor/Public/Components/MaterialStageBlends/DMMSBNormal.h b/Engine/Plugins/Experimental/DynamicMaterial/Source/DynamicMaterialEditor/Public/Components/MaterialStageBlends/DMMSBNormal.h
--- a/Engine/Plugins/Experimental/DynamicMaterial/Source/DynamicMaterialEditor/Public/Components/MaterialStageBlends/DMMSBNormal.h
+++ b/Engine/Plugins/Experimental/DynamicMaterial/Source/DynamicMaterialEditor/Public/Components/MaterialStageBlends/DMMSBNormal.h
@@ -5,8 +5,8 @@
 #include "Components/DMMaterialStageBlendFunction.h"
 #include "DMMSBNormal.generated.h"
 
-UCLASS(BlueprintType, ClassGroup = "Material Designer")
-class DYNAMICMATERIALEDITOR_API UDMMaterialStageBlendNormal : public UDMMaterialStageBlendFunction
+UCLASS(MinimalAPI, BlueprintType, ClassGroup = "Material Designer")
+class UDMMaterialStageBlendNormal : public UDMMaterialStageBlendFunction
 {
 	GENERATED_BODY()
 
diff --git a/Engine/Plugins/Experimental/DynamicMaterial/Source/DynamicMaterialEditor/Public/Components/MaterialStageBlends/DMMSBOverlay.h b/Engine/Plugins/Experimental/DynamicMaterial/Source/DynamicMaterialEditor/Public/Components/MaterialStageBlends/DMMSBOverlay.h
--- a/Engine/Plugins/Experimental/DynamicMaterial/Source/DynamicMaterialEditor/Public/Components/MaterialStageBlends/DMMSBOverlay.h
+++ b/Engine/Plugins/Experimental/DynamicMaterial/Source/DynamicMaterialEditor/Public/Components/MaterialStageBlends/DMMSBOverlay.h
@@ -5,8 +5,8 @@
 #include "Components/DMMaterialStageBlendFunction.h"
 #include "DMMSBOverlay.generated.h"
 
-UCLASS(BlueprintType, ClassGroup = "Material Designer")
-class DYNAMICMATERIALEDITOR_API UDMMaterialStageBlendOverlay : public UDMMaterialStageBlendFunction
+UCLASS(MinimalAPI, BlueprintType, ClassGroup = "Material Designer")
+class UDMMaterialStageBlendOverlay : public UDMMaterialStageBlendFunction
 {
 	GENERATED_BODY()
 
diff --git a/Engine/Plugins/Experimental/DynamicMaterial/Source/DynamicMaterialEditor/Public/Components/MaterialStageBlends/DMMSBPinLight.h b/Engine/Plugins/Experimental/DynamicMaterial/Source/DynamicMaterialEditor/Public/Components/MaterialStageBlends/DMMSBPinLight.h
--- a/Engine/Plugins/Experimental/DynamicMaterial/Source/DynamicMaterialEditor/Public/Components/MaterialStageBlends/DMMSBPinLight.h
+++ b/Engine/Plugins/Experimental/DynamicMaterial/Source/DynamicMaterialEditor/Public/Components/MaterialStageBlends/DMMSBPinLight.h
@@ -5,8 +5,8 @@
 #include "Components/DMMaterialStageBlendFunction.h"
 #include "DMMSBPinLight.generated.h"
 
-UCLASS(BlueprintType, ClassGroup = "Material Designer")
-class DYNAMICMATERIALEDITOR_API UDMMaterialStageBlendPinLight : public UDMMaterialStageBlendFunction
+UCLASS(MinimalAPI, BlueprintType, ClassGroup = "Material Designer")
+class UDMMaterialStageBlendPinLight : public UDMMaterialStageBlendFunction
 {
 	GENERATED_BODY()
 
diff --git a/Engine/Plugins/Experimental/DynamicMaterial/Source/DynamicMaterialEditor/Public/Components/MaterialStageBlends/DMMSBSaturation.h b/Engine/Plugins/Experimental/DynamicMaterial/Source/DynamicMaterialEditor/Public/Components/MaterialStageBlends/DMMSBSaturation.h
--- a/Engine/Plugins/Experimental/DynamicMaterial/Source/DynamicMaterialEditor/Public/Components/MaterialStageBlends/DMMSBSaturation.h
+++ b/Engine/Plugins/Experimental/DynamicMaterial/Source/DynamicMaterialEditor/Public/Components/MaterialStageBlends/DMMSBSaturation.h
@@ -5,8 +5,8 @@
 #include "Components/DMMaterialStageBlendFunction.h"
 #include "DMMSBSaturation.generated.h"
 
-UCLASS(BlueprintType, ClassGroup = "Material Designer")
-class DYNAMICMATERIALEDITOR_API UDMMaterialStageBlendSaturation : public UDMMaterialStageBlendFunction
+UCLASS(MinimalAPI, BlueprintType, ClassGroup = "Material Designer")
+class UDMMaterialStageBlendSaturation : public UDMMaterialStageBlendFunction
 {
 	GENERATED_BODY()
 
diff --git a/Engine/Plugins/Experimental/DynamicMaterial/Source/DynamicMaterialEditor/Public/Components/MaterialStageBlends/DMMSBScreen.h b/Engine/Plugins/Experimental/DynamicMaterial/Source/DynamicMaterialEditor/Public/Components/MaterialStageBlends/DMMSBScreen.h
--- a/Engine/Plugins/Experimental/DynamicMaterial/Source/DynamicMaterialEditor/Public/Components/MaterialStageBlends/DMMSBScreen.h
+++ b/Engine/Plugins/Experimental/DynamicMaterial/Source/DynamicMaterialEditor/Public/Components/MaterialStageBlends/DMMSBScreen.h
@@ -5,8 +5,8 @@
 #include "Components/DMMaterialStageBlendFunction.h"
 #include "DMMSBScreen.generated.h"
 
-UCLASS(BlueprintType, ClassGroup = "Material Designer")
-class DYNAMICMATERIALEDITOR_API UDMMaterialStageBlendScreen : public UDMMaterialStageBlendFunction
+UCLASS(MinimalAPI, BlueprintType, ClassGroup = "Material Designer")
+class UDMMaterialStageBlendScreen : public UDMMaterialStageBlendFunction
 {
 	GENERATED_BODY()
 
diff --git a/Engine/Plugins/Experimental/DynamicMaterial/Source/DynamicMaterialEditor/Public/Components/MaterialStageBlends/DMMSBSoftLight.h b/Engine/Plugins/Experimental/DynamicMaterial/Source/DynamicMaterialEditor/Public/Components/MaterialStageBlends/DMMSBSoftLight.h
--- a/Engine/Plugins/Experimental/DynamicMaterial/Source/DynamicMaterialEditor/Public/Components/MaterialStageBlends/DMMSBSoftLight.h
+++ b/Engine/Plugins/Experimental/DynamicMaterial/Source/DynamicMaterialEditor/Public/Components/MaterialStageBlends/DMMSBSoftLight.h
@@ -5,8 +5,8 @@
 #include "Components/DMMaterialStageBlendFunction.h"
 #include "DMMSBSoftLight.generated.h"
 
-UCLASS(BlueprintType, ClassGroup = "Material Designer")
-class DYNAMICMATERIALEDITOR_API UDMMaterialStageBlendSoftLight : public UDMMaterialStageBlendFunction
+UCLASS(MinimalAPI, BlueprintType, ClassGroup = "Material Designer")
+class UDMMaterialStageBlendSoftLight : public UDMMaterialStageBlendFunction
 {
 	GENERATED_BODY()
 
diff --git a/Engine/Plugins/Experimental/DynamicMaterial/Source/DynamicMaterialEditor/Public/Components/MaterialStageBlends/DMMSBSubtract.h b/Engine/Plugins/Experimental/DynamicMaterial/Source/DynamicMaterialEditor/Public/Components/MaterialStageBlends/DMMSBSubtract.h
--- a/Engine/Plugins/Experimental/DynamicMaterial/Source/DynamicMaterialEditor/Public/Components/MaterialStageBlends/DMMSBSubtract.h
+++ b/Engine/Plugins/Experimental/DynamicMaterial/Source/DynamicMaterialEditor/Public/Components/MaterialStageBlends/DMMSBSubtract.h
@@ -5,8 +5,8 @@
 #include "Components/DMMaterialStageBlendFunction.h"
 #include "DMMSBSubtract.generated.h"
 
-UCLASS(BlueprintType, ClassGroup = "Material Designer")
-class DYNAMICMATERIALEDITOR_API UDMMaterialStageBlendSubtract : public UDMMaterialStageBlendFunction
+UCLASS(MinimalAPI, BlueprintType, ClassGroup = "Material Designer")
+class UDMMaterialStageBlendSubtract : public UDMMaterialStageBlendFunction
 {
 	GENERATED_BODY()
 
diff --git a/Engine/Plugins/Experimental/DynamicMaterial/Source/DynamicMaterialEditor/Public/Components/MaterialStageBlends/DMMSBVividLight.h b/Engine/Plugins/Experimental/DynamicMaterial/Source/DynamicMaterialEditor/Public/Components/MaterialStageBlends/DMMSBVividLight.h
--- a/Engine/Plugins/Experimental/DynamicMaterial/Source/DynamicMaterialEditor/Public/Components/MaterialStageBlends/DMMSBVividLight.h
+++ b/Engine/Plugins/Experimental/DynamicMaterial/Source/DynamicMaterialEditor/Public/Components/MaterialStageBlends/DMMSBVividLight.h
@@ -5,8 +5,8 @@
 #include "Components/DMMaterialStageBlendFunction.h"
 #include "DMMSBVividLight.generated.h"
 
-UCLASS(BlueprintType, ClassGroup = "Material Designer")
-class DYNAMICMATERIALEDITOR_API UDMMaterialStageBlendVividLight : public UDMMaterialStageBlendFunction
+UCLASS(MinimalAPI, BlueprintType, ClassGroup = "Material Designer")
+class UDMMaterialStageBlendVividLight : public UDMMaterialStageBlendFunction
 {
 	GENERATED_BODY()
 
diff --git a/Engine/Plugins/Experimental/DynamicMaterial/Source/DynamicMaterialEditor/Public/Components/MaterialStageExpressions/DMMSEMathBase.h b/Engine/Plugins/Experimental/DynamicMaterial/Source/DynamicMaterialEditor/Public/Components/MaterialStageExpressions/DMMSEMathBase.h
--- a/Engine/Plugins/Experimental/DynamicMaterial/Source/DynamicMaterialEditor/Public/Components/MaterialStageExpressions/DMMSEMathBase.h
+++ b/Engine/Plugins/Experimental/DynamicMaterial/Source/DynamicMaterialEditor/Public/Components/MaterialStageExpressions/DMMSEMathBase.h
@@ -8,41 +8,47 @@
 
 class UDMMaterialStageInput;
 
-UCLASS(BlueprintType, ClassGroup = "Material Designer")
-class DYNAMICMATERIALEDITOR_API UDMMaterialStageExpressionMathBase : public UDMMaterialStageExpression
+UCLASS(MinimalAPI, BlueprintType, ClassGroup = "Material Designer")
+class UDMMaterialStageExpressionMathBase : public UDMMaterialStageExpression
 {
 	GENERATED_BODY()
 
 public:
+	DYNAMICMATERIALEDITOR_API UDMMaterialStageExpressionMathBase(const FText& InName, TSubclassOf<UMaterialExpression> InClass);
 
-	UDMMaterialStageExpressionMathBase(const FText& InName, TSubclassOf<UMaterialExpression> InClass);
-
-	virtual bool CanInputAcceptType(int32 InInputIndex, EDMValueType InValueType) const override;
-
-	virtual void Update(EDMUpdateType InUpdateType) override;
-
-	virtual void AddDefaultInput(int32 InInputIndex) const override;
+	//~ Begin UDMMaterialStageThroughput
+	DYNAMICMATERIALEDITOR_API virtual bool CanInputAcceptType(int32 InInputIndex, EDMValueType InValueType) const override;
+	DYNAMICMATERIALEDITOR_API virtual void AddDefaultInput(int32 InInputIndex) const override;
+	//~ End UDMMaterialStageThroughput
 
 	//~ Begin UDMMaterialStageSource
-	virtual bool UpdateStagePreviewMaterial(UDMMaterialStage* InStage, UMaterial* InPreviewMaterial, UMaterialExpression*& OutMaterialExpression,
-		int32& OutputIndex) override;
+	DYNAMICMATERIALEDITOR_API virtual bool GenerateStagePreviewMaterial(UDMMaterialStage* InStage, UMaterial* InPreviewMaterial, 
+		UMaterialExpression*& OutMaterialExpression, int32& OutputIndex) override;
 	//~ End UDMMaterialStageSource
 
-protected:
+	//~ Begin UDMMaterialComponent
+	DYNAMICMATERIALEDITOR_API virtual void Update(EDMUpdateType InUpdateType) override;
+	//~ End UDMMaterialComponent
 
+protected:
+	/** Whether this just works with Scalars. */
 	UPROPERTY(VisibleInstanceOnly, BlueprintReadOnly, Category = "Material Designer")
 	bool bSingleChannelOnly;
 
+	/** How many inputs this should automatically create. */
 	UPROPERTY(VisibleInstanceOnly, BlueprintReadOnly, Category = "Material Designer")
 	int32 VariadicInputCount;
 
+	/** Allow matching of single floats to match other types of floats. */
 	UPROPERTY(VisibleInstanceOnly, BlueprintReadOnly, Category = "Material Designer")
 	bool bAllowSingleFloatMatch;
 
 	UDMMaterialStageExpressionMathBase();
 
-	void SetupInputs(int32 InCount);
-
-	virtual void UpdatePreviewMaterial(UMaterial* InPreviewMaterial = nullptr) override;
+	/** Add automatic inputs (A-O (15))) */
+	DYNAMICMATERIALEDITOR_API void SetupInputs(int32 InCount);
 
+	//~ Begin UDMMaterialStageThroughput
+	DYNAMICMATERIALEDITOR_API virtual void GeneratePreviewMaterial(UMaterial* InPreviewMaterial) override;
+	//~ End UDMMaterialStageThroughput
 };
diff --git a/Engine/Plugins/Experimental/DynamicMaterial/Source/DynamicMaterialEditor/Public/Components/MaterialStageExpressions/DMMSETextureCoordinate.h b/Engine/Plugins/Experimental/DynamicMaterial/Source/DynamicMaterialEditor/Public/Components/MaterialStageExpressions/DMMSETextureCoordinate.h
--- a/Engine/Plugins/Experimental/DynamicMaterial/Source/DynamicMaterialEditor/Public/Components/MaterialStageExpressions/DMMSETextureCoordinate.h
+++ b/Engine/Plugins/Experimental/DynamicMaterial/Source/DynamicMaterialEditor/Public/Components/MaterialStageExpressions/DMMSETextureCoordinate.h
@@ -5,15 +5,17 @@
 #include "Components/DMMaterialStageExpression.h"
 #include "DMMSETextureCoordinate.generated.h"
 
-UCLASS(BlueprintType, ClassGroup = "Material Designer")
-class DYNAMICMATERIALEDITOR_API UDMMaterialStageExpressionTextureCoordinate : public UDMMaterialStageExpression
+UCLASS(MinimalAPI, BlueprintType, ClassGroup = "Material Designer")
+class UDMMaterialStageExpressionTextureCoordinate : public UDMMaterialStageExpression
 {
 	GENERATED_BODY()
 
 public:
 	UDMMaterialStageExpressionTextureCoordinate();
 
-	virtual void AddExpressionProperties(const TArray<UMaterialExpression*>& InExpressions) const override;
+	//~ Begin UDMMaterialStageSource
+	DYNAMICMATERIALEDITOR_API virtual void AddExpressionProperties(const TArray<UMaterialExpression*>& InExpressions) const override;
+	//~ End UDMMaterialStageSource
 
 protected:
 	UPROPERTY(EditInstanceOnly, BlueprintReadOnly, Category = "Material Designer")
diff --git a/Engine/Plugins/Experimental/DynamicMaterial/Source/DynamicMaterialEditor/Public/Components/MaterialStageExpressions/DMMSETextureSample.h b/Engine/Plugins/Experimental/DynamicMaterial/Source/DynamicMaterialEditor/Public/Components/MaterialStageExpressions/DMMSETextureSample.h
--- a/Engine/Plugins/Experimental/DynamicMaterial/Source/DynamicMaterialEditor/Public/Components/MaterialStageExpressions/DMMSETextureSample.h
+++ b/Engine/Plugins/Experimental/DynamicMaterial/Source/DynamicMaterialEditor/Public/Components/MaterialStageExpressions/DMMSETextureSample.h
@@ -5,8 +5,8 @@
 #include "Components/MaterialStageExpressions/DMMSETextureSampleBase.h"
 #include "DMMSETextureSample.generated.h"
 
-UCLASS(BlueprintType, ClassGroup = "Material Designer")
-class DYNAMICMATERIALEDITOR_API UDMMaterialStageExpressionTextureSample : public UDMMaterialStageExpressionTextureSampleBase
+UCLASS(MinimalAPI, BlueprintType, ClassGroup = "Material Designer")
+class UDMMaterialStageExpressionTextureSample : public UDMMaterialStageExpressionTextureSampleBase
 {
 	GENERATED_BODY()
 
diff --git a/Engine/Plugins/Experimental/DynamicMaterial/Source/DynamicMaterialEditor/Public/Components/MaterialStageExpressions/DMMSETextureSampleBase.h b/Engine/Plugins/Experimental/DynamicMaterial/Source/DynamicMaterialEditor/Public/Components/MaterialStageExpressions/DMMSETextureSampleBase.h
--- a/Engine/Plugins/Experimental/DynamicMaterial/Source/DynamicMaterialEditor/Public/Components/MaterialStageExpressions/DMMSETextureSampleBase.h
+++ b/Engine/Plugins/Experimental/DynamicMaterial/Source/DynamicMaterialEditor/Public/Components/MaterialStageExpressions/DMMSETextureSampleBase.h
@@ -5,41 +5,41 @@
 #include "Components/DMMaterialStageExpression.h"
 #include "DMMSETextureSampleBase.generated.h"
 
-UCLASS(BlueprintType, Blueprintable, ClassGroup = "Material Designer")
-class DYNAMICMATERIALEDITOR_API UDMMaterialStageExpressionTextureSampleBase : public UDMMaterialStageExpression
+UCLASS(minimalAPI, BlueprintType, Blueprintable, ClassGroup = "Material Designer")
+class UDMMaterialStageExpressionTextureSampleBase : public UDMMaterialStageExpression
 {
 	GENERATED_BODY()
 
 public:
-	UDMMaterialStageExpressionTextureSampleBase(const FText& InName, TSubclassOf<UMaterialExpression> InClass);
+	UDMMaterialStageExpressionTextureSampleBase();
+
+	UFUNCTION(BlueprintPure, Category = "Material Designer")
+	virtual bool IsClampTextureEnabled() const { return bClampTexture; }
+
+	UFUNCTION(BlueprintCallable, Category = "Material Designer")
+	DYNAMICMATERIALEDITOR_API virtual void SetClampTextureEnabled(bool bInValue);
 
 	//~ Begin UDMMaterialStageExpression
-	virtual void GenerateExpressions(const TSharedRef<FDMMaterialBuildState>& InBuildState) const override;
+	DYNAMICMATERIALEDITOR_API virtual void GenerateExpressions(const TSharedRef<FDMMaterialBuildState>& InBuildState) const override;
 	//~ End UDMMaterialStageExpression
 
 	//~ Begin UDMMaterialStageSource
-	virtual void OnComponentAdded() override;
-	virtual bool IsPropertyVisible(FName Property) const override;
-	virtual void AddExpressionProperties(const TArray<UMaterialExpression*>& InExpressions) const override;
-	virtual int32 GetInnateMaskOutput(int32 InOutputIndex, int32 InOutputChannels) const override;
-	virtual int32 GetOutputChannelOverride(int32 InOutputIndex) const override;
+	DYNAMICMATERIALEDITOR_API virtual void OnComponentAdded() override;
+	DYNAMICMATERIALEDITOR_API virtual bool IsPropertyVisible(FName Property) const override;
+	DYNAMICMATERIALEDITOR_API virtual void AddExpressionProperties(const TArray<UMaterialExpression*>& InExpressions) const override;
+	DYNAMICMATERIALEDITOR_API virtual int32 GetInnateMaskOutput(int32 InOutputIndex, int32 InOutputChannels) const override;
+	DYNAMICMATERIALEDITOR_API virtual int32 GetOutputChannelOverride(int32 InOutputIndex) const override;
 	//~ End UDMMaterialStageSource
 
 	//~ Begin UDMMaterialStageThroughput
-	virtual bool CanChangeInputType(int32 InInputIndex) const override;
+	DYNAMICMATERIALEDITOR_API virtual bool CanChangeInputType(int32 InInputIndex) const override;
 	virtual bool SupportsLayerMaskTextureUVLink() const override { return true; }
 	virtual int32 GetLayerMaskTextureUVLinkInputIndex() const override { return 1; }
-	virtual void InputUpdated(int32 InInputIndex, EDMUpdateType InUpdateType) override;
+	DYNAMICMATERIALEDITOR_API virtual void InputUpdated(int32 InInputIndex, EDMUpdateType InUpdateType) override;
 	//~ End UDMMaterialStageThroughput
 
-	UFUNCTION(BlueprintPure, Category = "Material Designer")
-	virtual bool IsClampTextureEnabled() const { return bClampTexture; }
-
-	UFUNCTION(BlueprintCallable, Category = "Material Designer")
-	virtual void SetClampTextureEnabled(bool bInValue);
-
 	//~ Begin UObject
-	virtual void PostEditChangeProperty(FPropertyChangedEvent& InPropertyChangedEvent) override;
+	DYNAMICMATERIALEDITOR_API virtual void PostEditChangeProperty(FPropertyChangedEvent& InPropertyChangedEvent) override;
 	//~ End UObject
 
 protected:
@@ -47,7 +47,7 @@ protected:
 		meta=(ToolTip="Forces a material rebuild.", LowPriority, NotKeyframeable, AllowPrivateAccess = "true"))
 	bool bClampTexture;
 
-	UDMMaterialStageExpressionTextureSampleBase();
+	DYNAMICMATERIALEDITOR_API UDMMaterialStageExpressionTextureSampleBase(const FText& InName, TSubclassOf<UMaterialExpression> InClass);
 
 	void UpdateMask();
 };
diff --git a/Engine/Plugins/Experimental/DynamicMaterial/Source/DynamicMaterialEditor/Public/Components/MaterialStageExpressions/DMMSETextureSampleEdgeColor.h b/Engine/Plugins/Experimental/DynamicMaterial/Source/DynamicMaterialEditor/Public/Components/MaterialStageExpressions/DMMSETextureSampleEdgeColor.h
--- a/Engine/Plugins/Experimental/DynamicMaterial/Source/DynamicMaterialEditor/Public/Components/MaterialStageExpressions/DMMSETextureSampleEdgeColor.h
+++ b/Engine/Plugins/Experimental/DynamicMaterial/Source/DynamicMaterialEditor/Public/Components/MaterialStageExpressions/DMMSETextureSampleEdgeColor.h
@@ -20,8 +20,8 @@ enum class EDMEdgeLocation : uint8
 	Custom
 };
 
-UCLASS(BlueprintType, ClassGroup = "Material Designer")
-class DYNAMICMATERIALEDITOR_API UDMMaterialStageExpressionTextureSampleEdgeColor : public UDMMaterialStageExpression
+UCLASS(MinimalAPI, BlueprintType, ClassGroup = "Material Designer")
+class UDMMaterialStageExpressionTextureSampleEdgeColor : public UDMMaterialStageExpression
 {
 	GENERATED_BODY()
 
@@ -29,13 +29,13 @@ public:
 	UDMMaterialStageExpressionTextureSampleEdgeColor();
 
 	//~ Begin UDMMaterialStageThroughput
-	virtual bool CanChangeInputType(int32 InInputIndex) const override;
-	virtual bool IsInputVisible(int32 InInputIndex) const override;
-	virtual void AddDefaultInput(int32 InInputIndex) const override;
+	DYNAMICMATERIALEDITOR_API virtual bool CanChangeInputType(int32 InInputIndex) const override;
+	DYNAMICMATERIALEDITOR_API virtual bool IsInputVisible(int32 InInputIndex) const override;
+	DYNAMICMATERIALEDITOR_API virtual void AddDefaultInput(int32 InInputIndex) const override;
 	//~ End UDMMaterialStageThroughput
 
 	//~ Begin UObject
-	virtual void PostEditChangeProperty(struct FPropertyChangedEvent& PropertyChangedEvent) override;
+	DYNAMICMATERIALEDITOR_API virtual void PostEditChangeProperty(FPropertyChangedEvent& InPropertyChangedEvent) override;
 	//~ End UObject
 
 protected:
diff --git a/Engine/Plugins/Experimental/DynamicMaterial/Source/DynamicMaterialEditor/Public/Components/MaterialStageGradients/DMMSGLinear.h b/Engine/Plugins/Experimental/DynamicMaterial/Source/DynamicMaterialEditor/Public/Components/MaterialStageGradients/DMMSGLinear.h
--- a/Engine/Plugins/Experimental/DynamicMaterial/Source/DynamicMaterialEditor/Public/Components/MaterialStageGradients/DMMSGLinear.h
+++ b/Engine/Plugins/Experimental/DynamicMaterial/Source/DynamicMaterialEditor/Public/Components/MaterialStageGradients/DMMSGLinear.h
@@ -16,21 +16,23 @@ enum class ELinearGradientTileType : uint8
 	TileAndMirror
 };
 
-UCLASS(BlueprintType, ClassGroup = "Material Designer")
-class DYNAMICMATERIALEDITOR_API UDMMaterialStageGradientLinear : public UDMMaterialStageGradient
+UCLASS(MinimalAPI, BlueprintType, ClassGroup = "Material Designer")
+class UDMMaterialStageGradientLinear : public UDMMaterialStageGradient
 {
 	GENERATED_BODY()
 
 public:
 	UDMMaterialStageGradientLinear();
 
-	virtual void GenerateExpressions(const TSharedRef<FDMMaterialBuildState>& InBuildState) const override;
-
 	UFUNCTION(BlueprintPure, Category = "Material Designer")
 	virtual ELinearGradientTileType GetTilingType() const { return Tiling; }
 
 	UFUNCTION(BlueprintCallable, Category = "Material Designer")
-	virtual void SetTilingType(ELinearGradientTileType InType);
+	DYNAMICMATERIALEDITOR_API virtual void SetTilingType(ELinearGradientTileType InType);
+
+	//~ Begin UDMMaterialStageSource
+	DYNAMICMATERIALEDITOR_API virtual void GenerateExpressions(const TSharedRef<FDMMaterialBuildState>& InBuildState) const override;
+	//~ End UDMMaterialStageSource
 
 protected:
 	UPROPERTY(EditInstanceOnly, BlueprintReadWrite, Getter = GetTilingType, Setter=SetTilingType, BlueprintSetter = SetTilingType, 
diff --git a/Engine/Plugins/Experimental/DynamicMaterial/Source/DynamicMaterialEditor/Public/Components/MaterialStageGradients/DMMSGRadial.h b/Engine/Plugins/Experimental/DynamicMaterial/Source/DynamicMaterialEditor/Public/Components/MaterialStageGradients/DMMSGRadial.h
--- a/Engine/Plugins/Experimental/DynamicMaterial/Source/DynamicMaterialEditor/Public/Components/MaterialStageGradients/DMMSGRadial.h
+++ b/Engine/Plugins/Experimental/DynamicMaterial/Source/DynamicMaterialEditor/Public/Components/MaterialStageGradients/DMMSGRadial.h
@@ -8,13 +8,15 @@
 struct FDMMaterialBuildState;
 class UMaterialExpression;
 
-UCLASS(BlueprintType, ClassGroup = "Material Designer")
-class DYNAMICMATERIALEDITOR_API UDMMaterialStageGradientRadial : public UDMMaterialStageGradient
+UCLASS(MinimalAPI, BlueprintType, ClassGroup = "Material Designer")
+class UDMMaterialStageGradientRadial : public UDMMaterialStageGradient
 {
 	GENERATED_BODY()
 
 public:
 	UDMMaterialStageGradientRadial();
 
-	virtual void GenerateExpressions(const TSharedRef<FDMMaterialBuildState>& InBuildState) const override;
+	//~ Begin UDMMaterialStageSource
+	DYNAMICMATERIALEDITOR_API virtual void GenerateExpressions(const TSharedRef<FDMMaterialBuildState>& InBuildState) const override;
+	//~ End UDMMaterialStageSource
 };
diff --git a/Engine/Plugins/Experimental/DynamicMaterial/Source/DynamicMaterialEditor/Public/Components/MaterialStageInputs/DMMSIExpression.h b/Engine/Plugins/Experimental/DynamicMaterial/Source/DynamicMaterialEditor/Public/Components/MaterialStageInputs/DMMSIExpression.h
--- a/Engine/Plugins/Experimental/DynamicMaterial/Source/DynamicMaterialEditor/Public/Components/MaterialStageInputs/DMMSIExpression.h
+++ b/Engine/Plugins/Experimental/DynamicMaterial/Source/DynamicMaterialEditor/Public/Components/MaterialStageInputs/DMMSIExpression.h
@@ -10,18 +10,19 @@
 class UDMMaterialLayerObject;
 class UDMMaterialStageExpression;
 
-UCLASS(BlueprintType, ClassGroup = "Material Designer")
-class DYNAMICMATERIALEDITOR_API UDMMaterialStageInputExpression : public UDMMaterialStageInputThroughput
+UCLASS(MinimalAPI, BlueprintType, ClassGroup = "Material Designer")
+class UDMMaterialStageInputExpression : public UDMMaterialStageInputThroughput
 {
 	GENERATED_BODY()
 
 public:
-	static UDMMaterialStage* CreateStage(TSubclassOf<UDMMaterialStageExpression> InMaterialStageExpressionClass, UDMMaterialLayerObject* InLayer = nullptr);
+	UFUNCTION(BlueprintCallable, Category = "Material Designer")
+	static DYNAMICMATERIALEDITOR_API UDMMaterialStage* CreateStage(TSubclassOf<UDMMaterialStageExpression> InMaterialStageExpressionClass, UDMMaterialLayerObject* InLayer = nullptr);
 
 	static const TArray<TStrongObjectPtr<UClass>>& GetAvailableInputExpressions();
 
 	UFUNCTION(BlueprintCallable, Category = "Material Designer")
-	static UDMMaterialStageInputExpression* ChangeStageSource_Expression(UDMMaterialStage* InStage, 
+	static DYNAMICMATERIALEDITOR_API UDMMaterialStageInputExpression* ChangeStageSource_Expression(UDMMaterialStage* InStage,
 		TSubclassOf<UDMMaterialStageExpression> InExpressionClass);
 
 	template<typename InExpressionClass>
@@ -30,8 +31,15 @@ public:
 		return ChangeStageSource_Expression(InStage, InExpressionClass::StaticClass());
 	}
 
+	/**
+	 * Change the input type of an input on a stage to an expression.
+	 * @param InInputIdx Index of the source input.
+	 * @param InInputChannel The channel of the input that the input connects to.
+	 * @param InOutputIdx The output index of the new input.
+	 * @param InOutputChannel The channel of the output to connect.
+	 */
 	UFUNCTION(BlueprintCallable, Category = "Material Designer")
-	static UDMMaterialStageInputExpression* ChangeStageInput_Expression(UDMMaterialStage* InStage, 
+	static DYNAMICMATERIALEDITOR_API UDMMaterialStageInputExpression* ChangeStageInput_Expression(UDMMaterialStage* InStage,
 		TSubclassOf<UDMMaterialStageExpression> InExpressionClass, int32 InInputIdx, int32 InInputChannel, int32 InOutputIdx,
 		int32 InOutputChannel);
 
@@ -44,18 +52,16 @@ public:
 	}
 
 	UFUNCTION(BlueprintPure, Category = "Material Designer")
-	TSubclassOf<UDMMaterialStageExpression> GetMaterialStageExpressionClass() const;
+	DYNAMICMATERIALEDITOR_API TSubclassOf<UDMMaterialStageExpression> GetMaterialStageExpressionClass() const;
 
 	UFUNCTION(BlueprintCallable, Category = "Material Designer")
-	void SetMaterialStageExpressionClass(TSubclassOf<UDMMaterialStageExpression> InMaterialStageExpressionClass);
+	DYNAMICMATERIALEDITOR_API void SetMaterialStageExpressionClass(TSubclassOf<UDMMaterialStageExpression> InMaterialStageExpressionClass);
 
 	UFUNCTION(BlueprintPure, Category = "Material Designer")
-	UDMMaterialStageExpression* GetMaterialStageExpression() const;
+	DYNAMICMATERIALEDITOR_API UDMMaterialStageExpression* GetMaterialStageExpression() const;
 
 protected:
 	static TArray<TStrongObjectPtr<UClass>> InputExpressions;
 
 	static void GenerateExpressionList();
-
-	UDMMaterialStageInputExpression() = default;
 };
diff --git a/Engine/Plugins/Experimental/DynamicMaterial/Source/DynamicMaterialEditor/Public/Components/MaterialStageInputs/DMMSIFunction.h b/Engine/Plugins/Experimental/DynamicMaterial/Source/DynamicMaterialEditor/Public/Components/MaterialStageInputs/DMMSIFunction.h
--- a/Engine/Plugins/Experimental/DynamicMaterial/Source/DynamicMaterialEditor/Public/Components/MaterialStageInputs/DMMSIFunction.h
+++ b/Engine/Plugins/Experimental/DynamicMaterial/Source/DynamicMaterialEditor/Public/Components/MaterialStageInputs/DMMSIFunction.h
@@ -9,34 +9,39 @@ class UDMMaterialLayerObject;
 class UDMMaterialStageFunction;
 class UMaterialFunctionInterface;
 
-UCLASS(BlueprintType, ClassGroup = "Material Designer")
-class DYNAMICMATERIALEDITOR_API UDMMaterialStageInputFunction : public UDMMaterialStageInputThroughput
+UCLASS(MinimalAPI, BlueprintType, ClassGroup = "Material Designer")
+class UDMMaterialStageInputFunction : public UDMMaterialStageInputThroughput
 {
 	GENERATED_BODY()
 
 public:
-	static UDMMaterialStage* CreateStage(UDMMaterialLayerObject* InLayer = nullptr);
+	UFUNCTION(BlueprintCallable, Category = "Material Designer")
+	static DYNAMICMATERIALEDITOR_API UDMMaterialStage* CreateStage(UDMMaterialLayerObject* InLayer = nullptr);
 
 	UFUNCTION(BlueprintCallable, Category = "Material Designer")
-	static UDMMaterialStageInputFunction* ChangeStageSource_Function(UDMMaterialStage* InStage,
+	static DYNAMICMATERIALEDITOR_API UDMMaterialStageInputFunction* ChangeStageSource_Function(UDMMaterialStage* InStage,
 		UMaterialFunctionInterface* InMaterialFunction);
 
+	/**
+	 * Change the input type of an input on a stage to a function.
+	 * @param InInputIdx Index of the source input.
+	 * @param InInputChannel The channel of the input that the input connects to.
+	 * @param InOutputIdx The output index of the new input.
+	 * @param InOutputChannel The channel of the output to connect.
+	 */
 	UFUNCTION(BlueprintCallable, Category = "Material Designer")
-	static UDMMaterialStageInputFunction* ChangeStageInput_Function(UDMMaterialStage* InStage,
+	static DYNAMICMATERIALEDITOR_API UDMMaterialStageInputFunction* ChangeStageInput_Function(UDMMaterialStage* InStage,
 		UMaterialFunctionInterface* InMaterialFunction, int32 InInputIdx, int32 InInputChannel, int32 InOutputIdx,
 		int32 InOutputChannel);
 
 	void Init();
 
 	UFUNCTION(BlueprintPure, Category = "Material Designer")
-	UDMMaterialStageFunction* GetMaterialStageFunction() const;
+	DYNAMICMATERIALEDITOR_API UDMMaterialStageFunction* GetMaterialStageFunction() const;
 
 	UFUNCTION(BlueprintPure, Category = "Material Designer")
-	UMaterialFunctionInterface* GetMaterialFunction() const;
+	DYNAMICMATERIALEDITOR_API UMaterialFunctionInterface* GetMaterialFunction() const;
 
 	UFUNCTION(BlueprintCallable, Category = "Material Designer")
-	void SetMaterialFunction(UMaterialFunctionInterface* InMaterialFunction);
-
-protected:
-	UDMMaterialStageInputFunction() = default;
+	DYNAMICMATERIALEDITOR_API void SetMaterialFunction(UMaterialFunctionInterface* InMaterialFunction);
 };
diff --git a/Engine/Plugins/Experimental/DynamicMaterial/Source/DynamicMaterialEditor/Public/Components/MaterialStageInputs/DMMSIGradient.h b/Engine/Plugins/Experimental/DynamicMaterial/Source/DynamicMaterialEditor/Public/Components/MaterialStageInputs/DMMSIGradient.h
--- a/Engine/Plugins/Experimental/DynamicMaterial/Source/DynamicMaterialEditor/Public/Components/MaterialStageInputs/DMMSIGradient.h
+++ b/Engine/Plugins/Experimental/DynamicMaterial/Source/DynamicMaterialEditor/Public/Components/MaterialStageInputs/DMMSIGradient.h
@@ -10,18 +10,19 @@
 class UDMMaterialLayerObject;
 class UDMMaterialStageGradient;
 
-UCLASS(BlueprintType, ClassGroup = "Material Designer")
-class DYNAMICMATERIALEDITOR_API UDMMaterialStageInputGradient : public UDMMaterialStageInputThroughput
+UCLASS(MinimalAPI, BlueprintType, ClassGroup = "Material Designer")
+class UDMMaterialStageInputGradient : public UDMMaterialStageInputThroughput
 {
 	GENERATED_BODY()
 
 public:
-	static UDMMaterialStage* CreateStage(TSubclassOf<UDMMaterialStageGradient> InMaterialStageGradientClass, UDMMaterialLayerObject* InLayer = nullptr);
+	UFUNCTION(BlueprintCallable, Category = "Material Designer")
+	static DYNAMICMATERIALEDITOR_API UDMMaterialStage* CreateStage(TSubclassOf<UDMMaterialStageGradient> InMaterialStageGradientClass, UDMMaterialLayerObject* InLayer = nullptr);
 
 	static const TArray<TStrongObjectPtr<UClass>>& GetAvailableGradients();
 
 	UFUNCTION(BlueprintCallable, Category = "Material Designer")
-	static UDMMaterialStageInputGradient* ChangeStageSource_Gradient(UDMMaterialStage* InStage,
+	static DYNAMICMATERIALEDITOR_API UDMMaterialStageInputGradient* ChangeStageSource_Gradient(UDMMaterialStage* InStage,
 		TSubclassOf<UDMMaterialStageGradient> InGradientClass);
 
 	template<typename InGradientClass>
@@ -30,8 +31,14 @@ public:
 		return ChangeStageSource_Gradient(InStage, InGradientClass::StaticClass());
 	}
 
+	/**
+	 * Change the input type of an input on a stage to a gradient.
+	 * @param InInputIdx Index of the source input.
+	 * @param InInputChannel The channel of the input that the input connects to.
+	 * @param InOutputChannel The channel of the output to connect.
+	 */
 	UFUNCTION(BlueprintCallable, Category = "Material Designer")
-	static UDMMaterialStageInputGradient* ChangeStageInput_Gradient(UDMMaterialStage* InStage,
+	static DYNAMICMATERIALEDITOR_API UDMMaterialStageInputGradient* ChangeStageInput_Gradient(UDMMaterialStage* InStage,
 		TSubclassOf<UDMMaterialStageGradient> InGradientClass, int32 InInputIdx, int32 InInputChannel, int32 InOutputChannel);
 
 	template<typename InGradientClass>
@@ -42,18 +49,16 @@ public:
 	}
 
 	UFUNCTION(BlueprintPure, Category = "Material Designer")
-	TSubclassOf<UDMMaterialStageGradient> GetMaterialStageGradientClass() const;
+	DYNAMICMATERIALEDITOR_API TSubclassOf<UDMMaterialStageGradient> GetMaterialStageGradientClass() const;
 
 	UFUNCTION(BlueprintCallable, Category = "Material Designer")
-	void SetMaterialStageGradientClass(TSubclassOf<UDMMaterialStageGradient> InMaterialStageGradientClass);
+	DYNAMICMATERIALEDITOR_API void SetMaterialStageGradientClass(TSubclassOf<UDMMaterialStageGradient> InMaterialStageGradientClass);
 
 	UFUNCTION(BlueprintPure, Category = "Material Designer")
-	UDMMaterialStageGradient* GetMaterialStageGradient() const;
+	DYNAMICMATERIALEDITOR_API UDMMaterialStageGradient* GetMaterialStageGradient() const;
 
 protected:
 	static TArray<TStrongObjectPtr<UClass>> Gradients;
 
 	static void GenerateGradientList();
-
-	UDMMaterialStageInputGradient() = default;
 };
diff --git a/Engine/Plugins/Experimental/DynamicMaterial/Source/DynamicMaterialEditor/Public/Components/MaterialStageInputs/DMMSISlot.h b/Engine/Plugins/Experimental/DynamicMaterial/Source/DynamicMaterialEditor/Public/Components/MaterialStageInputs/DMMSISlot.h
--- a/Engine/Plugins/Experimental/DynamicMaterial/Source/DynamicMaterialEditor/Public/Components/MaterialStageInputs/DMMSISlot.h
+++ b/Engine/Plugins/Experimental/DynamicMaterial/Source/DynamicMaterialEditor/Public/Components/MaterialStageInputs/DMMSISlot.h
@@ -11,73 +11,86 @@ class UDMMaterialStage;
 class UDynamicMaterialModel;
 struct FDMMaterialBuildState;
 
-UCLASS(BlueprintType, ClassGroup = "Material Designer")
-class DYNAMICMATERIALEDITOR_API UDMMaterialStageInputSlot : public UDMMaterialStageInput
+UCLASS(MinimalAPI, BlueprintType, ClassGroup = "Material Designer")
+class UDMMaterialStageInputSlot : public UDMMaterialStageInput
 {
 	GENERATED_BODY()
 
 public:
-	static const FString SlotPathToken;
+	DYNAMICMATERIALEDITOR_API static const FString SlotPathToken;
 
-	static UDMMaterialStage* CreateStage(UDMMaterialSlot* InSourceSlot, EDMMaterialPropertyType InMaterialProperty, 
+	UFUNCTION(BlueprintCallable, Category = "Material Designer")
+	static DYNAMICMATERIALEDITOR_API UDMMaterialStage* CreateStage(UDMMaterialSlot* InSourceSlot, EDMMaterialPropertyType InMaterialProperty,
 		UDMMaterialLayerObject* InLayer = nullptr);
 
 	UFUNCTION(BlueprintCallable, Category = "Material Designer")
-	static UDMMaterialStageInputSlot* ChangeStageSource_Slot(UDMMaterialStage* InStage, UDMMaterialSlot* InSlot,
+	static DYNAMICMATERIALEDITOR_API UDMMaterialStageInputSlot* ChangeStageSource_Slot(UDMMaterialStage* InStage, UDMMaterialSlot* InSlot,
 		EDMMaterialPropertyType InProperty);
 
+	/**
+	 * Change the input type of an input on a stage to the output of another slot.
+	 * @param InInputIdx Index of the source input.
+	 * @param InInputChannel The channel of the input that the input connects to.
+	 * @param InProperty The property of the slot to use.
+	 * @param InOutputIdx The output index of the new input.
+	 * @param InOutputChannel The channel of the output to connect.
+	 */
 	UFUNCTION(BlueprintCallable, Category = "Material Designer")
-	static UDMMaterialStageInputSlot* ChangeStageInput_Slot(UDMMaterialStage* InStage, int32 InInputIdx, int32 InInputChannel, 
+	static DYNAMICMATERIALEDITOR_API UDMMaterialStageInputSlot* ChangeStageInput_Slot(UDMMaterialStage* InStage, int32 InInputIdx, int32 InInputChannel,
 		UDMMaterialSlot* InSlot, EDMMaterialPropertyType InProperty, int32 InOutputIdx, int32 InOutputChannel);
 
+	UDMMaterialStageInputSlot();
+
 	UFUNCTION(BlueprintPure, Category = "Material Designer")
 	UDMMaterialSlot* GetSlot() const { return Slot; }
 
 	UFUNCTION(BlueprintPure, Category = "Material Designer")
 	EDMMaterialPropertyType GetMaterialProperty() const { return MaterialProperty; }
 
-	virtual FText GetComponentDescription() const override;
-	virtual FText GetChannelDescription(const FDMMaterialStageConnectorChannel& Channel) override;
-
 	UFUNCTION(BlueprintCallable, Category = "Material Designer")
-	void SetSlot(UDMMaterialSlot* InSlot);
+	DYNAMICMATERIALEDITOR_API void SetSlot(UDMMaterialSlot* InSlot);
 
 	UFUNCTION(BlueprintCallable, Category = "Material Designer")
-	void SetMaterialProperty(EDMMaterialPropertyType InMaterialProperty);
+	DYNAMICMATERIALEDITOR_API void SetMaterialProperty(EDMMaterialPropertyType InMaterialProperty);
 
-	virtual void GenerateExpressions(const TSharedRef<FDMMaterialBuildState>& InBuildState) const override;
-
-	//~ Begin UObject
-	virtual void PostLoad() override;
-	virtual void PostEditImport() override;
-	//~ End UObject
+	//~ Start UDMMaterialStageInput
+	DYNAMICMATERIALEDITOR_API virtual FText GetComponentDescription() const override;
+	DYNAMICMATERIALEDITOR_API virtual FText GetChannelDescription(const FDMMaterialStageConnectorChannel& Channel) override;
+	//~ End UDMMaterialStageInput
 
 	//~ Start UDMMaterialComponent
-	virtual void PostEditorDuplicate(UDynamicMaterialModel* InMaterialModel, UDMMaterialComponent* InParent) override;
+	DYNAMICMATERIALEDITOR_API virtual void GenerateExpressions(const TSharedRef<FDMMaterialBuildState>& InBuildState) const override;
+	DYNAMICMATERIALEDITOR_API virtual void PostEditorDuplicate(UDynamicMaterialModel* InMaterialModel, UDMMaterialComponent* InParent) override;
 	//~ End UDMMaterialComponent
 
-protected:
-	UDMMaterialStageInputSlot();
-	virtual ~UDMMaterialStageInputSlot() override = default;
+	//~ Begin UObject
+	DYNAMICMATERIALEDITOR_API virtual void PostLoad() override;
+	DYNAMICMATERIALEDITOR_API virtual void PostEditImport() override;
+	//~ End UObject
 
+protected:
 	UPROPERTY(VisibleInstanceOnly, BlueprintReadOnly, Category = "Material Designer")
 	TObjectPtr<UDMMaterialSlot> Slot;
 
 	UPROPERTY(VisibleInstanceOnly, BlueprintReadOnly, Category = "Material Designer")
 	EDMMaterialPropertyType MaterialProperty;
 
+	virtual ~UDMMaterialStageInputSlot() override = default;
+
 	void OnSlotUpdated(UDMMaterialComponent* InComponent, EDMUpdateType InUpdateType);
 	void OnSlotConnectorsUpdated(UDMMaterialSlot* InSlot);
 	void OnSlotRemoved(UDMMaterialComponent* InComponent, EDMComponentLifetimeState InLifetimeState);
 	void OnParentSlotRemoved(UDMMaterialComponent* InComponent, EDMComponentLifetimeState InLifetimeState);
 
-	virtual void UpdateOutputConnectors() override;
-
 	void InitSlot();
 	void DeinitSlot();
 
+	//~ Start UDMMaterialStageInput
+	DYNAMICMATERIALEDITOR_API virtual void UpdateOutputConnectors() override;
+	//~ End UDMMaterialStageInput
+
 	//~ Begin UDMMaterialComponent
-	virtual void OnComponentRemoved() override;
-	virtual UDMMaterialComponent* GetSubComponentByPath(FDMComponentPath& InPath, const FDMComponentPathSegment& InPathSegment) const override;
+	DYNAMICMATERIALEDITOR_API virtual void OnComponentRemoved() override;
+	DYNAMICMATERIALEDITOR_API virtual UDMMaterialComponent* GetSubComponentByPath(FDMComponentPath& InPath, const FDMComponentPathSegment& InPathSegment) const override;
 	//~ End UDMMaterialComponent
 };
diff --git a/Engine/Plugins/Experimental/DynamicMaterial/Source/DynamicMaterialEditor/Public/Components/MaterialStageInputs/DMMSITextureUV.h b/Engine/Plugins/Experimental/DynamicMaterial/Source/DynamicMaterialEditor/Public/Components/MaterialStageInputs/DMMSITextureUV.h
--- a/Engine/Plugins/Experimental/DynamicMaterial/Source/DynamicMaterialEditor/Public/Components/MaterialStageInputs/DMMSITextureUV.h
+++ b/Engine/Plugins/Experimental/DynamicMaterial/Source/DynamicMaterialEditor/Public/Components/MaterialStageInputs/DMMSITextureUV.h
@@ -15,67 +15,74 @@ class UDMTextureUV;
 class UDynamicMaterialModel;
 class UMaterialExpressionScalarParameter;
 class UMaterialInstanceDynamic;
+enum class EDMMaterialParameterGroup : uint8;
 struct FDMMaterialBuildState;
 
-UCLASS(BlueprintType, ClassGroup = "Material Designer")
-class DYNAMICMATERIALEDITOR_API UDMMaterialStageInputTextureUV : public UDMMaterialStageInput
+UCLASS(MinimalAPI, BlueprintType, ClassGroup = "Material Designer")
+class UDMMaterialStageInputTextureUV : public UDMMaterialStageInput
 {
 	GENERATED_BODY()
 
 public:
-	static const FString TextureUVPathToken;
+	DYNAMICMATERIALEDITOR_API static const FString TextureUVPathToken;
 
-	static UDMMaterialStage* CreateStage(UDynamicMaterialModel* InMaterialModel, UDMMaterialLayerObject* InLayer = nullptr);
+	UFUNCTION(BlueprintCallable, Category = "Material Designer")
+	static DYNAMICMATERIALEDITOR_API UDMMaterialStage* CreateStage(UDynamicMaterialModel* InMaterialModel, UDMMaterialLayerObject* InLayer = nullptr);
 
 	UFUNCTION(BlueprintCallable, Category = "Material Designer")
-	static UDMMaterialStageInputTextureUV* ChangeStageSource_UV(UDMMaterialStage* InStage, bool bInDoUpdate);
+	static DYNAMICMATERIALEDITOR_API UDMMaterialStageInputTextureUV* ChangeStageSource_UV(UDMMaterialStage* InStage, bool bInDoUpdate);
 
 	UFUNCTION(BlueprintCallable, Category = "Material Designer")
-	static UDMMaterialStageInputTextureUV* ChangeStageInput_UV(UDMMaterialStage* InStage, int32 InInputIdx, int32 InInputChannel, 
-		int32 InOutputChannel);
+	static DYNAMICMATERIALEDITOR_API UDMMaterialStageInputTextureUV* ChangeStageInput_UV(UDMMaterialStage* InStage, int32 InInputIdx, 
+		int32 InInputChannel, int32 InOutputChannel);
 
-	virtual ~UDMMaterialStageInputTextureUV() override = default;
+	DYNAMICMATERIALEDITOR_API UDMMaterialStageInputTextureUV();
 
-	virtual FText GetComponentDescription() const override;
-	virtual FText GetChannelDescription(const FDMMaterialStageConnectorChannel& Channel) override;
+	virtual ~UDMMaterialStageInputTextureUV() override = default;
 
 	void Init(UDynamicMaterialModel* InMaterialModel);
 
 	UFUNCTION(BlueprintPure, Category = "Material Designer")
 	UDMTextureUV* GetTextureUV() { return TextureUV; }
 
-	virtual void GenerateExpressions(const TSharedRef<FDMMaterialBuildState>& InBuildState) const override;
-
-	//~ Begin UObject
-	virtual bool Modify(bool bInAlwaysMarkDirty = true) override;
-	virtual void PostLoad() override;
-	virtual void PostEditImport() override;
-	//~ End UObject
+	//~ Start UDMMaterialStageInput
+	DYNAMICMATERIALEDITOR_API virtual FText GetComponentDescription() const override;
+	DYNAMICMATERIALEDITOR_API virtual FText GetChannelDescription(const FDMMaterialStageConnectorChannel& Channel) override;
+	//~ End UDMMaterialStageInput
 
 	//~ Start UDMMaterialComponent
-	virtual void PostEditorDuplicate(UDynamicMaterialModel* InMaterialModel, UDMMaterialComponent* InParent) override;
+	DYNAMICMATERIALEDITOR_API virtual void GenerateExpressions(const TSharedRef<FDMMaterialBuildState>& InBuildState) const override;
+	DYNAMICMATERIALEDITOR_API virtual void PostEditorDuplicate(UDynamicMaterialModel* InMaterialModel, UDMMaterialComponent* InParent) override;
 	//~ End UDMMaterialComponent
 
+	//~ Begin UObject
+	DYNAMICMATERIALEDITOR_API virtual bool Modify(bool bInAlwaysMarkDirty = true) override;
+	DYNAMICMATERIALEDITOR_API virtual void PostLoad() override;
+	DYNAMICMATERIALEDITOR_API virtual void PostEditImport() override;
+	//~ End UObject
+
 protected:
-	static UMaterialExpressionScalarParameter* CreateScalarParameter(const TSharedRef<FDMMaterialBuildState>& InBuildState, FName InParamName, float InValue = 0.f);
+	DYNAMICMATERIALEDITOR_API static UMaterialExpressionScalarParameter* CreateScalarParameter(const TSharedRef<FDMMaterialBuildState>& InBuildState, FName InParamName,
+		EDMMaterialParameterGroup InParameterGroup, float InValue = 0.f);
+
 	static TArray<UMaterialExpression*> CreateTextureUVExpressions(const TSharedRef<FDMMaterialBuildState>& InBuildState, UDMTextureUV* InTextureUV);
 
 	UPROPERTY(EditInstanceOnly, BlueprintReadOnly, Instanced, Category = "Material Designer")
 	TObjectPtr<UDMTextureUV> TextureUV;
 
-	UDMMaterialStageInputTextureUV();
-
-	virtual void UpdateOutputConnectors() override;
-
 	void OnTextureUVUpdated(UDMMaterialComponent* InComponent, EDMUpdateType InUpdateType);
 
 	void InitTextureUV();
 
 	void AddEffects(const TSharedRef<FDMMaterialBuildState>& InBuildState, TArray<UMaterialExpression*>& InOutExpressions) const;
 
+	//~ Start UDMMaterialStageInput
+	DYNAMICMATERIALEDITOR_API virtual void UpdateOutputConnectors() override;
+	//~ End UDMMaterialStageInput
+
 	//~ Begin UDMMaterialComponent
-	virtual void OnComponentAdded() override;
-	virtual void OnComponentRemoved() override;
-	virtual UDMMaterialComponent* GetSubComponentByPath(FDMComponentPath& InPath, const FDMComponentPathSegment& InPathSegment) const override;
+	DYNAMICMATERIALEDITOR_API virtual void OnComponentAdded() override;
+	DYNAMICMATERIALEDITOR_API virtual void OnComponentRemoved() override;
+	DYNAMICMATERIALEDITOR_API virtual UDMMaterialComponent* GetSubComponentByPath(FDMComponentPath& InPath, const FDMComponentPathSegment& InPathSegment) const override;
 	//~ End UDMMaterialComponent
 };
diff --git a/Engine/Plugins/Experimental/DynamicMaterial/Source/DynamicMaterialEditor/Public/Components/MaterialStageInputs/DMMSIThroughput.h b/Engine/Plugins/Experimental/DynamicMaterial/Source/DynamicMaterialEditor/Public/Components/MaterialStageInputs/DMMSIThroughput.h
--- a/Engine/Plugins/Experimental/DynamicMaterial/Source/DynamicMaterialEditor/Public/Components/MaterialStageInputs/DMMSIThroughput.h
+++ b/Engine/Plugins/Experimental/DynamicMaterial/Source/DynamicMaterialEditor/Public/Components/MaterialStageInputs/DMMSIThroughput.h
@@ -11,48 +11,51 @@ class UDMMaterialSubStage;
 class UDynamicMaterialModel;
 struct FDMMaterialBuildState;
 
-UCLASS(Abstract, BlueprintType, ClassGroup = "Material Designer")
-class DYNAMICMATERIALEDITOR_API UDMMaterialStageInputThroughput : public UDMMaterialStageInput
+UCLASS(MinimalAPI, Abstract, BlueprintType, ClassGroup = "Material Designer")
+class UDMMaterialStageInputThroughput : public UDMMaterialStageInput
 {
 	GENERATED_BODY()
 
 public:
-	static const FString SubStagePathToken;
+	DYNAMICMATERIALEDITOR_API static const FString SubStagePathToken;
 
-	virtual FText GetComponentDescription() const override;
-	virtual FText GetChannelDescription(const FDMMaterialStageConnectorChannel& Channel) override;
+	DYNAMICMATERIALEDITOR_API UDMMaterialStageInputThroughput();
 
 	UFUNCTION(BlueprintPure, Category = "Material Designer")
-	TSubclassOf<UDMMaterialStageThroughput> GetMaterialStageThroughputClass() const;
+	DYNAMICMATERIALEDITOR_API TSubclassOf<UDMMaterialStageThroughput> GetMaterialStageThroughputClass() const;
 
 	UFUNCTION(BlueprintPure, Category = "Material Designer")
-	UDMMaterialStageThroughput* GetMaterialStageThroughput() const;
+	DYNAMICMATERIALEDITOR_API UDMMaterialStageThroughput* GetMaterialStageThroughput() const;
 
 	UFUNCTION(BlueprintPure, Category = "Material Designer")
 	UDMMaterialSubStage* GetSubStage() const { return SubStage; }
 
-	virtual void GenerateExpressions(const TSharedRef<FDMMaterialBuildState>& InBuildState) const override;
-	virtual int32 GetInnateMaskOutput(int32 OutputIndex, int32 OutputChannels) const;
-	virtual int32 GetOutputChannelOverride(int32 InOutputIndex) const override;
+	//~ Start UDMMaterialStageInput
+	DYNAMICMATERIALEDITOR_API virtual FText GetComponentDescription() const override;
+	DYNAMICMATERIALEDITOR_API virtual FText GetChannelDescription(const FDMMaterialStageConnectorChannel& Channel) override;
+	//~ End UDMMaterialStageInput
 
-	virtual bool IsPropertyVisible(FName Property) const override;
-
-	//~ Begin UObject
-	virtual bool Modify(bool bInAlwaysMarkDirty = true) override;
-	virtual void PostLoad() override;
-	virtual void PostEditImport() override;
-	//~ End UObject
+	//~ Start UDMMaterialStageSource
+	DYNAMICMATERIALEDITOR_API virtual void GenerateExpressions(const TSharedRef<FDMMaterialBuildState>& InBuildState) const override;
+	DYNAMICMATERIALEDITOR_API virtual int32 GetInnateMaskOutput(int32 OutputIndex, int32 OutputChannels) const;
+	DYNAMICMATERIALEDITOR_API virtual int32 GetOutputChannelOverride(int32 InOutputIndex) const override;
+	//~ End UDMMaterialStageSource
 
 	//~ Start UDMMaterialComponent
-	virtual void PostEditorDuplicate(UDynamicMaterialModel* InMaterialModel, UDMMaterialComponent* InParent) override;
+	DYNAMICMATERIALEDITOR_API virtual bool IsPropertyVisible(FName Property) const override;
+	DYNAMICMATERIALEDITOR_API virtual void PostEditorDuplicate(UDynamicMaterialModel* InMaterialModel, UDMMaterialComponent* InParent) override;
 	//~ End UDMMaterialComponent
 
+	//~ Begin UObject
+	DYNAMICMATERIALEDITOR_API virtual bool Modify(bool bInAlwaysMarkDirty = true) override;
+	DYNAMICMATERIALEDITOR_API virtual void PostLoad() override;
+	DYNAMICMATERIALEDITOR_API virtual void PostEditImport() override;
+	//~ End UObject
+
 protected:
 	UPROPERTY(VisibleInstanceOnly, BlueprintReadOnly, Instanced, Category = "Material Designer")
 	TObjectPtr<UDMMaterialSubStage> SubStage;
 
-	UDMMaterialStageInputThroughput();
-
 	void OnSubStageUpdated(UDMMaterialComponent* InComponent, EDMUpdateType InUpdateType);
 
 	void SetMaterialStageThroughputClass(TSubclassOf<UDMMaterialStageThroughput> InMaterialStageThroughputClass);
@@ -60,9 +63,9 @@ protected:
 	void InitSubStage();
 
 	//~ Begin UDMMaterialComponent
-	virtual void OnComponentAdded() override;
-	virtual void OnComponentRemoved() override;
-	virtual void GetComponentPathInternal(TArray<FString>& OutChildComponentPathComponents) const override;
-	virtual UDMMaterialComponent* GetSubComponentByPath(FDMComponentPath& InPath, const FDMComponentPathSegment& InPathSegment) const override;
+	DYNAMICMATERIALEDITOR_API virtual void OnComponentAdded() override;
+	DYNAMICMATERIALEDITOR_API virtual void OnComponentRemoved() override;
+	DYNAMICMATERIALEDITOR_API virtual void GetComponentPathInternal(TArray<FString>& OutChildComponentPathComponents) const override;
+	DYNAMICMATERIALEDITOR_API virtual UDMMaterialComponent* GetSubComponentByPath(FDMComponentPath& InPath, const FDMComponentPathSegment& InPathSegment) const override;
 	//~ End UDMMaterialComponent
 };
