Material Designer

diff --git a/Engine/Plugins/Experimental/DynamicMaterial/Source/DynamicMaterialEditor/Public/Components/MaterialStageInputs/DMMSIValue.h b/Engine/Plugins/Experimental/DynamicMaterial/Source/DynamicMaterialEditor/Public/Components/MaterialStageInputs/DMMSIValue.h
--- a/Engine/Plugins/Experimental/DynamicMaterial/Source/DynamicMaterialEditor/Public/Components/MaterialStageInputs/DMMSIValue.h
+++ b/Engine/Plugins/Experimental/DynamicMaterial/Source/DynamicMaterialEditor/Public/Components/MaterialStageInputs/DMMSIValue.h
@@ -10,95 +10,102 @@ class UDMMaterialLayerObject;
 class UDMMaterialStage;
 class UDynamicMaterialModel;
 
-UCLASS(BlueprintType, ClassGroup = "Material Designer")
-class DYNAMICMATERIALEDITOR_API UDMMaterialStageInputValue : public UDMMaterialStageInput
+UCLASS(MinimalAPI, BlueprintType, ClassGroup = "Material Designer")
+class UDMMaterialStageInputValue : public UDMMaterialStageInput
 {
 	GENERATED_BODY()
 		
 public:
-	static const FString ValuePathToken;
+	DYNAMICMATERIALEDITOR_API static const FString ValuePathToken;
 
 	static FName GetValuePropertyName() { return GET_MEMBER_NAME_CHECKED(UDMMaterialStageInputValue, Value); }
 
-	static UDMMaterialStage* CreateStage(UDMMaterialValue* InValue, UDMMaterialLayerObject* InLayer = nullptr);
+	UFUNCTION(BlueprintCallable, Category = "Material Designer")
+	static DYNAMICMATERIALEDITOR_API UDMMaterialStage* CreateStage(UDMMaterialValue* InValue, UDMMaterialLayerObject* InLayer = nullptr);
 
 	static UDMMaterialStageInputValue* ChangeStageSource_NewLocalValue(UDMMaterialStage* InStage, EDMValueType InValueType);
 
 	UFUNCTION(BlueprintCallable, Category = "Material Designer")
-	static UDMMaterialStageInputValue* ChangeStageSource_NewLocalValue(UDMMaterialStage* InStage, TSubclassOf<UDMMaterialValue> InValueClass);
+	static DYNAMICMATERIALEDITOR_API UDMMaterialStageInputValue* ChangeStageSource_NewLocalValue(UDMMaterialStage* InStage, TSubclassOf<UDMMaterialValue> InValueClass);
 
 	UFUNCTION(BlueprintCallable, Category = "Material Designer")
-	static UDMMaterialStageInputValue* ChangeStageSource_Value(UDMMaterialStage* InStage, UDMMaterialValue* InValue);
+	static DYNAMICMATERIALEDITOR_API UDMMaterialStageInputValue* ChangeStageSource_Value(UDMMaterialStage* InStage, UDMMaterialValue* InValue);
 
 	static UDMMaterialStageInputValue* ChangeStageSource_NewValue(UDMMaterialStage* InStage, EDMValueType InValueType);
 
 	UFUNCTION(BlueprintCallable, Category = "Material Designer")
-	static UDMMaterialStageInputValue* ChangeStageSource_NewValue(UDMMaterialStage* InStage, TSubclassOf<UDMMaterialValue> InValueClass);
+	static DYNAMICMATERIALEDITOR_API UDMMaterialStageInputValue* ChangeStageSource_NewValue(UDMMaterialStage* InStage, TSubclassOf<UDMMaterialValue> InValueClass);
 
-	static UDMMaterialStageInputValue* ChangeStageInput_NewLocalValue(UDMMaterialStage* InStage, int32 InInputIdx, int32 InInputChannel,
+	static DYNAMICMATERIALEDITOR_API UDMMaterialStageInputValue* ChangeStageInput_NewLocalValue(UDMMaterialStage* InStage, int32 InInputIdx, int32 InInputChannel,
 		EDMValueType InValueType, int32 InOutputChannel);
 
 	UFUNCTION(BlueprintCallable, Category = "Material Designer")
-	static UDMMaterialStageInputValue* ChangeStageInput_NewLocalValue(UDMMaterialStage* InStage, int32 InInputIdx, int32 InInputChannel,
+	static DYNAMICMATERIALEDITOR_API UDMMaterialStageInputValue* ChangeStageInput_NewLocalValue(UDMMaterialStage* InStage, int32 InInputIdx, int32 InInputChannel,
 		TSubclassOf<UDMMaterialValue> InValueClass, int32 InOutputChannel);
 
 	UFUNCTION(BlueprintCallable, Category = "Material Designer")
-	static UDMMaterialStageInputValue* ChangeStageInput_Value(UDMMaterialStage* InStage, int32 InInputIdx, int32 InInputChannel,
+	static DYNAMICMATERIALEDITOR_API UDMMaterialStageInputValue* ChangeStageInput_Value(UDMMaterialStage* InStage, int32 InInputIdx, int32 InInputChannel,
 		UDMMaterialValue* InValue, int32 InOutputChannel);
 
 	static UDMMaterialStageInputValue* ChangeStageInput_NewValue(UDMMaterialStage* InStage, int32 InInputIdx, int32 InInputChannel,
 		EDMValueType InValueType, int32 InOutputChannel);
 
 	UFUNCTION(BlueprintCallable, Category = "Material Designer")
-	static UDMMaterialStageInputValue* ChangeStageInput_NewValue(UDMMaterialStage* InStage, int32 InInputIdx, int32 InInputChannel,
+	static DYNAMICMATERIALEDITOR_API UDMMaterialStageInputValue* ChangeStageInput_NewValue(UDMMaterialStage* InStage, int32 InInputIdx, int32 InInputChannel,
 		TSubclassOf<UDMMaterialValue> InValueClass, int32 InOutputChannel);
 
-	virtual FText GetComponentDescription() const override;
-	virtual FText GetChannelDescription(const FDMMaterialStageConnectorChannel& Channel) override;
+	DYNAMICMATERIALEDITOR_API UDMMaterialStageInputValue();
+
+	virtual ~UDMMaterialStageInputValue() override = default;
 
 	UFUNCTION(BlueprintPure, Category = "Material Designer")
 	UDMMaterialValue* GetValue() const { return Value; }
 
 	UFUNCTION(BlueprintCallable, Category = "Material Designer")
-	void SetValue(UDMMaterialValue* InValue);
-
-	virtual void GenerateExpressions(const TSharedRef<FDMMaterialBuildState>& InBuildState) const override;
+	DYNAMICMATERIALEDITOR_API void SetValue(UDMMaterialValue* InValue);
 
-	virtual int32 GetInnateMaskOutput(int32 OutputIndex, int32 OutputChannels) const override;
+	DYNAMICMATERIALEDITOR_API void ApplyDefaultLayerSettings();
 
-	void ApplyDefaultLayerSettings();
+	//~ Start UDMMaterialStageInput
+	DYNAMICMATERIALEDITOR_API virtual FText GetChannelDescription(const FDMMaterialStageConnectorChannel& Channel) override;
+	//~ End UDMMaterialStageInput
 
-	//~ Begin UObject
-	virtual bool Modify(bool bInAlwaysMarkDirty = true) override;
-	virtual void PostLoad() override;
-	virtual void PostEditImport() override;
-	//~ End UObject
+	//~ Start UDMMaterialStageSource
+	DYNAMICMATERIALEDITOR_API virtual int32 GetInnateMaskOutput(int32 OutputIndex, int32 OutputChannels) const override;
+	//~ Start UDMMaterialStageSource
 
 	//~ Start UDMMaterialComponent
-	virtual void PostEditorDuplicate(UDynamicMaterialModel* InMaterialModel, UDMMaterialComponent* InParent) override;
+	DYNAMICMATERIALEDITOR_API virtual FText GetComponentDescription() const override;
+	DYNAMICMATERIALEDITOR_API virtual void GenerateExpressions(const TSharedRef<FDMMaterialBuildState>& InBuildState) const override;
+	DYNAMICMATERIALEDITOR_API virtual void PostEditorDuplicate(UDynamicMaterialModel* InMaterialModel, UDMMaterialComponent* InParent) override;
 	//~ End UDMMaterialComponent
 
-protected:
-	UDMMaterialStageInputValue();
-	virtual ~UDMMaterialStageInputValue() override = default;
+	//~ Begin UObject
+	DYNAMICMATERIALEDITOR_API virtual bool Modify(bool bInAlwaysMarkDirty = true) override;
+	DYNAMICMATERIALEDITOR_API virtual void PostLoad() override;
+	DYNAMICMATERIALEDITOR_API virtual void PostEditImport() override;
+	//~ End UObject
 
+protected:
 	UPROPERTY(VisibleInstanceOnly, BlueprintReadOnly, Instanced, Category = "Material Designer")
 	TObjectPtr<UDMMaterialValue> Value;
 
 	void OnValueUpdated(UDMMaterialComponent* InComponent, EDMUpdateType InUpdateType);
 
-	virtual void UpdateOutputConnectors() override;
-
 	void InitInputValue();
 	void DeinitInputValue();
 
-	//~ Begin UDMMaterialStageSource
-	virtual void OnComponentAdded() override;
-	virtual void OnComponentRemoved() override;
-	virtual UDMMaterialComponent* GetSubComponentByPath(FDMComponentPath& InPath, const FDMComponentPathSegment& InPathSegment) const override;
-	//~ End UDMMaterialStageSource
-
 	bool IsSharedStageValue() const;
 
 	void ApplyWholeLayerValue();
+
+	//~ Start UDMMaterialStageInput
+	DYNAMICMATERIALEDITOR_API virtual void UpdateOutputConnectors() override;
+	//~ End UDMMaterialStageInput
+
+	//~ Begin UDMMaterialStageSource
+	DYNAMICMATERIALEDITOR_API virtual void OnComponentAdded() override;
+	DYNAMICMATERIALEDITOR_API virtual void OnComponentRemoved() override;
+	DYNAMICMATERIALEDITOR_API virtual UDMMaterialComponent* GetSubComponentByPath(FDMComponentPath& InPath, const FDMComponentPathSegment& InPathSegment) const override;
+	//~ End UDMMaterialStageSource
 };
diff --git a/Engine/Plugins/Experimental/DynamicMaterial/Source/DynamicMaterialEditor/Public/DMEDefs.h b/Engine/Plugins/Experimental/DynamicMaterial/Source/DynamicMaterialEditor/Public/DMEDefs.h
--- a/Engine/Plugins/Experimental/DynamicMaterial/Source/DynamicMaterialEditor/Public/DMEDefs.h
+++ b/Engine/Plugins/Experimental/DynamicMaterial/Source/DynamicMaterialEditor/Public/DMEDefs.h
@@ -130,6 +130,7 @@ struct FDMPropertyHandle
 	FName ValueName = NAME_None;
 	FName CategoryOverrideName = NAME_None;
 	TOptional<float> MaxWidth;
+	bool bEnabled = true;
 };
 
 UENUM(BlueprintType, meta = (Bitflags, UseEnumValuesAsMaskValuesInEditor = "true"))
diff --git a/Engine/Plugins/Experimental/DynamicMaterial/Source/DynamicMaterialEditor/Public/DMObjectMaterialProperty.h b/Engine/Plugins/Experimental/DynamicMaterial/Source/DynamicMaterialEditor/Public/DMObjectMaterialProperty.h
--- a/Engine/Plugins/Experimental/DynamicMaterial/Source/DynamicMaterialEditor/Public/DMObjectMaterialProperty.h
+++ b/Engine/Plugins/Experimental/DynamicMaterial/Source/DynamicMaterialEditor/Public/DMObjectMaterialProperty.h
@@ -11,6 +11,7 @@
 class FProperty;
 class UDynamicMaterialInstance;
 class UDynamicMaterialModel;
+class UDynamicMaterialModelBase;
 class UObject;
 class UPrimitiveComponent;
 
@@ -18,17 +19,17 @@ class UPrimitiveComponent;
  * Defines a material property slot that can be a Material Designer Instance.
  */
 USTRUCT(BlueprintType)
-struct DYNAMICMATERIALEDITOR_API FDMObjectMaterialProperty
+struct FDMObjectMaterialProperty
 {
 	GENERATED_BODY()
 
 	FDMObjectMaterialProperty();
 
-	// UPrimitiveComponent Material Index
-	FDMObjectMaterialProperty(UPrimitiveComponent* InOuter, int32 InIndex);
+	/** UPrimitiveComponent Material Index */
+	DYNAMICMATERIALEDITOR_API FDMObjectMaterialProperty(UPrimitiveComponent* InOuter, int32 InIndex);
 
-	// Class Property (including potential array index)
-	FDMObjectMaterialProperty(UObject* InOuter, FProperty* InProperty, int32 InIndex = INDEX_NONE);
+	/** Class Property (including potential array index) */
+	DYNAMICMATERIALEDITOR_API FDMObjectMaterialProperty(UObject* InOuter, FProperty* InProperty, int32 InIndex = INDEX_NONE);
 
 	UPROPERTY(VisibleInstanceOnly, BlueprintReadOnly, Category = "Material")
 	TWeakObjectPtr<UObject> OuterWeak;
@@ -44,16 +45,17 @@ struct DYNAMICMATERIALEDITOR_API FDMObjectMaterialProperty
 	UPROPERTY(VisibleInstanceOnly, BlueprintReadOnly, Category = "Material")
 	int32 Index;
 
-	UDynamicMaterialModel* GetMaterialModel() const;
+	DYNAMICMATERIALEDITOR_API UDynamicMaterialModelBase* GetMaterialModelBase() const;
 
-	UDynamicMaterialInstance* GetMaterial() const;
-	void SetMaterial(UDynamicMaterialInstance* DynamicMaterial);
+	DYNAMICMATERIALEDITOR_API UDynamicMaterialInstance* GetMaterial() const;
 
-	bool IsValid() const;
+	DYNAMICMATERIALEDITOR_API void SetMaterial(UDynamicMaterialInstance* DynamicMaterial);
 
-	FText GetPropertyName(bool bInIgnoreNewStatus) const;
+	DYNAMICMATERIALEDITOR_API bool IsValid() const;
 
-	void Reset();
+	DYNAMICMATERIALEDITOR_API FText GetPropertyName(bool bInIgnoreNewStatus) const;
+
+	DYNAMICMATERIALEDITOR_API void Reset();
 
 	template<typename InClass>
 	InClass* GetTypedOuter() const
diff --git a/Engine/Plugins/Experimental/DynamicMaterial/Source/DynamicMaterialEditor/Public/DMWorldSubsystem.h b/Engine/Plugins/Experimental/DynamicMaterial/Source/DynamicMaterialEditor/Public/DMWorldSubsystem.h
--- a/Engine/Plugins/Experimental/DynamicMaterial/Source/DynamicMaterialEditor/Public/DMWorldSubsystem.h
+++ b/Engine/Plugins/Experimental/DynamicMaterial/Source/DynamicMaterialEditor/Public/DMWorldSubsystem.h
@@ -12,14 +12,14 @@
 class AActor;
 class IDetailKeyframeHandler;
 class UDynamicMaterialInstance;
-class UDynamicMaterialModel;
+class UDynamicMaterialModelBase;
 struct FDMObjectMaterialProperty;
 
-DECLARE_DELEGATE_RetVal(UDynamicMaterialModel*, FDMGetMaterialModelDelegate)
-DECLARE_DELEGATE_OneParam(FDMSetMaterialModelDelegate, UDynamicMaterialModel*)
+DECLARE_DELEGATE_RetVal(UDynamicMaterialModelBase*, FDMGetMaterialModelDelegate)
+DECLARE_DELEGATE_OneParam(FDMSetMaterialModelDelegate, UDynamicMaterialModelBase*)
 DECLARE_DELEGATE_OneParam(FDMSetMaterialObjectPropertyDelegate, const FDMObjectMaterialProperty&)
 DECLARE_DELEGATE_OneParam(FDMSetMaterialActorDelegate, AActor*)
-DECLARE_DELEGATE_RetVal_OneParam(bool, FDMIsValidDelegate, UDynamicMaterialModel*)
+DECLARE_DELEGATE_RetVal_OneParam(bool, FDMIsValidDelegate, UDynamicMaterialModelBase*)
 DECLARE_DELEGATE_RetVal_TwoParams(bool, FDMSetMaterialValueDelegate, const FDMObjectMaterialProperty&, UDynamicMaterialInstance*)
 DECLARE_DELEGATE(FDMInvokeTabDelegate)
 
@@ -47,7 +47,7 @@ public:
 		return CustomModelEditorGetDelegate;
 	}
 
-	UDynamicMaterialModel* ExecuteGetCustomEditorModelDelegate();
+	UDynamicMaterialModelBase* ExecuteGetCustomEditorModelDelegate();
 
 	/** Sets material in a custom editor tab. */
 	FDMSetMaterialModelDelegate::RegistrationType& GetSetCustomEditorModelDelegate() const
@@ -55,7 +55,7 @@ public:
 		return CustomModelEditorSetDelegate;
 	}
 
-	void ExecuteSetCustomEditorModelDelegate(UDynamicMaterialModel* InMaterialModel);
+	void ExecuteSetCustomEditorModelDelegate(UDynamicMaterialModelBase* InMaterialModel);
 
 	/** Sets the object property in custom editor tab. */
 	FDMSetMaterialObjectPropertyDelegate::RegistrationType& GetCustomObjectPropertyEditorDelegate() const
@@ -79,7 +79,7 @@ public:
 		return IsValidDelegate;
 	}
 
-	bool ExecuteIsValidDelegate(UDynamicMaterialModel* InMaterialModel);
+	bool ExecuteIsValidDelegate(UDynamicMaterialModelBase* InMaterialModel);
 
 	/** Used to redirect SetMaterial to different objects/paths. */
 	FDMSetMaterialValueDelegate::RegistrationType& GetMaterialValueSetterDelegate() const
diff --git a/Engine/Plugins/Experimental/DynamicMaterial/Source/DynamicMaterialEditor/Public/IDynamicMaterialEditorModule.h b/Engine/Plugins/Experimental/DynamicMaterial/Source/DynamicMaterialEditor/Public/IDynamicMaterialEditorModule.h
--- a/Engine/Plugins/Experimental/DynamicMaterial/Source/DynamicMaterialEditor/Public/IDynamicMaterialEditorModule.h
+++ b/Engine/Plugins/Experimental/DynamicMaterial/Source/DynamicMaterialEditor/Public/IDynamicMaterialEditorModule.h
@@ -11,7 +11,7 @@
 class AActor;
 class UClass;
 class UDynamicMaterialInstance;
-class UDynamicMaterialModel;
+class UDynamicMaterialModelBase;
 class UObject;
 struct FDMObjectMaterialProperty;
 struct IDMOnWizardCompleteCallback;
@@ -52,9 +52,9 @@ public:
 
 	virtual void OpenEditor(UWorld* InWorld) const = 0;
 
-	virtual UDynamicMaterialModel* GetOpenedMaterialModel(UWorld* InWorld) const = 0;
+	virtual UDynamicMaterialModelBase* GetOpenedMaterialModel(UWorld* InWorld) const = 0;
 
-	virtual void OpenMaterialModel(UDynamicMaterialModel* InMaterialModel, UWorld* InWorld, bool bInInvokeTab) const = 0;
+	virtual void OpenMaterialModel(UDynamicMaterialModelBase* InMaterialModel, UWorld* InWorld, bool bInInvokeTab) const = 0;
 
 	virtual void OpenMaterialObjectProperty(const FDMObjectMaterialProperty& InObjectProperty, UWorld* InWorld, bool bInInvokeTab) const = 0;
 
diff --git a/Engine/Plugins/Experimental/DynamicMaterial/Source/DynamicMaterialEditor/Public/Material/DynamicMaterialInstanceFactory.h b/Engine/Plugins/Experimental/DynamicMaterial/Source/DynamicMaterialEditor/Public/Material/DynamicMaterialInstanceFactory.h
--- a/Engine/Plugins/Experimental/DynamicMaterial/Source/DynamicMaterialEditor/Public/Material/DynamicMaterialInstanceFactory.h
+++ b/Engine/Plugins/Experimental/DynamicMaterial/Source/DynamicMaterialEditor/Public/Material/DynamicMaterialInstanceFactory.h
@@ -11,11 +11,11 @@ class UDynamicMaterialInstanceFactory : public UFactory
 	GENERATED_BODY()
 
 public:
-	UDynamicMaterialInstanceFactory();
+	DYNAMICMATERIALEDITOR_API UDynamicMaterialInstanceFactory();
 
 	//~ Begin UFactory
-	virtual UObject* FactoryCreateNew(UClass* Class, UObject* InParent, FName Name, EObjectFlags Flags, UObject* Context, FFeedbackContext* Warn) override;
-	virtual FText GetDisplayName() const override;
-	virtual FText GetToolTip() const override;
+	DYNAMICMATERIALEDITOR_API virtual UObject* FactoryCreateNew(UClass* Class, UObject* InParent, FName Name, EObjectFlags Flags, UObject* Context, FFeedbackContext* Warn) override;
+	DYNAMICMATERIALEDITOR_API virtual FText GetDisplayName() const override;
+	DYNAMICMATERIALEDITOR_API virtual FText GetToolTip() const override;
 	//~ End UFactory
 };
diff --git a/Engine/Plugins/Experimental/DynamicMaterial/Source/DynamicMaterialEditor/Public/Model/DMMaterialBuildState.h b/Engine/Plugins/Experimental/DynamicMaterial/Source/DynamicMaterialEditor/Public/Model/DMMaterialBuildState.h
--- a/Engine/Plugins/Experimental/DynamicMaterial/Source/DynamicMaterialEditor/Public/Model/DMMaterialBuildState.h
+++ b/Engine/Plugins/Experimental/DynamicMaterial/Source/DynamicMaterialEditor/Public/Model/DMMaterialBuildState.h
@@ -30,129 +30,127 @@ struct FDMMaterialStageConnection;
  * It is an entirely transient object. It is not meant to be saved outside of the material building processs.
  * It also provides some helper functions for creating UMaterialExpressions.
  */
-struct DYNAMICMATERIALEDITOR_API FDMMaterialBuildState : public IDMMaterialBuildStateInterface, public TSharedFromThis<FDMMaterialBuildState>
+struct FDMMaterialBuildState : public IDMMaterialBuildStateInterface, public TSharedFromThis<FDMMaterialBuildState>
 {
-	FDMMaterialBuildState(UMaterial* InDynamicMaterial, UDynamicMaterialModel* InMaterialModel, bool bInDirtyAssets = true);
+	DYNAMICMATERIALEDITOR_API FDMMaterialBuildState(UMaterial* InDynamicMaterial, UDynamicMaterialModel* InMaterialModel, bool bInDirtyAssets = true);
 
 	virtual ~FDMMaterialBuildState() override;
 
-	virtual UMaterial* GetDynamicMaterial() const override;
+	DYNAMICMATERIALEDITOR_API virtual UMaterial* GetDynamicMaterial() const override;
 
-	virtual UDynamicMaterialModel* GetMaterialModel() const override;
+	DYNAMICMATERIALEDITOR_API virtual UDynamicMaterialModel* GetMaterialModel() const override;
 
-	const UDMMaterialProperty* GetCurrentMaterialProperty() const;
+	/**
+	 * The current material property being generated for.
+	 * May be null in the case of global values or global parameters.
+	 */
+	DYNAMICMATERIALEDITOR_API const UDMMaterialProperty* GetCurrentMaterialProperty() const;
 
-	void SetCurrentMaterialProperty(const UDMMaterialProperty* InProperty);
+	DYNAMICMATERIALEDITOR_API void SetCurrentMaterialProperty(const UDMMaterialProperty* InProperty);
 
+	/** Whether assets can potentially be dirtied by the build process. */
 	bool ShouldDirtyAssets() const { return bDirtyAssets; }
 
-	void SetIgnoreUVs();
-
+	/** If ignoring UVs is on, UV nodes will not be processed. Useful for preview materials. */
 	bool IsIgnoringUVs() const { return bIgnoreUVs; }
 
-	void SetPreviewMaterial();
+	DYNAMICMATERIALEDITOR_API void SetIgnoreUVs();
 
+	/** Whether the build state is building a full material or a preview. */
 	bool IsPreviewMaterial() const { return bIsPreviewMaterial; }
 
-	virtual IDMMaterialBuildUtilsInterface& GetBuildUtils() const override;
+	DYNAMICMATERIALEDITOR_API void SetPreviewMaterial();
+
+	/** A handy set of tools for creating material expressions. */
+	DYNAMICMATERIALEDITOR_API virtual IDMMaterialBuildUtilsInterface& GetBuildUtils() const override;
 
-	FExpressionInput* GetMaterialProperty(EDMMaterialPropertyType InProperty) const;
+	/** Returns the FExpressionInput for the given material property on the material's editor only data. */
+	DYNAMICMATERIALEDITOR_API FExpressionInput* GetMaterialProperty(EDMMaterialPropertyType InProperty) const;
 
 	/** Slots */
-	bool HasSlot(const UDMMaterialSlot* InSlot) const;
+	DYNAMICMATERIALEDITOR_API bool HasSlot(const UDMMaterialSlot* InSlot) const;
 
-	const TArray<UMaterialExpression*>& GetSlotExpressions(const UDMMaterialSlot* InSlot) const;
+	DYNAMICMATERIALEDITOR_API const TArray<UMaterialExpression*>& GetSlotExpressions(const UDMMaterialSlot* InSlot) const;
 
-	UMaterialExpression* GetLastSlotExpression(const UDMMaterialSlot* InSlot) const;
+	DYNAMICMATERIALEDITOR_API UMaterialExpression* GetLastSlotExpression(const UDMMaterialSlot* InSlot) const;
 
-	void AddSlotExpressions(const UDMMaterialSlot* InSlot, const TArray<UMaterialExpression*>& InSlotExpressions);
+	DYNAMICMATERIALEDITOR_API void AddSlotExpressions(const UDMMaterialSlot* InSlot, const TArray<UMaterialExpression*>& InSlotExpressions);
 
-	bool HasSlotProperties(const UDMMaterialSlot* InSlot) const;
+	DYNAMICMATERIALEDITOR_API bool HasSlotProperties(const UDMMaterialSlot* InSlot) const;
 
-	void AddSlotPropertyExpressions(const UDMMaterialSlot* InSlot, const TMap<EDMMaterialPropertyType, 
+	DYNAMICMATERIALEDITOR_API void AddSlotPropertyExpressions(const UDMMaterialSlot* InSlot, const TMap<EDMMaterialPropertyType,
 		TArray<UMaterialExpression*>>& InSlotPropertyExpressions);
 
-	const TMap<EDMMaterialPropertyType, TArray<UMaterialExpression*>>& GetSlotPropertyExpressions(const UDMMaterialSlot* InSlot);
+	DYNAMICMATERIALEDITOR_API const TMap<EDMMaterialPropertyType, TArray<UMaterialExpression*>>& GetSlotPropertyExpressions(const UDMMaterialSlot* InSlot);
 
-	UMaterialExpression* GetLastSlotPropertyExpression(const UDMMaterialSlot* InSlot, EDMMaterialPropertyType InMaterialProperty) const;
+	DYNAMICMATERIALEDITOR_API UMaterialExpression* GetLastSlotPropertyExpression(const UDMMaterialSlot* InSlot, EDMMaterialPropertyType InMaterialProperty) const;
 
-	TArray<const UDMMaterialSlot*> GetSlots() const;
+	DYNAMICMATERIALEDITOR_API TArray<const UDMMaterialSlot*> GetSlots() const;
 
-	const TMap<const UDMMaterialSlot*, TArray<UMaterialExpression*>>& GetSlotMap() const;
+	DYNAMICMATERIALEDITOR_API const TMap<const UDMMaterialSlot*, TArray<UMaterialExpression*>>& GetSlotMap() const;
 
 	/** Layers */
-	bool HasLayer(const UDMMaterialLayerObject* InLayer) const;
+	DYNAMICMATERIALEDITOR_API bool HasLayer(const UDMMaterialLayerObject* InLayer) const;
 
-	const TArray<UMaterialExpression*>& GetLayerExpressions(const UDMMaterialLayerObject* InLayer) const;
+	DYNAMICMATERIALEDITOR_API const TArray<UMaterialExpression*>& GetLayerExpressions(const UDMMaterialLayerObject* InLayer) const;
 
-	UMaterialExpression* GetLastLayerExpression(const UDMMaterialLayerObject* InLayer) const;
+	DYNAMICMATERIALEDITOR_API UMaterialExpression* GetLastLayerExpression(const UDMMaterialLayerObject* InLayer) const;
 
-	void AddLayerExpressions(const UDMMaterialLayerObject* InLayer, const TArray<UMaterialExpression*>& InLayerExpressions);
+	DYNAMICMATERIALEDITOR_API void AddLayerExpressions(const UDMMaterialLayerObject* InLayer, const TArray<UMaterialExpression*>& InLayerExpressions);
 
-	TArray<const UDMMaterialLayerObject*> GetLayers() const;
+	DYNAMICMATERIALEDITOR_API TArray<const UDMMaterialLayerObject*> GetLayers() const;
 
-	const TMap<const UDMMaterialLayerObject*, TArray<UMaterialExpression*>>& GetLayerMap() const;
+	DYNAMICMATERIALEDITOR_API const TMap<const UDMMaterialLayerObject*, TArray<UMaterialExpression*>>& GetLayerMap() const;
 
 	/** Stages */
-	bool HasStage(const UDMMaterialStage* InStage) const;
+	DYNAMICMATERIALEDITOR_API bool HasStage(const UDMMaterialStage* InStage) const;
 
-	const TArray<UMaterialExpression*>& GetStageExpressions(const UDMMaterialStage* InStage) const;
+	DYNAMICMATERIALEDITOR_API const TArray<UMaterialExpression*>& GetStageExpressions(const UDMMaterialStage* InStage) const;
 
-	UMaterialExpression* GetLastStageExpression(const UDMMaterialStage* InStage) const;
+	DYNAMICMATERIALEDITOR_API UMaterialExpression* GetLastStageExpression(const UDMMaterialStage* InStage) const;
 
-	void AddStageExpressions(const UDMMaterialStage* InStage, const TArray<UMaterialExpression*>& InStageExpressions);
+	DYNAMICMATERIALEDITOR_API void AddStageExpressions(const UDMMaterialStage* InStage, const TArray<UMaterialExpression*>& InStageExpressions);
 
-	TArray<const UDMMaterialStage*> GetStages() const;
+	DYNAMICMATERIALEDITOR_API TArray<const UDMMaterialStage*> GetStages() const;
 
-	const TMap<const UDMMaterialStage*, TArray<UMaterialExpression*>>& GetStageMap() const;
+	DYNAMICMATERIALEDITOR_API const TMap<const UDMMaterialStage*, TArray<UMaterialExpression*>>& GetStageMap() const;
 
 	/** Stage Sources */
-	bool HasStageSource(const UDMMaterialStageSource* InStageSource) const;
+	DYNAMICMATERIALEDITOR_API bool HasStageSource(const UDMMaterialStageSource* InStageSource) const;
 
-	const TArray<UMaterialExpression*>& GetStageSourceExpressions(const UDMMaterialStageSource* InStageSource) const;
+	DYNAMICMATERIALEDITOR_API const TArray<UMaterialExpression*>& GetStageSourceExpressions(const UDMMaterialStageSource* InStageSource) const;
 
-	UMaterialExpression* GetLastStageSourceExpression(const UDMMaterialStageSource* InStageSource) const;
+	DYNAMICMATERIALEDITOR_API UMaterialExpression* GetLastStageSourceExpression(const UDMMaterialStageSource* InStageSource) const;
 
-	void AddStageSourceExpressions(const UDMMaterialStageSource* InStageSource, 
+	DYNAMICMATERIALEDITOR_API void AddStageSourceExpressions(const UDMMaterialStageSource* InStageSource,
 		const TArray<UMaterialExpression*>& InStageSourceExpressions);
 
-	TArray<const UDMMaterialStageSource*> GetStageSources() const;
+	DYNAMICMATERIALEDITOR_API TArray<const UDMMaterialStageSource*> GetStageSources() const;
 
-	const TMap<const UDMMaterialStageSource*, TArray<UMaterialExpression*>>& GetStageSourceMap() const;
+	DYNAMICMATERIALEDITOR_API const TMap<const UDMMaterialStageSource*, TArray<UMaterialExpression*>>& GetStageSourceMap() const;
 
 	/** Material Values */
-	virtual bool HasValue(const UDMMaterialValue* InValue) const override;
-
-	virtual const TArray<UMaterialExpression*>& GetValueExpressions(const UDMMaterialValue* InValue) const override;
-
-	virtual UMaterialExpression* GetLastValueExpression(const UDMMaterialValue* InValue) const override;
-
-	virtual void AddValueExpressions(const UDMMaterialValue* InValue, const TArray<UMaterialExpression*>& InValueExpressions) override;
-
-	virtual TArray<const UDMMaterialValue*> GetValues() const override;
-
-	virtual const TMap<const UDMMaterialValue*, TArray<UMaterialExpression*>>& GetValueMap() const override;
+	DYNAMICMATERIALEDITOR_API virtual bool HasValue(const UDMMaterialValue* InValue) const override;
 
-	/** Callbacks */
-	bool HasCallback(const UDMMaterialStageSource* InCallback) const;
+	DYNAMICMATERIALEDITOR_API virtual const TArray<UMaterialExpression*>& GetValueExpressions(const UDMMaterialValue* InValue) const override;
 
-	const TArray<UDMMaterialStageSource*>& GetCallbackExpressions(const UDMMaterialStageSource* InCallback) const;
+	DYNAMICMATERIALEDITOR_API virtual UMaterialExpression* GetLastValueExpression(const UDMMaterialValue* InValue) const override;
 
-	void AddCallbackExpressions(const UDMMaterialStageSource* InCallback, const TArray<UDMMaterialStageSource*>& InCallbackExpressions);
+	DYNAMICMATERIALEDITOR_API virtual void AddValueExpressions(const UDMMaterialValue* InValue, const TArray<UMaterialExpression*>& InValueExpressions) override;
 
-	TArray<const UDMMaterialStageSource*> GetCallbacks() const;
+	DYNAMICMATERIALEDITOR_API virtual TArray<const UDMMaterialValue*> GetValues() const override;
 
-	const TMap<const UDMMaterialStageSource*, TArray<UDMMaterialStageSource*>>& GetCallbackMap() const;
+	DYNAMICMATERIALEDITOR_API virtual const TMap<const UDMMaterialValue*, TArray<UMaterialExpression*>>& GetValueMap() const override;
 
 	/** Other */
-	virtual void AddOtherExpressions(const TArray<UMaterialExpression*>& InOtherExpressions) override;
+	DYNAMICMATERIALEDITOR_API virtual void AddOtherExpressions(const TArray<UMaterialExpression*>& InOtherExpressions) override;
 
-	virtual const TSet<UMaterialExpression*>& GetOtherExpressions() override;
+	DYNAMICMATERIALEDITOR_API virtual const TSet<UMaterialExpression*>& GetOtherExpressions() override;
 
 	/** Global Expressions */
-	UMaterialExpression* GetGlobalExpression(FName InName) const;
+	DYNAMICMATERIALEDITOR_API UMaterialExpression* GetGlobalExpression(FName InName) const;
 
-	void SetGlobalExpression(FName InName, UMaterialExpression* InExpression);
+	DYNAMICMATERIALEDITOR_API void SetGlobalExpression(FName InName, UMaterialExpression* InExpression);
 
 private:
 	UMaterial* DynamicMaterial = nullptr;
@@ -169,7 +167,6 @@ private:
 	TMap<const UDMMaterialLayerObject*, TArray<UMaterialExpression*>> Layers;
 	TMap<const UDMMaterialStage*, TArray<UMaterialExpression*>> Stages;
 	TMap<const UDMMaterialStageSource*, TArray<UMaterialExpression*>> StageSources;
-	TMap<const UDMMaterialStageSource*, TArray<UDMMaterialStageSource*>> Callbacks;
 	TSet<UMaterialExpression*> OtherExpressions;
 	TMap<FName, UMaterialExpression*> GlobalExpressions;
 };
diff --git a/Engine/Plugins/Experimental/DynamicMaterial/Source/DynamicMaterialEditor/Public/Model/DMMaterialBuildUtils.h b/Engine/Plugins/Experimental/DynamicMaterial/Source/DynamicMaterialEditor/Public/Model/DMMaterialBuildUtils.h
--- a/Engine/Plugins/Experimental/DynamicMaterial/Source/DynamicMaterialEditor/Public/Model/DMMaterialBuildUtils.h
+++ b/Engine/Plugins/Experimental/DynamicMaterial/Source/DynamicMaterialEditor/Public/Model/DMMaterialBuildUtils.h
@@ -6,18 +6,22 @@
 
 struct FDMMaterialBuildState;
 
-struct DYNAMICMATERIALEDITOR_API FDMMaterialBuildUtils : public IDMMaterialBuildUtilsInterface
+/** Utilities for creating material expressions. */
+struct FDMMaterialBuildUtils : public IDMMaterialBuildUtilsInterface
 {
-	FDMMaterialBuildUtils(FDMMaterialBuildState& InBuildState);
+	DYNAMICMATERIALEDITOR_API FDMMaterialBuildUtils(FDMMaterialBuildState& InBuildState);
 
-	UMaterialExpression* CreateDefaultExpression() const;
+	/** Creates a default expression outputing 0 on a single channel. */
+	DYNAMICMATERIALEDITOR_API UMaterialExpression* CreateDefaultExpression() const;
 
-	virtual UMaterialExpression* CreateExpression(TSubclassOf<UMaterialExpression> InExpressionClass, const FString& InComment,
+	/**
+	 * Create an expression with the comment as its description.
+	 * @param InAsset For things like texture nodes, will set the default value of that asset property.
+	 */
+	DYNAMICMATERIALEDITOR_API virtual UMaterialExpression* CreateExpression(TSubclassOf<UMaterialExpression> InExpressionClass, const FString& InComment,
 		UObject* InAsset = nullptr) const override;
 
-	virtual UMaterialExpression* CreateExpressionParameter(TSubclassOf<UMaterialExpression> InExpressionClass, FName InParameterName,
-		const FString& InComment, UObject* InAsset = nullptr) const override;
-
+	/** @See CreateExpression */
 	template<typename InExpressionClass
 		UE_REQUIRES(std::derived_from<InExpressionClass, UMaterialExpression>)>
 	InExpressionClass* CreateExpression(const FString& InComment, UObject* InAsset = nullptr) const
@@ -25,26 +29,41 @@ struct DYNAMICMATERIALEDITOR_API FDMMaterialBuildUtils : public IDMMaterialBuild
 		return Cast<InExpressionClass>(CreateExpression(TSubclassOf<UMaterialExpression>(InExpressionClass::StaticClass()), InComment, InAsset));
 	}
 
+	/**
+	 * Create a parameter expression with the comment as its description.
+	 * @param InParameterName The name of the parameter exposed in the material.
+	 * @param InParameterGroup Determines the type of group assigned to the parameter.
+	 * @param InAsset For things like texture nodes, will set the default value of that asset property.
+	 */
+	DYNAMICMATERIALEDITOR_API virtual UMaterialExpression* CreateExpressionParameter(TSubclassOf<UMaterialExpression> InExpressionClass, FName InParameterName,
+		EDMMaterialParameterGroup InParameterGroup, const FString& InComment, UObject* InAsset = nullptr) const override;
+
+	/** @See CreateExpressionParameter */
 	template<typename InExpressionClass
 		UE_REQUIRES(std::derived_from<InExpressionClass, UMaterialExpression>)>
-	InExpressionClass* CreateExpressionParameter(FName InParameterName, const FString& InComment, UObject* InAsset = nullptr) const
+	InExpressionClass* CreateExpressionParameter(FName InParameterName, EDMMaterialParameterGroup InParameterGroup, const FString& InComment, UObject* InAsset = nullptr) const
 	{
-		return Cast<InExpressionClass>(CreateExpressionParameter(TSubclassOf<UMaterialExpression>(InExpressionClass::StaticClass()), InParameterName, InComment, InAsset));
+		return Cast<InExpressionClass>(CreateExpressionParameter(TSubclassOf<UMaterialExpression>(InExpressionClass::StaticClass()), InParameterName, InParameterGroup, InComment, InAsset));
 	}
 
-	virtual TArray<UMaterialExpression*> CreateExpressionInputs(const TArray<FDMMaterialStageConnection>& InInputConnectionMap,
+	/** Creates a series of nodes that try to render every single input on different parts of the material. */
+	DYNAMICMATERIALEDITOR_API virtual TArray<UMaterialExpression*> CreateExpressionInputs(const TArray<FDMMaterialStageConnection>& InInputConnectionMap,
 		int32 InStageSourceInputIdx, const TArray<UDMMaterialStageInput*>& InStageInputs, int32& OutOutputIndex, 
 		int32& OutOutputChannel) const override;
 
-	virtual TArray<UMaterialExpression*> CreateExpressionInput(UDMMaterialStageInput* InInput) const override;
+	/** Creates a series of nodes that display a single input. */
+	DYNAMICMATERIALEDITOR_API virtual TArray<UMaterialExpression*> CreateExpressionInput(UDMMaterialStageInput* InInput) const override;
 
-	virtual UMaterialExpressionComponentMask* CreateExpressionBitMask(UMaterialExpression* InExpression, int32 InOutputIndex,
+	/** Creates a mask expression with the given channels exposed. */
+	DYNAMICMATERIALEDITOR_API virtual UMaterialExpressionComponentMask* CreateExpressionBitMask(UMaterialExpression* InExpression, int32 InOutputIndex,
 		int32 InOutputChannels) const override;
 
-	virtual UMaterialExpressionAppendVector* CreateExpressionAppend(UMaterialExpression* InExpressionA, int32 InOutputIndexA, 
+	/** Creates an append expression to combine vectors/scalars together to create larger vectors. */
+	DYNAMICMATERIALEDITOR_API virtual UMaterialExpressionAppendVector* CreateExpressionAppend(UMaterialExpression* InExpressionA, int32 InOutputIndexA,
 		UMaterialExpression* InExpressionB, int32 InOutputIndexB) const override;
 
-	virtual void UpdatePreviewMaterial(UMaterialExpression* InLastExpression, int32 InOutputIndex, int32 InOutputChannel,
+	/** Updates the emissive channel of the give material to show the output of the given expression. */
+	DYNAMICMATERIALEDITOR_API virtual void UpdatePreviewMaterial(UMaterialExpression* InLastExpression, int32 InOutputIndex, int32 InOutputChannel,
 		int32 InSize) const override;
 
 protected:
diff --git a/Engine/Plugins/Experimental/DynamicMaterial/Source/DynamicMaterialEditor/Public/Model/DMOnWizardCompleteCallback.h b/Engine/Plugins/Experimental/DynamicMaterial/Source/DynamicMaterialEditor/Public/Model/DMOnWizardCompleteCallback.h
--- a/Engine/Plugins/Experimental/DynamicMaterial/Source/DynamicMaterialEditor/Public/Model/DMOnWizardCompleteCallback.h
+++ b/Engine/Plugins/Experimental/DynamicMaterial/Source/DynamicMaterialEditor/Public/Model/DMOnWizardCompleteCallback.h
@@ -56,12 +56,12 @@ struct FDMMaterialModelCreatedCallbackDelegate : public FDMMaterialModelCreatedC
 {
 	DECLARE_DELEGATE_OneParam(FOnModelCreated, const FDMOnWizardCompleteCallbackParams&)
 
-	FDMMaterialModelCreatedCallbackDelegate(uint32 InPriority, const FOnModelCreated& InOnModelCreatedDelegate);
+	DYNAMICMATERIALEDITOR_API FDMMaterialModelCreatedCallbackDelegate(uint32 InPriority, const FOnModelCreated& InOnModelCreatedDelegate);
 
 	virtual ~FDMMaterialModelCreatedCallbackDelegate() override = default;
 
 	//~ Begin IDMMaterialModelCreatedCallback
-	virtual void OnModelCreated(const FDMOnWizardCompleteCallbackParams& InParams) override;
+	DYNAMICMATERIALEDITOR_API virtual void OnModelCreated(const FDMOnWizardCompleteCallbackParams& InParams) override;
 	//~ End IDMMaterialModelCreatedCallback
 
 protected:
diff --git a/Engine/Plugins/Experimental/DynamicMaterial/Source/DynamicMaterialEditor/Public/Model/DynamicMaterialModelEditorOnlyData.h b/Engine/Plugins/Experimental/DynamicMaterial/Source/DynamicMaterialEditor/Public/Model/DynamicMaterialModelEditorOnlyData.h
--- a/Engine/Plugins/Experimental/DynamicMaterial/Source/DynamicMaterialEditor/Public/Model/DynamicMaterialModelEditorOnlyData.h
+++ b/Engine/Plugins/Experimental/DynamicMaterial/Source/DynamicMaterialEditor/Public/Model/DynamicMaterialModelEditorOnlyData.h
@@ -23,6 +23,7 @@ class UDMMaterialValueFloat1;
 class UDMTextureSet;
 class UDMTextureUV;
 class UDynamicMaterialModel;
+class UDynamicMaterialModelBase;
 class UDynamicMaterialModelEditorOnlyData;
 class UMaterialExpression;
 struct FDMComponentPath;
@@ -30,9 +31,9 @@ struct FDMComponentPathSegment;
 struct FDMMaterialBuildState;
 struct FDMMaterialChannelListPreset;
 
-DECLARE_MULTICAST_DELEGATE_OneParam(FDMOnMaterialBuilt, UDynamicMaterialModel*);
-DECLARE_MULTICAST_DELEGATE_OneParam(FDMOnValueListUpdated, UDynamicMaterialModel*);
-DECLARE_MULTICAST_DELEGATE_OneParam(FDMOnSlotListUpdated, UDynamicMaterialModel*);
+DECLARE_MULTICAST_DELEGATE_OneParam(FDMOnMaterialBuilt, UDynamicMaterialModelBase*);
+DECLARE_MULTICAST_DELEGATE_OneParam(FDMOnValueListUpdated, UDynamicMaterialModelBase*);
+DECLARE_MULTICAST_DELEGATE_OneParam(FDMOnSlotListUpdated, UDynamicMaterialModelBase*);
 
 UENUM(BlueprintType)
 enum class EDMState : uint8
@@ -41,8 +42,8 @@ enum class EDMState : uint8
 	Building
 };
 
-UCLASS(BlueprintType, EditInlineNew, DefaultToInstanced, ClassGroup = "Material Designer")
-class DYNAMICMATERIALEDITOR_API UDynamicMaterialModelEditorOnlyData : public UObject, public IDynamicMaterialModelEditorOnlyDataInterface, public FNotifyHook, public IDMBuildable
+UCLASS(MinimalAPI, BlueprintType, EditInlineNew, DefaultToInstanced, ClassGroup = "Material Designer")
+class UDynamicMaterialModelEditorOnlyData : public UObject, public IDynamicMaterialModelEditorOnlyDataInterface, public FNotifyHook, public IDMBuildable
 {
 	GENERATED_BODY()
 
@@ -50,35 +51,38 @@ class DYNAMICMATERIALEDITOR_API UDynamicMaterialModelEditorOnlyData : public UOb
 	friend class SDMComponentEdit;
 
 public:
-	static const FString SlotsPathToken;
-	static const FString RGBSlotPathToken;
-	static const FString OpacitySlotPathToken;
-	static const FString RoughnessPathToken;
-	static const FString SpecularPathToken;
-	static const FString MetallicPathToken;
-	static const FString NormalPathToken;
-	static const FString PixelDepthOffsetPathToken;
-	static const FString WorldPositionOffsetPathToken;
-	static const FString AmbientOcclusionPathToken;
-	static const FString AnisotropyPathToken;
-	static const FString RefractionPathToken;
-	static const FString TangentPathToken;
-	static const FString Custom1PathToken;
-	static const FString Custom2PathToken;
-	static const FString Custom3PathToken;
-	static const FString Custom4PathToken;
-	static const FString PropertiesPathToken;
-
-	static const TArray<EMaterialDomain> SupportedDomains;
-	static const TArray<EBlendMode> SupportedBlendModes;
-
-	static const FName AlphaValueName;
-
-	static UDynamicMaterialModelEditorOnlyData* Get(UDynamicMaterialModel* InModel);
-	static UDynamicMaterialModelEditorOnlyData* Get(TWeakObjectPtr<UDynamicMaterialModel> InModelWeak);
-	static UDynamicMaterialModelEditorOnlyData* Get(const TScriptInterface<IDynamicMaterialModelEditorOnlyDataInterface>& InInterface);
-	static UDynamicMaterialModelEditorOnlyData* Get(IDynamicMaterialModelEditorOnlyDataInterface* InInterface);
-	static UDynamicMaterialModelEditorOnlyData* Get(UDynamicMaterialInstance* InInstance);
+	DYNAMICMATERIALEDITOR_API static const FString SlotsPathToken;
+	DYNAMICMATERIALEDITOR_API static const FString BaseColorSlotPathToken;
+	DYNAMICMATERIALEDITOR_API static const FString EmissiveSlotPathToken;
+	DYNAMICMATERIALEDITOR_API static const FString OpacitySlotPathToken;
+	DYNAMICMATERIALEDITOR_API static const FString RoughnessPathToken;
+	DYNAMICMATERIALEDITOR_API static const FString SpecularPathToken;
+	DYNAMICMATERIALEDITOR_API static const FString MetallicPathToken;
+	DYNAMICMATERIALEDITOR_API static const FString NormalPathToken;
+	DYNAMICMATERIALEDITOR_API static const FString PixelDepthOffsetPathToken;
+	DYNAMICMATERIALEDITOR_API static const FString WorldPositionOffsetPathToken;
+	DYNAMICMATERIALEDITOR_API static const FString AmbientOcclusionPathToken;
+	DYNAMICMATERIALEDITOR_API static const FString AnisotropyPathToken;
+	DYNAMICMATERIALEDITOR_API static const FString RefractionPathToken;
+	DYNAMICMATERIALEDITOR_API static const FString TangentPathToken;
+	DYNAMICMATERIALEDITOR_API static const FString Custom1PathToken;
+	DYNAMICMATERIALEDITOR_API static const FString Custom2PathToken;
+	DYNAMICMATERIALEDITOR_API static const FString Custom3PathToken;
+	DYNAMICMATERIALEDITOR_API static const FString Custom4PathToken;
+	DYNAMICMATERIALEDITOR_API static const FString PropertiesPathToken;
+
+	DYNAMICMATERIALEDITOR_API static const TArray<EMaterialDomain> SupportedDomains;
+	DYNAMICMATERIALEDITOR_API static const TArray<EBlendMode> SupportedBlendModes;
+
+	DYNAMICMATERIALEDITOR_API static const FName AlphaValueName;
+
+	DYNAMICMATERIALEDITOR_API static UDynamicMaterialModelEditorOnlyData* Get(UDynamicMaterialModelBase* InModelBase);
+	DYNAMICMATERIALEDITOR_API static UDynamicMaterialModelEditorOnlyData* Get(const TWeakObjectPtr<UDynamicMaterialModelBase>& InModelBaseWeak);
+	DYNAMICMATERIALEDITOR_API static UDynamicMaterialModelEditorOnlyData* Get(UDynamicMaterialModel* InModel);
+	DYNAMICMATERIALEDITOR_API static UDynamicMaterialModelEditorOnlyData* Get(const TWeakObjectPtr<UDynamicMaterialModel>& InModelWeak);
+	DYNAMICMATERIALEDITOR_API static UDynamicMaterialModelEditorOnlyData* Get(const TScriptInterface<IDynamicMaterialModelEditorOnlyDataInterface>& InInterface);
+	DYNAMICMATERIALEDITOR_API static UDynamicMaterialModelEditorOnlyData* Get(IDynamicMaterialModelEditorOnlyDataInterface* InInterface);
+	DYNAMICMATERIALEDITOR_API static UDynamicMaterialModelEditorOnlyData* Get(UDynamicMaterialInstance* InInstance);
 
 	UDynamicMaterialModelEditorOnlyData();
 
@@ -88,7 +92,7 @@ public:
 	UDynamicMaterialModel* GetMaterialModel() const { return MaterialModel; }
 
 	UFUNCTION(BlueprintPure, Category = "Material Designer")
-	UMaterial* GetGeneratedMaterial() const;
+	DYNAMICMATERIALEDITOR_API UMaterial* GetGeneratedMaterial() const;
 
 	UFUNCTION(BlueprintPure, Category = "Material Designer")
 	EDMState GetState() const { return State; }
@@ -97,85 +101,85 @@ public:
 	TEnumAsByte<EMaterialDomain> GetDomain() const { return Domain; }
 
 	UFUNCTION(BlueprintCallable, Category = "Material Designer")
-	void SetDomain(TEnumAsByte<EMaterialDomain> InDomain);
+	DYNAMICMATERIALEDITOR_API void SetDomain(TEnumAsByte<EMaterialDomain> InDomain);
 
 	UFUNCTION(BlueprintPure, Category = "Material Designer")
 	TEnumAsByte<EBlendMode> GetBlendMode() const { return BlendMode; }
 
 	UFUNCTION(BlueprintCallable, Category = "Material Designer")
-	void SetBlendMode(TEnumAsByte<EBlendMode> InBlendMode);
+	DYNAMICMATERIALEDITOR_API void SetBlendMode(TEnumAsByte<EBlendMode> InBlendMode);
 
 	UFUNCTION(BlueprintPure, Category = "Material Designer")
 	EDMMaterialShadingModel GetShadingModel() const { return ShadingModel; }
 
 	UFUNCTION(BlueprintCallable, Category = "Material Designer")
-	void SetShadingModel(EDMMaterialShadingModel InShadingModel);
+	DYNAMICMATERIALEDITOR_API void SetShadingModel(EDMMaterialShadingModel InShadingModel);
 
 	UFUNCTION(BlueprintPure, Category = "Material Designer")
-	bool IsPixelAnimationFlagSet() const;
+	DYNAMICMATERIALEDITOR_API bool IsPixelAnimationFlagSet() const;
 
 	UFUNCTION(BlueprintCallable, Category = "Material Designer")
-	void SetPixelAnimationFlag(bool bInFlagValue);
+	DYNAMICMATERIALEDITOR_API void SetPixelAnimationFlag(bool bInFlagValue);
 
 	UFUNCTION(BlueprintPure, Category = "Material Designer")
-	bool IsTwoSidedFlagSet() const;
+	DYNAMICMATERIALEDITOR_API bool IsTwoSidedFlagSet() const;
 
 	UFUNCTION(BlueprintCallable, Category = "Material Designer")
-	void SetTwoSidedFlag(bool bInFlagValue);
+	DYNAMICMATERIALEDITOR_API void SetTwoSidedFlag(bool bInFlagValue);
 
 	UFUNCTION(BlueprintCallable, Category = "Material Designer")
-	FName GetChannelListPreset() const;
+	DYNAMICMATERIALEDITOR_API FName GetChannelListPreset() const;
 
 	UFUNCTION(BlueprintCallable, Category = "Material Designer")
-	void SetChannelListPreset(FName InPresetName);
+	DYNAMICMATERIALEDITOR_API void SetChannelListPreset(FName InPresetName);
 
 	UFUNCTION(BlueprintCallable, Category = "Material Designer")
-	void OpenMaterialEditor() const;
+	DYNAMICMATERIALEDITOR_API void OpenMaterialEditor() const;
 
 	UFUNCTION(BlueprintPure, Category = "Material Designer")
-	TMap<EDMMaterialPropertyType, UDMMaterialProperty*> GetMaterialProperties() const;
+	DYNAMICMATERIALEDITOR_API TMap<EDMMaterialPropertyType, UDMMaterialProperty*> GetMaterialProperties() const;
 
 	UFUNCTION(BlueprintPure, Category = "Material Designer")
-	UDMMaterialProperty* GetMaterialProperty(EDMMaterialPropertyType MaterialProperty) const;
+	DYNAMICMATERIALEDITOR_API UDMMaterialProperty* GetMaterialProperty(EDMMaterialPropertyType MaterialProperty) const;
 
 	UFUNCTION(BlueprintPure, Category = "Material Designer")
 	const TArray<UDMMaterialSlot*>& GetSlots() const { return Slots; }
 
 	/** Gets slot by index. Highly recommended to use GetSlotForMaterialProperty(PropertyType). */
 	UFUNCTION(BlueprintPure, Category = "Material Designer")
-	UDMMaterialSlot* GetSlot(int32 Index) const;
+	DYNAMICMATERIALEDITOR_API UDMMaterialSlot* GetSlot(int32 Index) const;
 
 	UFUNCTION(BlueprintPure, Category = "Material Designer")
-	UDMMaterialSlot* GetSlotForMaterialProperty(EDMMaterialPropertyType InType) const;
+	DYNAMICMATERIALEDITOR_API UDMMaterialSlot* GetSlotForMaterialProperty(EDMMaterialPropertyType InType) const;
 
 	/** Adds the next available slot. Highly recommended to use AddSlotForMaterialProperty(PropertyType). */
 	UFUNCTION(BlueprintCallable, Category = "Material Designer")
-	UDMMaterialSlot* AddSlot();
+	DYNAMICMATERIALEDITOR_API UDMMaterialSlot* AddSlot();
 
 	UFUNCTION(BlueprintCallable, Category = "Material Designer")
-	UDMMaterialSlot* AddSlotForMaterialProperty(EDMMaterialPropertyType InType);
+	DYNAMICMATERIALEDITOR_API UDMMaterialSlot* AddSlotForMaterialProperty(EDMMaterialPropertyType InType);
 
 	/** Removes the next slot by index. Highly recommended to use RemoveSlotForMaterialProperty(PropertyType). */
 	UFUNCTION(BlueprintCallable, Category = "Material Designer")
-	UDMMaterialSlot* RemoveSlot(int32 Index);
+	DYNAMICMATERIALEDITOR_API UDMMaterialSlot* RemoveSlot(int32 Index);
 
 	UFUNCTION(BlueprintCallable, Category = "Material Designer")
-	UDMMaterialSlot* RemoveSlotForMaterialProperty(EDMMaterialPropertyType InType);
+	DYNAMICMATERIALEDITOR_API UDMMaterialSlot* RemoveSlotForMaterialProperty(EDMMaterialPropertyType InType);
 
 	UFUNCTION(BlueprintPure, Category = "Material Designer")
-	TArray<EDMMaterialPropertyType> GetMaterialPropertiesForSlot(UDMMaterialSlot* Slot) const;
+	DYNAMICMATERIALEDITOR_API TArray<EDMMaterialPropertyType> GetMaterialPropertiesForSlot(const UDMMaterialSlot* Slot) const;
 
 	UFUNCTION(BlueprintCallable, Category = "Material Designer")
-	void AssignMaterialPropertyToSlot(EDMMaterialPropertyType Property, UDMMaterialSlot* Slot);
+	DYNAMICMATERIALEDITOR_API void AssignMaterialPropertyToSlot(EDMMaterialPropertyType Property, UDMMaterialSlot* Slot);
 
 	UFUNCTION(BlueprintCallable, Category = "Material Designer")
-	void UnassignMaterialProperty(EDMMaterialPropertyType Property);
+	DYNAMICMATERIALEDITOR_API void UnassignMaterialProperty(EDMMaterialPropertyType Property);
 	
-	FDMOnMaterialBuilt& GetOnMaterialBuiltDelegate() { return OnMaterialBuiltDelegate; }
-	FDMOnValueListUpdated& GetOnValueListUpdateDelegate() { return OnValueListUpdateDelegate; }
-	FDMOnSlotListUpdated& GetOnSlotListUpdateDelegate() { return OnSlotListUpdateDelegate; }
+	FDMOnMaterialBuilt::RegistrationType& GetOnMaterialBuiltDelegate() { return OnMaterialBuiltDelegate; }
+	FDMOnValueListUpdated::RegistrationType& GetOnValueListUpdateDelegate() { return OnValueListUpdateDelegate; }
+	FDMOnSlotListUpdated::RegistrationType& GetOnSlotListUpdateDelegate() { return OnSlotListUpdateDelegate; }
 
-	TSharedRef<FDMMaterialBuildState> CreateBuildState(UMaterial* InMaterialToBuild, bool bInDirtyAssets = true) const;
+	DYNAMICMATERIALEDITOR_API TSharedRef<FDMMaterialBuildState> CreateBuildState(UMaterial* InMaterialToBuild, bool bInDirtyAssets = true) const;
 
 	/**
 	 * Integrates a Texture Set with the model.
@@ -184,54 +188,53 @@ public:
 	 * @return True on success.
 	 */
 	UFUNCTION(BlueprintCallable, Category = "Material Designer")
-	bool AddTextureSet(UDMTextureSet* InTextureSet, bool bInReplaceSlots);
+	DYNAMICMATERIALEDITOR_API bool AddTextureSet(UDMTextureSet* InTextureSet, bool bInReplaceSlots);
 
 	bool NeedsWizard() const;
 
 	void OnWizardComplete();
 
-	void ForEachMaterialPropertyType(TFunctionRef<EDMIterationResult(EDMMaterialPropertyType InType)> InCallable,
+	DYNAMICMATERIALEDITOR_API void ForEachMaterialPropertyType(TFunctionRef<EDMIterationResult(EDMMaterialPropertyType InType)> InCallable,
 		EDMMaterialPropertyType InStart = static_cast<EDMMaterialPropertyType>(static_cast<uint8>(EDMMaterialPropertyType::None) + 1),
 		EDMMaterialPropertyType InEnd = static_cast<EDMMaterialPropertyType>(static_cast<uint8>(EDMMaterialPropertyType::Any) - 1));
 
+	void SaveEditor();
+
 	//~ Begin FNotifyHook
-	virtual void NotifyPreChange(class FEditPropertyChain* PropertyAboutToChange) {}
-	virtual void NotifyPostChange(const FPropertyChangedEvent& PropertyChangedEvent, class FEditPropertyChain* PropertyThatChanged);
+	DYNAMICMATERIALEDITOR_API virtual void NotifyPostChange(const FPropertyChangedEvent& InPropertyChangedEvent, class FEditPropertyChain* PropertyThatChanged);
 	//~ End FNotifyHook
 
 	//~ Begin UObject
-	virtual void PostLoad() override;
-	virtual void PostEditUndo() override;
-	virtual void PostEditImport() override;
-	virtual void PostDuplicate(bool bDuplicateForPIE) override;
-	virtual void PostEditChangeChainProperty(FPropertyChangedChainEvent& PropertyChangedEvent) override;
-	virtual void Serialize(FArchive& Ar) override;
+	DYNAMICMATERIALEDITOR_API virtual void PostLoad() override;
+	DYNAMICMATERIALEDITOR_API virtual void PostEditUndo() override;
+	DYNAMICMATERIALEDITOR_API virtual void PostEditImport() override;
+	DYNAMICMATERIALEDITOR_API virtual void PostDuplicate(bool bDuplicateForPIE) override;
+	DYNAMICMATERIALEDITOR_API virtual void PostEditChangeChainProperty(FPropertyChangedChainEvent& PropertyChangedEvent) override;
+	DYNAMICMATERIALEDITOR_API virtual void Serialize(FArchive& Ar) override;
 	//~ End UObject
 
-	void SaveEditor();
-
 	//~ Begin IDMBuildable
 	virtual void DoBuild_Implementation(bool bInDirtyAssets) override;
 	//~ End IDMBuildable
 
 	//~ Begin IDynamicMaterialModelEditorOnlyDataInterface
-	virtual void PostEditorDuplicate() override;
-	virtual void RequestMaterialBuild() override;
-	virtual void OnValueListUpdate() override;
-	virtual void OnValueUpdated(UDMMaterialValue* InValue, EDMUpdateType InUpdateType) override;
-	virtual void OnTextureUVUpdated(UDMTextureUV* InTextureUV) override;
-	virtual void LoadDeprecatedModelData(UDynamicMaterialModel* InMaterialModel) override;
-	virtual TSharedRef<IDMMaterialBuildStateInterface> CreateBuildStateInterface(UMaterial* InMaterialToBuild) const override;
-	virtual void SetPropertyComponent(EDMMaterialPropertyType InPropertyType, FName InComponentName, UDMMaterialComponent* InComponent) override;
-	virtual UDMMaterialComponent* GetSubComponentByPath(FDMComponentPath& InPath) const override;
-	virtual UDMMaterialComponent* GetSubComponentByPath(FDMComponentPath& InPath, const FDMComponentPathSegment& InPathSegment) const override;
+	DYNAMICMATERIALEDITOR_API virtual void PostEditorDuplicate() override;
+	DYNAMICMATERIALEDITOR_API virtual void RequestMaterialBuild() override;
+	DYNAMICMATERIALEDITOR_API virtual void OnValueListUpdate() override;
+	DYNAMICMATERIALEDITOR_API virtual void OnValueUpdated(UDMMaterialValue* InValue, EDMUpdateType InUpdateType) override;
+	DYNAMICMATERIALEDITOR_API virtual void OnTextureUVUpdated(UDMTextureUV* InTextureUV) override;
+	DYNAMICMATERIALEDITOR_API virtual void LoadDeprecatedModelData(UDynamicMaterialModel* InMaterialModel) override;
+	DYNAMICMATERIALEDITOR_API virtual TSharedRef<IDMMaterialBuildStateInterface> CreateBuildStateInterface(UMaterial* InMaterialToBuild) const override;
+	DYNAMICMATERIALEDITOR_API virtual void SetPropertyComponent(EDMMaterialPropertyType InPropertyType, FName InComponentName, UDMMaterialComponent* InComponent) override;
+	DYNAMICMATERIALEDITOR_API virtual UDMMaterialComponent* GetSubComponentByPath(FDMComponentPath& InPath) const override;
+	DYNAMICMATERIALEDITOR_API virtual UDMMaterialComponent* GetSubComponentByPath(FDMComponentPath& InPath, const FDMComponentPathSegment& InPathSegment) const override;
 	//~ End IDynamicMaterialModelEditorOnlyDataInterface
 
 protected:
 	UPROPERTY(VisibleInstanceOnly, BlueprintReadOnly, Category = "Material Designer")
 	TObjectPtr<UDynamicMaterialModel> MaterialModel;
 
-	UPROPERTY(VisibleInstanceOnly, BlueprintReadOnly, Transient, DuplicateTransient, TextExportTransient, Category = "Material Designer")
+	UPROPERTY(VisibleInstanceOnly, BlueprintReadOnly, Transient, TextExportTransient, Category = "Material Designer")
 	EDMState State;
 
 	UPROPERTY(EditInstanceOnly, BlueprintReadOnly, Category = "Material Designer", 
@@ -259,7 +262,7 @@ protected:
 	UPROPERTY(VisibleInstanceOnly, BlueprintReadOnly, Instanced, Category = "Material Designer")
 	TMap<EDMMaterialPropertyType, TObjectPtr<UDMMaterialProperty>> Properties;
 
-	UPROPERTY(VisibleInstanceOnly, BlueprintReadOnly, DuplicateTransient, TextExportTransient, Category = "Material Designer")
+	UPROPERTY(VisibleInstanceOnly, BlueprintReadOnly, TextExportTransient, Category = "Material Designer")
 	TMap<EDMMaterialPropertyType, TObjectPtr<UDMMaterialSlot>> PropertySlotMap;
 
 	UPROPERTY(VisibleInstanceOnly, BlueprintReadOnly, Instanced, Category = "Material Designer")
@@ -295,11 +298,6 @@ protected:
 	 */
 	void EnsureSwapSlotMaterialProperty(EDMMaterialPropertyType InPropertyFrom, EDMMaterialPropertyType InPropertyTo);
 
-	//~ Begin IDynamicMaterialModelEditorOnlyDataInterface
-	virtual void ReinitComponents() override;
-	virtual void ResetData() override;
-	//~ End IDynamicMaterialModelEditorOnlyDataInterface
-
 	void AssignPropertyAlphaValues();
 
 	UFUNCTION()
@@ -313,6 +311,10 @@ protected:
 	void OnPixelAnimationFlagChanged();
 	void OnTwoSidedFlagChanged();
 
+	//~ Begin IDynamicMaterialModelEditorOnlyDataInterface
+	virtual void ReinitComponents() override;
+	//~ End IDynamicMaterialModelEditorOnlyDataInterface
+
 	void LoadDeprecatedModelData_Base(bool bInCreateMaterialPackage, EBlendMode InBlendMode, EDMMaterialShadingModel InShadingModel);
 	void LoadDeprecatedModelData_Expressions(TArray<TObjectPtr<UMaterialExpression>>& InExpressions);
 	void LoadDeprecatedModelData_Properties(TMap<EDMMaterialPropertyType, TObjectPtr<UObject>>& InProperties);
diff --git a/Engine/Plugins/Experimental/DynamicMaterial/Source/DynamicMaterialEditor/Public/Slate/Properties/Generators/DMComponentPropertyRowGenerator.h b/Engine/Plugins/Experimental/DynamicMaterial/Source/DynamicMaterialEditor/Public/Slate/Properties/Generators/DMComponentPropertyRowGenerator.h
--- a/Engine/Plugins/Experimental/DynamicMaterial/Source/DynamicMaterialEditor/Public/Slate/Properties/Generators/DMComponentPropertyRowGenerator.h
+++ b/Engine/Plugins/Experimental/DynamicMaterial/Source/DynamicMaterialEditor/Public/Slate/Properties/Generators/DMComponentPropertyRowGenerator.h
@@ -12,22 +12,44 @@ class SDMComponentEdit;
 class UDMMaterialComponent;
 struct FDMPropertyHandle;
 
-class DYNAMICMATERIALEDITOR_API FDMComponentPropertyRowGenerator
+/** Used to generate editable properties for objects editable in the Material Designer. */
+class FDMComponentPropertyRowGenerator
 {
 public:
-	static const TSharedRef<FDMComponentPropertyRowGenerator>& Get();
+	DYNAMICMATERIALEDITOR_API static const TSharedRef<FDMComponentPropertyRowGenerator>& Get();
 
-	FDMComponentPropertyRowGenerator() = default;
 	virtual ~FDMComponentPropertyRowGenerator() = default;
 
-	virtual void AddComponentProperties(const TSharedRef<SDMComponentEdit>& InComponentEditWidget, UDMMaterialComponent* InComponent, 
-		TArray<FDMPropertyHandle>& InOutPropertyRows, TSet<UDMMaterialComponent*>& InOutProcessedObjects);
-
-	virtual void AddPropertyEditRows(const TSharedRef<SDMComponentEdit>& InComponentEditWidget, UDMMaterialComponent* InComponent, 
-		const FName& InProperty, TArray<FDMPropertyHandle>& InOutPropertyRows, TSet<UDMMaterialComponent*>& InOutProcessedObjects);
-
-	virtual void AddPropertyEditRows(const TSharedRef<SDMComponentEdit>& InComponentEditWidget, UDMMaterialComponent* InComponent, FProperty* InProperty, 
-		void* MemoryPtr, TArray<FDMPropertyHandle>& InOutPropertyRows, TSet<UDMMaterialComponent*>& InOutProcessedObjects);
-
-	virtual bool AllowKeyframeButton(UDMMaterialComponent* InComponent, FProperty* InProperty);
+	/**
+	 * Generate properties for the given component.
+	 * @param InComponentEditWidget The edit widget generating the properties.
+	 * @param InComponent The component being edited.
+	 * @param InOutPropertyRows The generated rows.
+	 * @param InOutProcessedObjects The already processed objects - add to this to avoid possible recursive generation.
+	 */
+	DYNAMICMATERIALEDITOR_API virtual void AddComponentProperties(const TSharedRef<SDMComponentEdit>& InComponentEditWidget, 
+		UDMMaterialComponent* InComponent, TArray<FDMPropertyHandle>& InOutPropertyRows, TSet<UDMMaterialComponent*>& InOutProcessedObjects);
+
+	/**
+	 * Add the rows needed for a specific property.
+	 * @param InProperty The name of the property.
+	 */
+	DYNAMICMATERIALEDITOR_API virtual void AddPropertyEditRows(const TSharedRef<SDMComponentEdit>& InComponentEditWidget, 
+		UDMMaterialComponent* InComponent, const FName& InProperty, TArray<FDMPropertyHandle>& InOutPropertyRows, 
+		TSet<UDMMaterialComponent*>& InOutProcessedObjects);
+
+	/**
+	 * Adds the rows needed for a specific property.
+	 * @param InProperty The FProperty for this property.
+	 * @param MemoryPtr The pointer to the property's value in memory
+	 */
+	DYNAMICMATERIALEDITOR_API virtual void AddPropertyEditRows(const TSharedRef<SDMComponentEdit>& InComponentEditWidget, 
+		UDMMaterialComponent* InComponent, FProperty* InProperty, void* MemoryPtr, TArray<FDMPropertyHandle>& InOutPropertyRows, 
+		TSet<UDMMaterialComponent*>& InOutProcessedObjects);
+
+	DYNAMICMATERIALEDITOR_API virtual bool AllowKeyframeButton(UDMMaterialComponent* InComponent, FProperty* InProperty);
+
+protected:
+	/** Returns true if the component edit widget is editing a Material Designer Dynamic. */
+	DYNAMICMATERIALEDITOR_API static bool IsDynamic(const TSharedRef<SDMComponentEdit>& InComponentEditWidget);
 };
diff --git a/Engine/Plugins/Experimental/DynamicMaterial/Source/DynamicMaterialEditor/Public/Utils/DMMaterialEffectStackPresetSubsystem.h b/Engine/Plugins/Experimental/DynamicMaterial/Source/DynamicMaterialEditor/Public/Utils/DMMaterialEffectStackPresetSubsystem.h
--- a/Engine/Plugins/Experimental/DynamicMaterial/Source/DynamicMaterialEditor/Public/Utils/DMMaterialEffectStackPresetSubsystem.h
+++ b/Engine/Plugins/Experimental/DynamicMaterial/Source/DynamicMaterialEditor/Public/Utils/DMMaterialEffectStackPresetSubsystem.h
@@ -11,23 +11,23 @@ class UDMMaterialEffect;
 class UDMMaterialEffectStack;
 struct FDMMaterialEffectStackJson;
 
-UCLASS(BlueprintType)
+UCLASS(MinimalAPI, BlueprintType)
 class UDMMaterialEffectStackPresetSubsystem : public UEditorSubsystem
 {
 	GENERATED_BODY()
 
 public:
-	static UDMMaterialEffectStackPresetSubsystem* Get();
+	DYNAMICMATERIALEDITOR_API static UDMMaterialEffectStackPresetSubsystem* Get();
 
 	UFUNCTION(BlueprintCallable, Category = "Material Designer")
-	bool SavePreset(const FString& InPresetName, const FDMMaterialEffectStackJson& InPreset) const;
+	DYNAMICMATERIALEDITOR_API bool SavePreset(const FString& InPresetName, const FDMMaterialEffectStackJson& InPreset) const;
 
 	UFUNCTION(BlueprintCallable, Category = "Material Designer")
-	bool LoadPreset(const FString& InPresetName, FDMMaterialEffectStackJson& OutPreset) const;
+	DYNAMICMATERIALEDITOR_API bool LoadPreset(const FString& InPresetName, FDMMaterialEffectStackJson& OutPreset) const;
 
 	UFUNCTION(BlueprintCallable, Category = "Material Designer")
-	bool RemovePreset(const FString& InPresetName) const;
+	DYNAMICMATERIALEDITOR_API bool RemovePreset(const FString& InPresetName) const;
 
 	UFUNCTION(BlueprintCallable, Category = "Material Designer")
-	TArray<FString> GetPresetNames() const;
+	DYNAMICMATERIALEDITOR_API TArray<FString> GetPresetNames() const;
 };
diff --git a/Engine/Plugins/Experimental/DynamicMaterial/Source/DynamicMaterialTextureSetEditor/Private/Widgets/SDMTextureSetBuilderMaterialPropertyCell.cpp b/Engine/Plugins/Experimental/DynamicMaterial/Source/DynamicMaterialTextureSetEditor/Private/Widgets/SDMTextureSetBuilderMaterialPropertyCell.cpp
--- a/Engine/Plugins/Experimental/DynamicMaterial/Source/DynamicMaterialTextureSetEditor/Private/Widgets/SDMTextureSetBuilderMaterialPropertyCell.cpp
+++ b/Engine/Plugins/Experimental/DynamicMaterial/Source/DynamicMaterialTextureSetEditor/Private/Widgets/SDMTextureSetBuilderMaterialPropertyCell.cpp
@@ -156,9 +156,9 @@ bool SDMTextureSetBuilderMaterialPropertyCell::GetPropertyEnabled() const
 	return Texture.IsValid();
 }
 
-void SDMTextureSetBuilderMaterialPropertyCell::NotifyPostChange(const FPropertyChangedEvent& PropertyChangedEvent, FProperty* PropertyThatChanged)
+void SDMTextureSetBuilderMaterialPropertyCell::NotifyPostChange(const FPropertyChangedEvent& InPropertyChangedEvent, FProperty* InPropertyThatChanged)
 {
-	if (PropertyThatChanged && PropertyThatChanged->GetFName() == GET_MEMBER_NAME_CHECKED(FDMTextureSetBuilderEntry, ChannelMask))
+	if (InPropertyThatChanged && InPropertyThatChanged->GetFName() == GET_MEMBER_NAME_CHECKED(FDMTextureSetBuilderEntry, ChannelMask))
 	{
 		SetMaterialForChannelMask();
 	}
diff --git a/Engine/Plugins/Experimental/DynamicMaterial/Source/DynamicMaterialTextureSetEditor/Private/Widgets/SDMTextureSetBuilderMaterialPropertyCell.h b/Engine/Plugins/Experimental/DynamicMaterial/Source/DynamicMaterialTextureSetEditor/Private/Widgets/SDMTextureSetBuilderMaterialPropertyCell.h
--- a/Engine/Plugins/Experimental/DynamicMaterial/Source/DynamicMaterialTextureSetEditor/Private/Widgets/SDMTextureSetBuilderMaterialPropertyCell.h
+++ b/Engine/Plugins/Experimental/DynamicMaterial/Source/DynamicMaterialTextureSetEditor/Private/Widgets/SDMTextureSetBuilderMaterialPropertyCell.h
@@ -31,7 +31,7 @@ public:
 	//~ End SDMTextureSetBuilderCellBase
 
 	//~ Begin FNotifyHook
-	virtual void NotifyPostChange(const FPropertyChangedEvent& PropertyChangedEvent, FProperty* PropertyThatChanged) override;
+	virtual void NotifyPostChange(const FPropertyChangedEvent& InPropertyChangedEvent, FProperty* InPropertyThatChanged) override;
 	//~ End FNotifyHook
 
 protected:
diff --git a/Engine/Plugins/Experimental/DynamicMaterial/Source/DynamicMaterialTextureSetEditor/Public/DMTextureSetBlueprintFunctionLibrary.h b/Engine/Plugins/Experimental/DynamicMaterial/Source/DynamicMaterialTextureSetEditor/Public/DMTextureSetBlueprintFunctionLibrary.h
--- a/Engine/Plugins/Experimental/DynamicMaterial/Source/DynamicMaterialTextureSetEditor/Public/DMTextureSetBlueprintFunctionLibrary.h
+++ b/Engine/Plugins/Experimental/DynamicMaterial/Source/DynamicMaterialTextureSetEditor/Public/DMTextureSetBlueprintFunctionLibrary.h
@@ -15,7 +15,7 @@ DECLARE_DELEGATE_TwoParams(FDMTextureSetBuilderOnComplete, UDMTextureSet*, /* Wa
 /**
  * Material Designer Texture Set Blueprint Function Library
  */
-UCLASS()
+UCLASS(BlueprintType)
 class UDMTextureSetBlueprintFunctionLibrary : public UBlueprintFunctionLibrary
 {
 	GENERATED_BODY()
