Motion Design
- Added option for Output Translucent Velocity (previously enabled by default)
- Added option for Nanite Tessellation
- Added option for Responsive AA (previously enabled by default)
- Added tooltips for the flags.

diff --git a/Engine/Plugins/Experimental/DynamicMaterial/Config/BaseDynamicMaterial.ini b/Engine/Plugins/Experimental/DynamicMaterial/Config/BaseDynamicMaterial.ini
--- a/Engine/Plugins/Experimental/DynamicMaterial/Config/BaseDynamicMaterial.ini
+++ b/Engine/Plugins/Experimental/DynamicMaterial/Config/BaseDynamicMaterial.ini
@@ -33,4 +33,11 @@
 
 +PropertyRedirects=(OldName="/Script/DynamicMaterial.DMDefaultMaterialPropertySlotValue.Type",NewName="/Script/DynamicMaterial.DMDefaultMaterialPropertySlotValue.DefaultType")
 +PropertyRedirects=(OldName="/Script/DynamicMaterial.DynamicMaterialInstance.MaterialModel",NewName="/Script/DynamicMaterial.DynamicMaterialInstance.MaterialModelBase")
-+FunctionRedirects=(OldName="/Script/DynamicMaterial.DynamicMaterialModel.RemoveValueByName",NewName="/Script/DynamicMaterial.DynamicMaterialModel.RemoveValueByParameterName")
\ No newline at end of file
++FunctionRedirects=(OldName="/Script/DynamicMaterial.DynamicMaterialModel.RemoveValueByName",NewName="/Script/DynamicMaterial.DynamicMaterialModel.RemoveValueByParameterName")
+
++PropertyRedirects=(OldName="/Script/DynamicMaterialEditor.DynamicMaterialModelEditorOnlyData.bTwoSidedFlag",NewName="/Script/DynamicMaterialEditor.DynamicMaterialModelEditorOnlyData.bTwoSided")
++PropertyRedirects=(OldName="/Script/DynamicMaterialEditor.DynamicMaterialModelEditorOnlyData.bPixelAnimationFlag",NewName="/Script/DynamicMaterialEditor.DynamicMaterialModelEditorOnlyData.bHasPixelAnimation")
++FunctionRedirects=(OldName="/Script/DynamicMaterialEditor.DynamicMaterialModelEditorOnlyData.IsPixelAnimationFlagSet",NewName="/Script/DynamicMaterialEditor.DynamicMaterialModelEditorOnlyData.GetHasPixelAnimation")
++FunctionRedirects=(OldName="/Script/DynamicMaterialEditor.DynamicMaterialModelEditorOnlyData.SetPixelAnimationFlag",NewName="/Script/DynamicMaterialEditor.DynamicMaterialModelEditorOnlyData.SetHasPixelAnimation")
++FunctionRedirects=(OldName="/Script/DynamicMaterialEditor.DynamicMaterialModelEditorOnlyData.IsTwoSidedFlagSet",NewName="/Script/DynamicMaterialEditor.DynamicMaterialModelEditorOnlyData.GetIsTwoSided")
++FunctionRedirects=(OldName="/Script/DynamicMateThatrialEditor.DynamicMaterialModelEditorOnlyData.SetTwoSidedFlag",NewName="/Script/DynamicMaterialEditor.DynamicMaterialModelEditorOnlyData.SetIsTwoSided")
\ No newline at end of file
diff --git a/Engine/Plugins/Experimental/DynamicMaterial/Source/DynamicMaterialEditor/Private/Model/DynamicMaterialModelEditorOnlyData.cpp b/Engine/Plugins/Experimental/DynamicMaterial/Source/DynamicMaterialEditor/Private/Model/DynamicMaterialModelEditorOnlyData.cpp
--- a/Engine/Plugins/Experimental/DynamicMaterial/Source/DynamicMaterialEditor/Private/Model/DynamicMaterialModelEditorOnlyData.cpp
+++ b/Engine/Plugins/Experimental/DynamicMaterial/Source/DynamicMaterialEditor/Private/Model/DynamicMaterialModelEditorOnlyData.cpp
@@ -168,8 +168,11 @@ UDynamicMaterialModelEditorOnlyData::UDynamicMaterialModelEditorOnlyData()
 	, Domain(EMaterialDomain::MD_Surface)
 	, BlendMode(EBlendMode::BLEND_Opaque)
 	, ShadingModel(EDMMaterialShadingModel::DefaultLit)
-	, bPixelAnimationFlag(false)
-	, bTwoSidedFlag(true)
+	, bHasPixelAnimation(false)
+	, bTwoSided(true)
+	, bOutputTranslucentVelocityEnabled(true)
+	, bNaniteTessellationEnabled(true)
+	, bResponsiveAAEnabled(true)
 	, ChannelListPreset(NAME_None)
 {
 	Properties.Emplace(EDMMaterialPropertyType::BaseColor,           CreateDefaultSubobject<UDMMaterialPropertyBaseColor>(          "MaterialProperty_BaseColor"));
@@ -246,8 +249,8 @@ void UDynamicMaterialModelEditorOnlyData::OnChannelListPresetChanged()
 
 	SetBlendMode(Preset->DefaultBlendMode);
 	SetShadingModel(Preset->DefaultShadingModel);
-	SetPixelAnimationFlag(Preset->bDefaultAnimated);
-	SetTwoSidedFlag(Preset->bDefaultTwoSided);
+	SetHasPixelAnimation(Preset->bDefaultAnimated);
+	SetIsTwoSided(Preset->bDefaultTwoSided);
 }
 
 void UDynamicMaterialModelEditorOnlyData::EnsurePresetSlots()
@@ -324,13 +327,13 @@ void UDynamicMaterialModelEditorOnlyData::OnBlendModeChanged()
 	switch (BlendMode)
 	{
 		case EBlendMode::BLEND_Opaque:
-			SetPixelAnimationFlag(false);
+			SetHasPixelAnimation(false);
 			RemoveSlotForMaterialProperty(EDMMaterialPropertyType::Opacity);
 			RemoveSlotForMaterialProperty(EDMMaterialPropertyType::OpacityMask);
 			break;
 
 		case EBlendMode::BLEND_Masked:
-			SetPixelAnimationFlag(false);
+			SetHasPixelAnimation(false);
 			EnsureSwapSlotMaterialProperty(EDMMaterialPropertyType::Opacity, EDMMaterialPropertyType::OpacityMask);
 			break;
 
@@ -349,12 +352,7 @@ void UDynamicMaterialModelEditorOnlyData::OnShadingModelChanged()
 	RequestMaterialBuild();
 }
 
-void UDynamicMaterialModelEditorOnlyData::OnPixelAnimationFlagChanged()
-{
-	RequestMaterialBuild();
-}
-
-void UDynamicMaterialModelEditorOnlyData::OnTwoSidedFlagChanged()
+void UDynamicMaterialModelEditorOnlyData::OnMaterialFlagChanged()
 {
 	RequestMaterialBuild();
 }
@@ -418,13 +416,6 @@ void UDynamicMaterialModelEditorOnlyData::CreateMaterial()
 
 		FAssetRegistryModule::AssetCreated(MaterialModel->DynamicMaterial);
 	}
-
-	MaterialModel->DynamicMaterial->bOutputTranslucentVelocity = true;
-	MaterialModel->DynamicMaterial->bEnableResponsiveAA = true;
-
-	// Not setting this to true can cause the level associated with this material to dirty itself when it
-	// is used with Niagara. It doesn't negatively affect the material in any meaningful way.
-	MaterialModel->DynamicMaterial->bUsedWithNiagaraMeshParticles = true;
 }
 
 void UDynamicMaterialModelEditorOnlyData::BuildMaterial(bool bInDirtyAssets)
@@ -451,8 +442,15 @@ void UDynamicMaterialModelEditorOnlyData::BuildMaterial(bool bInDirtyAssets)
 	Expressions.Empty();
 	MaterialModel->DynamicMaterial->MaterialDomain = Domain;
 	MaterialModel->DynamicMaterial->BlendMode = BlendMode;
-	MaterialModel->DynamicMaterial->bHasPixelAnimation = bPixelAnimationFlag;
-	MaterialModel->DynamicMaterial->TwoSided = bTwoSidedFlag;
+	MaterialModel->DynamicMaterial->bHasPixelAnimation = bHasPixelAnimation;
+	MaterialModel->DynamicMaterial->TwoSided = bTwoSided;
+	MaterialModel->DynamicMaterial->bOutputTranslucentVelocity = bOutputTranslucentVelocityEnabled;
+	MaterialModel->DynamicMaterial->bEnableTessellation = bNaniteTessellationEnabled;
+	MaterialModel->DynamicMaterial->bEnableResponsiveAA = bResponsiveAAEnabled;
+
+	// Not setting this to true can cause the level associated with this material to dirty itself when it
+	// is used with Niagara. It doesn't negatively affect the material in any meaningful way.
+	MaterialModel->DynamicMaterial->bUsedWithNiagaraMeshParticles = true;
 
 	switch (ShadingModel)
 	{
@@ -1109,38 +1107,89 @@ void UDynamicMaterialModelEditorOnlyData::SetShadingModel(EDMMaterialShadingMode
 	OnShadingModelChanged();
 }
 
-bool UDynamicMaterialModelEditorOnlyData::IsPixelAnimationFlagSet() const
+bool UDynamicMaterialModelEditorOnlyData::GetHasPixelAnimation() const
+{
+	return bHasPixelAnimation;
+}
+
+void UDynamicMaterialModelEditorOnlyData::SetHasPixelAnimation(bool bInHasAnimation)
+{
+	if (bHasPixelAnimation == bInHasAnimation)
+	{
+		return;
+	}
+
+	bHasPixelAnimation = bInHasAnimation;
+
+	OnMaterialFlagChanged();
+}
+
+bool UDynamicMaterialModelEditorOnlyData::GetIsTwoSided() const
 {
-	return bPixelAnimationFlag;
+	return bTwoSided;
 }
 
-void UDynamicMaterialModelEditorOnlyData::SetPixelAnimationFlag(bool bInFlagValue)
+void UDynamicMaterialModelEditorOnlyData::SetIsTwoSided(bool bInEnabled)
 {
-	if (bPixelAnimationFlag == bInFlagValue)
+	if (bTwoSided == bInEnabled)
 	{
 		return;
 	}
 
-	bPixelAnimationFlag = bInFlagValue;
+	bTwoSided = bInEnabled;
 
-	OnPixelAnimationFlagChanged();
+	OnMaterialFlagChanged();
 }
 
-bool UDynamicMaterialModelEditorOnlyData::IsTwoSidedFlagSet() const
+bool UDynamicMaterialModelEditorOnlyData::IsOutputTranslucentVelocityEnabled() const
 {
-	return bTwoSidedFlag;
+	return bOutputTranslucentVelocityEnabled;
 }
 
-void UDynamicMaterialModelEditorOnlyData::SetTwoSidedFlag(bool bInFlagValue)
+void UDynamicMaterialModelEditorOnlyData::SetOutputTranslucentVelocityEnabled(bool bInEnabled)
 {
-	if (bTwoSidedFlag == bInFlagValue)
+	if (bOutputTranslucentVelocityEnabled == bInEnabled)
 	{
 		return;
 	}
 
-	bTwoSidedFlag = bInFlagValue;
+	bOutputTranslucentVelocityEnabled = bInEnabled;
 
-	OnTwoSidedFlagChanged();
+	OnMaterialFlagChanged();
+}
+
+bool UDynamicMaterialModelEditorOnlyData::IsNaniteTessellationEnabled() const
+{
+	return bNaniteTessellationEnabled;
+}
+
+void UDynamicMaterialModelEditorOnlyData::SetNaniteTessellationEnabled(bool bInEnabled)
+{
+	if (bNaniteTessellationEnabled == bInEnabled)
+	{
+		return;
+	}
+
+	bNaniteTessellationEnabled = bInEnabled;
+
+	OnMaterialFlagChanged();
+}
+
+bool UDynamicMaterialModelEditorOnlyData::IsResponsiveAAEnabled() const
+{
+	return bResponsiveAAEnabled;
+}
+
+void UDynamicMaterialModelEditorOnlyData::SetResponsiveAAEnabled(bool bInEnabled)
+{
+	if (bResponsiveAAEnabled == bInEnabled)
+	{
+		return;
+	}
+
+	bResponsiveAAEnabled = bInEnabled;
+
+	OnMaterialFlagChanged();
 }
 
 FName UDynamicMaterialModelEditorOnlyData::GetChannelListPreset() const
@@ -1488,13 +1537,13 @@ void UDynamicMaterialModelEditorOnlyData::PostEditChangeChainProperty(FPropertyC
 	{
 		OnShadingModelChanged();
 	}
-	else if (Property == GET_MEMBER_NAME_CHECKED(ThisClass, bPixelAnimationFlag))
-	{
-		OnPixelAnimationFlagChanged();
-	}
-	else if (Property == GET_MEMBER_NAME_CHECKED(ThisClass, bTwoSidedFlag))
+	else if (Property == GET_MEMBER_NAME_CHECKED(ThisClass, bHasPixelAnimation)
+		|| Property == GET_MEMBER_NAME_CHECKED(ThisClass, bTwoSided)
+		|| Property == GET_MEMBER_NAME_CHECKED(ThisClass, bOutputTranslucentVelocityEnabled)
+		|| Property == GET_MEMBER_NAME_CHECKED(ThisClass, bNaniteTessellationEnabled)
+		|| Property == GET_MEMBER_NAME_CHECKED(ThisClass, bResponsiveAAEnabled))
 	{
-		OnTwoSidedFlagChanged();
+		OnMaterialFlagChanged();
 	}
 }
 
diff --git a/Engine/Plugins/Experimental/DynamicMaterial/Source/DynamicMaterialEditor/Private/Slate/SDMComponentEdit.cpp b/Engine/Plugins/Experimental/DynamicMaterial/Source/DynamicMaterialEditor/Private/Slate/SDMComponentEdit.cpp
--- a/Engine/Plugins/Experimental/DynamicMaterial/Source/DynamicMaterialEditor/Private/Slate/SDMComponentEdit.cpp
+++ b/Engine/Plugins/Experimental/DynamicMaterial/Source/DynamicMaterialEditor/Private/Slate/SDMComponentEdit.cpp
@@ -1036,8 +1036,11 @@ void SDMComponentEdit::GenerateMaterialModelPropertyRows(const TSharedRef<SDMEdi
 		AddVariable(EditorOnlyData, GET_MEMBER_NAME_CHECKED(UDynamicMaterialModelEditorOnlyData, Domain));
 		AddVariable(EditorOnlyData, GET_MEMBER_NAME_CHECKED(UDynamicMaterialModelEditorOnlyData, BlendMode));
 		AddVariable(EditorOnlyData, GET_MEMBER_NAME_CHECKED(UDynamicMaterialModelEditorOnlyData, ShadingModel));
-		AddVariable(EditorOnlyData, GET_MEMBER_NAME_CHECKED(UDynamicMaterialModelEditorOnlyData, bPixelAnimationFlag));
-		AddVariable(EditorOnlyData, GET_MEMBER_NAME_CHECKED(UDynamicMaterialModelEditorOnlyData, bTwoSidedFlag));
+		AddVariable(EditorOnlyData, GET_MEMBER_NAME_CHECKED(UDynamicMaterialModelEditorOnlyData, bHasPixelAnimation));
+		AddVariable(EditorOnlyData, GET_MEMBER_NAME_CHECKED(UDynamicMaterialModelEditorOnlyData, bTwoSided));
+		AddVariable(EditorOnlyData, GET_MEMBER_NAME_CHECKED(UDynamicMaterialModelEditorOnlyData, bResponsiveAAEnabled));
+		AddVariable(EditorOnlyData, GET_MEMBER_NAME_CHECKED(UDynamicMaterialModelEditorOnlyData, bOutputTranslucentVelocityEnabled));
+		AddVariable(EditorOnlyData, GET_MEMBER_NAME_CHECKED(UDynamicMaterialModelEditorOnlyData, bNaniteTessellationEnabled));
 	}
 }
 
diff --git a/Engine/Plugins/Experimental/DynamicMaterial/Source/DynamicMaterialEditor/Public/Model/DynamicMaterialModelEditorOnlyData.h b/Engine/Plugins/Experimental/DynamicMaterial/Source/DynamicMaterialEditor/Public/Model/DynamicMaterialModelEditorOnlyData.h
--- a/Engine/Plugins/Experimental/DynamicMaterial/Source/DynamicMaterialEditor/Public/Model/DynamicMaterialModelEditorOnlyData.h
+++ b/Engine/Plugins/Experimental/DynamicMaterial/Source/DynamicMaterialEditor/Public/Model/DynamicMaterialModelEditorOnlyData.h
@@ -116,16 +116,34 @@ public:
 	DYNAMICMATERIALEDITOR_API void SetShadingModel(EDMMaterialShadingModel InShadingModel);
 
 	UFUNCTION(BlueprintPure, Category = "Material Designer")
-	DYNAMICMATERIALEDITOR_API bool IsPixelAnimationFlagSet() const;
+	DYNAMICMATERIALEDITOR_API bool GetHasPixelAnimation() const;
 
 	UFUNCTION(BlueprintCallable, Category = "Material Designer")
-	DYNAMICMATERIALEDITOR_API void SetPixelAnimationFlag(bool bInFlagValue);
+	DYNAMICMATERIALEDITOR_API void SetHasPixelAnimation(bool bInHasAnimation);
 
 	UFUNCTION(BlueprintPure, Category = "Material Designer")
-	DYNAMICMATERIALEDITOR_API bool IsTwoSidedFlagSet() const;
+	DYNAMICMATERIALEDITOR_API bool GetIsTwoSided() const;
 
 	UFUNCTION(BlueprintCallable, Category = "Material Designer")
-	DYNAMICMATERIALEDITOR_API void SetTwoSidedFlag(bool bInFlagValue);
+	DYNAMICMATERIALEDITOR_API void SetIsTwoSided(bool bInEnabled);
+
+	UFUNCTION(BlueprintPure, Category = "Material Designer")
+	DYNAMICMATERIALEDITOR_API bool IsOutputTranslucentVelocityEnabled() const;
+
+	UFUNCTION(BlueprintCallable, Category = "Material Designer")
+	DYNAMICMATERIALEDITOR_API void SetOutputTranslucentVelocityEnabled(bool bInEnabled);
+
+	UFUNCTION(BlueprintPure, Category = "Material Designer")
+	DYNAMICMATERIALEDITOR_API bool IsNaniteTessellationEnabled() const;
+
+	UFUNCTION(BlueprintCallable, Category = "Material Designer")
+	DYNAMICMATERIALEDITOR_API void SetNaniteTessellationEnabled(bool bInEnabled);
+
+	UFUNCTION(BlueprintPure, Category = "Material Designer")
+	DYNAMICMATERIALEDITOR_API bool IsResponsiveAAEnabled() const;
+
+	UFUNCTION(BlueprintCallable, Category = "Material Designer")
+	DYNAMICMATERIALEDITOR_API void SetResponsiveAAEnabled(bool bInEnabled);
 
 	UFUNCTION(BlueprintCallable, Category = "Material Designer")
 	DYNAMICMATERIALEDITOR_API FName GetChannelListPreset() const;
@@ -248,15 +266,35 @@ protected:
 	UPROPERTY(EditInstanceOnly, BlueprintReadOnly, Category = "Material Designer", meta = (ValidEnumValues = "Unlit,DefaultLit"))
 	EDMMaterialShadingModel ShadingModel;
 
+	/**
+	 * Whether the opaque material has any pixel animations happening, that isn't included in the geometric velocities.
+	 * This allows to disable renderer's heuristics that assumes animation is fully described with motion vector, such as TSR's anti-flickering heuristic.
+	 */
 	UPROPERTY(EditInstanceOnly, BlueprintReadOnly, Category = "Material Designer")
-	bool bPixelAnimationFlag;
+	bool bHasPixelAnimation;
 
+	/** Indicates that the material should be rendered without backface culling and the normal should be flipped for backfaces. */
 	UPROPERTY(EditInstanceOnly, BlueprintReadOnly, Category = "Material Designer")
-	bool bTwoSidedFlag;
+	bool bTwoSided;
+
+	/** When true, translucent materials will output motion vectors and write to depth buffer in velocity pass. */
+	UPROPERTY(EditInstanceOnly, BlueprintReadOnly, Category = "Material Designer", meta = (DisplayName = "Output Translucent Velocity"))
+	bool bOutputTranslucentVelocityEnabled;
+
+	/** Whether tessellation is enabled on the material. NOTE: Required for displacement to work. */
+	UPROPERTY(EditInstanceOnly, BlueprintReadOnly, Category = "Material Designer", meta = (DisplayName = "Nanite Tessellation"))
+	bool bNaniteTessellationEnabled;
+
+	/**
+	 * Indicates that the material should be rendered using responsive anti-aliasing. Improves sharpness of small moving particles such as sparks.
+	 * Only use for small moving features because it will cause aliasing of the background.
+	 */
+	UPROPERTY(EditInstanceOnly, BlueprintReadOnly, Category = "Material Designer", meta = (DisplayName = "Responsive AA"))
+	bool bResponsiveAAEnabled;
 
+	/** Sets the available channels and default material properties. */
 	UPROPERTY(EditInstanceOnly, BlueprintReadOnly, Category = "Material Designer", 
-		meta = (GetOptions = GetPresetOptions, NoResetToDefault, DisplayName = "Material Type Preset",
-			ToolTip = "Sets the available channels and default material properties."))
+		meta = (GetOptions = GetPresetOptions, NoResetToDefault, DisplayName = "Material Type Preset"))
 	FName ChannelListPreset;
 
 	UPROPERTY(VisibleInstanceOnly, BlueprintReadOnly, Instanced, Category = "Material Designer")
@@ -308,8 +346,7 @@ protected:
 	void OnDomainChanged();
 	void OnBlendModeChanged();
 	void OnShadingModelChanged();
-	void OnPixelAnimationFlagChanged();
-	void OnTwoSidedFlagChanged();
+	void OnMaterialFlagChanged();
 
 	//~ Begin IDynamicMaterialModelEditorOnlyDataInterface
 	virtual void ReinitComponents() override;
