Material Designer
- Fixed crash when creating a Material Designer Dynamic asset.

diff --git a/Engine/Plugins/Experimental/DynamicMaterial/Source/DynamicMaterial/Private/Components/DMTextureUVDynamic.cpp b/Engine/Plugins/Experimental/DynamicMaterial/Source/DynamicMaterial/Private/Components/DMTextureUVDynamic.cpp
--- a/Engine/Plugins/Experimental/DynamicMaterial/Source/DynamicMaterial/Private/Components/DMTextureUVDynamic.cpp
+++ b/Engine/Plugins/Experimental/DynamicMaterial/Source/DynamicMaterial/Private/Components/DMTextureUVDynamic.cpp
@@ -117,9 +117,6 @@ void UDMTextureUVDynamic::SetMIDParameters(UMaterialInstanceDynamic* InMID) cons
 	TArray<UDMMaterialParameter*> MaterialParameters = ParentTextureUV->GetParameters();
 
 	check(InMID);
-	check(MaterialParameters.IsEmpty() == false);
-
-	using namespace UE::DynamicMaterial;
 
 	auto UpdateMID = [InMID](FName InParamName, float InValue)
 		{
@@ -129,13 +126,13 @@ void UDMTextureUVDynamic::SetMIDParameters(UMaterialInstanceDynamic* InMID) cons
 			}
 		};
 
-	UpdateMID(MaterialParameters[ParamID::PivotX]->GetParameterName(), GetPivot().X);
-	UpdateMID(MaterialParameters[ParamID::PivotY]->GetParameterName(), GetPivot().Y);
-	UpdateMID(MaterialParameters[ParamID::TilingX]->GetParameterName(), GetTiling().X);
-	UpdateMID(MaterialParameters[ParamID::TilingY]->GetParameterName(), GetTiling().Y);
-	UpdateMID(MaterialParameters[ParamID::Rotation]->GetParameterName(), GetRotation());
-	UpdateMID(MaterialParameters[ParamID::OffsetX]->GetParameterName(), GetOffset().X);
-	UpdateMID(MaterialParameters[ParamID::OffsetY]->GetParameterName(), GetOffset().Y);
+	UpdateMID(ParentTextureUV->GetMaterialParameterName(UDMTextureUV::NAME_Offset, 0), GetOffset().X);
+	UpdateMID(ParentTextureUV->GetMaterialParameterName(UDMTextureUV::NAME_Offset, 1), GetOffset().Y);
+	UpdateMID(ParentTextureUV->GetMaterialParameterName(UDMTextureUV::NAME_Rotation, 0), GetRotation());
+	UpdateMID(ParentTextureUV->GetMaterialParameterName(UDMTextureUV::NAME_Pivot, 0), GetPivot().X);
+	UpdateMID(ParentTextureUV->GetMaterialParameterName(UDMTextureUV::NAME_Pivot, 1), GetPivot().Y);
+	UpdateMID(ParentTextureUV->GetMaterialParameterName(UDMTextureUV::NAME_Tiling, 0), GetTiling().X);
+	UpdateMID(ParentTextureUV->GetMaterialParameterName(UDMTextureUV::NAME_Tiling, 1), GetTiling().Y);
 }
 
 #if WITH_EDITOR
