Material Designer
- Changed the default texture for opacity slots to the blank, white texture.

diff --git a/Engine/Plugins/Experimental/DynamicMaterial/Source/DynamicMaterialEditor/Private/DynamicMaterialEditorSettings.cpp b/Engine/Plugins/Experimental/DynamicMaterial/Source/DynamicMaterialEditor/Private/DynamicMaterialEditorSettings.cpp
--- a/Engine/Plugins/Experimental/DynamicMaterial/Source/DynamicMaterialEditor/Private/DynamicMaterialEditorSettings.cpp
+++ b/Engine/Plugins/Experimental/DynamicMaterial/Source/DynamicMaterialEditor/Private/DynamicMaterialEditorSettings.cpp
@@ -20,10 +20,10 @@
 namespace UE::DynamicMaterialEditor::Private
 {
 	const TMap<EDMMaterialPropertyType, FDMDefaultMaterialPropertySlotValue> DefaultSlotValues = {
-		{EDMMaterialPropertyType::BaseColor,           TSoftObjectPtr<UTexture>(FSoftObjectPath(TEXT("Texture2D'/DynamicMaterial/T_Default_Texture.T_Default_Texture'")))},
+		{EDMMaterialPropertyType::BaseColor,           TSoftObjectPtr<UTexture>(FSoftObjectPath(TEXT("/Script/Engine.Texture2D'/DynamicMaterial/T_Default_Texture.T_Default_Texture'")))},
 		{EDMMaterialPropertyType::EmissiveColor,       TSoftObjectPtr<UTexture>(FSoftObjectPath(TEXT("/Script/Engine.Texture2D'/DynamicMaterial/Textures/SlotDefaults/T_MD_Emissive.T_MD_Emissive'")))},
-		{EDMMaterialPropertyType::Opacity,             TSoftObjectPtr<UTexture>(FSoftObjectPath(TEXT("/Script/Engine.Texture2D'/DynamicMaterial/Textures/T_DM_HorizontalGradient.T_DM_HorizontalGradient'")))},
-		{EDMMaterialPropertyType::OpacityMask,         TSoftObjectPtr<UTexture>(FSoftObjectPath(TEXT("/Script/Engine.Texture2D'/DynamicMaterial/Textures/T_DM_HorizontalGradient.T_DM_HorizontalGradient'")))},
+		{EDMMaterialPropertyType::Opacity,             TSoftObjectPtr<UTexture>(FSoftObjectPath(TEXT("/Script/Engine.Texture2D'/Engine/EngineResources/WhiteSquareTexture.WhiteSquareTexture'")))},
+		{EDMMaterialPropertyType::OpacityMask,         TSoftObjectPtr<UTexture>(FSoftObjectPath(TEXT("/Script/Engine.Texture2D'/Engine/EngineResources/WhiteSquareTexture.WhiteSquareTexture'")))},
 		{EDMMaterialPropertyType::Metallic,            TSoftObjectPtr<UTexture>(FSoftObjectPath(TEXT("/Script/Engine.Texture2D'/DynamicMaterial/Textures/SlotDefaults/T_MD_Metallic.T_MD_Metallic'")))},
 		{EDMMaterialPropertyType::Specular,            FLinearColor(0.5f, 0.5f, 0.5f, 1.f)},
 		{EDMMaterialPropertyType::Roughness,           TSoftObjectPtr<UTexture>(FSoftObjectPath(TEXT("/Script/Engine.Texture2D'/DynamicMaterial/Textures/SlotDefaults/T_MD_Roughness.T_MD_Roughness'")))},
@@ -122,7 +122,7 @@ UDynamicMaterialEditorSettings::UDynamicMaterialEditorSettings()
 
 	ResetAllLayoutSettings();
 
-	DefaultMask = TSoftObjectPtr<UTexture>(FSoftObjectPath(TEXT("Texture2D'/Engine/EngineResources/WhiteSquareTexture.WhiteSquareTexture'")));
+	DefaultMask = TSoftObjectPtr<UTexture>(FSoftObjectPath(TEXT("/Script/Engine.Texture2D'/Engine/EngineResources/WhiteSquareTexture.WhiteSquareTexture'")));
 
 	FDMMaterialChannelListPreset Opaque;
 	Opaque.Name = TEXT("Opaque");
