Motion Design
- When using "Add Texture to Scene", the material type preset is now set to "Emissive".
- New actor is now uses the texture's name.

diff --git a/Engine/Plugins/Experimental/Avalanche/Source/AvalancheEditor/Private/MaterialDesigner/AvaMaterialDesignerTextureAssetFactory.cpp b/Engine/Plugins/Experimental/Avalanche/Source/AvalancheEditor/Private/MaterialDesigner/AvaMaterialDesignerTextureAssetFactory.cpp
--- a/Engine/Plugins/Experimental/Avalanche/Source/AvalancheEditor/Private/MaterialDesigner/AvaMaterialDesignerTextureAssetFactory.cpp
+++ b/Engine/Plugins/Experimental/Avalanche/Source/AvalancheEditor/Private/MaterialDesigner/AvaMaterialDesignerTextureAssetFactory.cpp
@@ -1,6 +1,7 @@
 // Copyright Epic Games, Inc. All Rights Reserved.
 
 #include "MaterialDesigner/AvaMaterialDesignerTextureAssetFactory.h"
+
 #include "AvaShapeActor.h"
 #include "Components/DMMaterialLayer.h"
 #include "Components/DMMaterialSlot.h"
@@ -57,6 +58,8 @@ AActor* UAvaMaterialDesignerTextureAssetFactory::SpawnActor(UObject* InAsset, UL
 		return InNewActor;
 	}
 
+	InNewActor->Rename(*Texture->GetName(), nullptr, REN_DontCreateRedirectors);
+
 	AAvaShapeActor* ShapeActor = Cast<AAvaShapeActor>(InNewActor);
 
 	if (!ShapeActor)
@@ -121,7 +124,10 @@ AActor* UAvaMaterialDesignerTextureAssetFactory::SpawnActor(UObject* InAsset, UL
 		return InNewActor;
 	}
 
-	const UDMMaterialSlot* Slot = EditorOnlyData->GetSlotForMaterialProperty(EDMMaterialPropertyType::BaseColor);
+	EditorOnlyData->SetChannelListPreset(TEXT("Emissive"));
+	EditorOnlyData->OnWizardComplete();
+
+	const UDMMaterialSlot* Slot = EditorOnlyData->GetSlotForMaterialProperty(EDMMaterialPropertyType::EmissiveColor);
 
 	if (!Slot)
 	{
diff --git a/Engine/Plugins/Experimental/DynamicMaterial/Source/DynamicMaterialEditor/Private/Slate/SDMMaterialWizard.cpp b/Engine/Plugins/Experimental/DynamicMaterial/Source/DynamicMaterialEditor/Private/Slate/SDMMaterialWizard.cpp
--- a/Engine/Plugins/Experimental/DynamicMaterial/Source/DynamicMaterialEditor/Private/Slate/SDMMaterialWizard.cpp
+++ b/Engine/Plugins/Experimental/DynamicMaterial/Source/DynamicMaterialEditor/Private/Slate/SDMMaterialWizard.cpp
@@ -40,12 +40,21 @@ void SDMMaterialWizard::Construct(const FArguments& InArgs, const TSharedRef<SDM
 		{
 			CurrentPreset = Settings->MaterialChannelPresets[0].Name;
 		}
-	}	
+	}
+
+	if (UDynamicMaterialModel* MaterialModel = Cast<UDynamicMaterialModel>(InEditor->GetMaterialModelBase()))
+	{
+		if (UDynamicMaterialModelEditorOnlyData* EditorOnlyData = UDynamicMaterialModelEditorOnlyData::Get(MaterialModel))
+		{
+			// Subscribe to this in case the wizard completes externally and this widget is no longer needed.
+			EditorOnlyData->GetOnMaterialBuiltDelegate().AddSP(this, &SDMMaterialWizard::OnMaterialBuilt);
+		}
+	}
 
 	ChildSlot
-		[
-			CreateLayout()
-		];
+	[
+		CreateLayout()
+	];
 }
 
 TSharedPtr<SDMEditor> SDMMaterialWizard::GetEditor() const
@@ -253,13 +262,44 @@ FReply SDMMaterialWizard::Accept_OnClick()
 		{
 			if (UDynamicMaterialModelEditorOnlyData* EditorOnlyData = UDynamicMaterialModelEditorOnlyData::Get(MaterialModel))
 			{
-				AActor* MaterialActor = Editor->GetMaterialActor();
-				FDMObjectMaterialProperty MaterialProperty = Editor->GetMaterialObjectProperty();
-				Editor->ClearEditor();
+				EditorOnlyData->GetOnMaterialBuiltDelegate().RemoveAll(this);
 
 				EditorOnlyData->SetChannelListPreset(CurrentPreset);
 				EditorOnlyData->OnWizardComplete();
 
+				OpenMaterialInEditor(MaterialModel);
+			}
+		}
+	}
+
+	return FReply::Handled();
+}
+
+void SDMMaterialWizard::OnMaterialBuilt(UDynamicMaterialModelBase* InMaterialModel)
+{
+	if (UDynamicMaterialModelEditorOnlyData* EditorOnlyData = UDynamicMaterialModelEditorOnlyData::Get(InMaterialModel))
+	{
+		if (!EditorOnlyData->NeedsWizard())
+		{
+			OpenMaterialInEditor(InMaterialModel);
+		}
+	}
+}
+
+void SDMMaterialWizard::OpenMaterialInEditor(UDynamicMaterialModelBase* InMaterialModel)
+{
+	if (TSharedPtr<SDMEditor> Editor = GetEditor())
+	{
+		if (UDynamicMaterialModel* MaterialModel = GetMaterialModel())
+		{
+			if (UDynamicMaterialModelEditorOnlyData* EditorOnlyData = UDynamicMaterialModelEditorOnlyData::Get(MaterialModel))
+			{
+				EditorOnlyData->GetOnMaterialBuiltDelegate().RemoveAll(this);
+
+				AActor* MaterialActor = Editor->GetMaterialActor();
+				FDMObjectMaterialProperty MaterialProperty = Editor->GetMaterialObjectProperty();
+				Editor->ClearEditor();
+
 				// Refresh display
 				if (MaterialProperty.IsValid())
 				{
@@ -274,8 +314,6 @@ FReply SDMMaterialWizard::Accept_OnClick()
 			}
 		}
 	}
-
-	return FReply::Handled();
 }
 
 #undef LOCTEXT_NAMESPACE
diff --git a/Engine/Plugins/Experimental/DynamicMaterial/Source/DynamicMaterialEditor/Private/Slate/SDMMaterialWizard.h b/Engine/Plugins/Experimental/DynamicMaterial/Source/DynamicMaterialEditor/Private/Slate/SDMMaterialWizard.h
--- a/Engine/Plugins/Experimental/DynamicMaterial/Source/DynamicMaterialEditor/Private/Slate/SDMMaterialWizard.h
+++ b/Engine/Plugins/Experimental/DynamicMaterial/Source/DynamicMaterialEditor/Private/Slate/SDMMaterialWizard.h
@@ -9,10 +9,9 @@ class SBox;
 class SDMEditor;
 class SWidget;
 class UDynamicMaterialModel;
+class UDynamicMaterialModelBase;
 enum class ECheckBoxState : uint8;
 enum class EDMMaterialShadingModel : uint8;
-enum EBlendMode : int;
-enum EMaterialDomain : int;
 
 class SDMMaterialWizard : public SCompoundWidget
 {
@@ -40,4 +39,8 @@ protected:
 	void Preset_OnChange(ECheckBoxState InState, FName InPresetName);
 
 	FReply Accept_OnClick();
+
+	void OnMaterialBuilt(UDynamicMaterialModelBase* InMaterialModel);
+
+	void OpenMaterialInEditor(UDynamicMaterialModelBase* InMaterialModel);
 };
diff --git a/Engine/Plugins/Experimental/DynamicMaterial/Source/DynamicMaterialEditor/Public/Model/DynamicMaterialModelEditorOnlyData.h b/Engine/Plugins/Experimental/DynamicMaterial/Source/DynamicMaterialEditor/Public/Model/DynamicMaterialModelEditorOnlyData.h
--- a/Engine/Plugins/Experimental/DynamicMaterial/Source/DynamicMaterialEditor/Public/Model/DynamicMaterialModelEditorOnlyData.h
+++ b/Engine/Plugins/Experimental/DynamicMaterial/Source/DynamicMaterialEditor/Public/Model/DynamicMaterialModelEditorOnlyData.h
@@ -210,7 +210,7 @@ public:
 
 	bool NeedsWizard() const;
 
-	void OnWizardComplete();
+	DYNAMICMATERIALEDITOR_API void OnWizardComplete();
 
 	DYNAMICMATERIALEDITOR_API void ForEachMaterialPropertyType(TFunctionRef<EDMIterationResult(EDMMaterialPropertyType InType)> InCallable,
 		EDMMaterialPropertyType InStart = static_cast<EDMMaterialPropertyType>(static_cast<uint8>(EDMMaterialPropertyType::None) + 1),
