Material Designer
- Text source crash fixed.
- Added default font and text.

diff --git a/Engine/Plugins/Experimental/DynamicMaterial/Source/DynamicMaterial/Private/Components/RenderTargetRenderers/DMRenderTargetWidgetRendererBase.cpp b/Engine/Plugins/Experimental/DynamicMaterial/Source/DynamicMaterial/Private/Components/RenderTargetRenderers/DMRenderTargetWidgetRendererBase.cpp
--- a/Engine/Plugins/Experimental/DynamicMaterial/Source/DynamicMaterial/Private/Components/RenderTargetRenderers/DMRenderTargetWidgetRendererBase.cpp
+++ b/Engine/Plugins/Experimental/DynamicMaterial/Source/DynamicMaterial/Private/Components/RenderTargetRenderers/DMRenderTargetWidgetRendererBase.cpp
@@ -18,8 +18,6 @@ UDMRenderTargetWidgetRendererBase::UDMRenderTargetWidgetRendererBase()
 
 void UDMRenderTargetWidgetRendererBase::UpdateRenderTarget_Internal()
 {
-	Super::UpdateRenderTarget_Internal();
-
 	UDMMaterialValueRenderTarget* RenderTargetValue = GetRenderTargetValue();
 
 	if (!RenderTargetValue)
