Material Designer
- Added default font and text.
- Fixed text not rendering correctly (repeating)
- Fixed text not rendering correctly (missing lines)
- Moved "override render target size" to the render target category.

diff --git a/Engine/Plugins/Experimental/DynamicMaterial/Source/DynamicMaterial/Private/Components/RenderTargetRenderers/DMRenderTargetTextRenderer.cpp b/Engine/Plugins/Experimental/DynamicMaterial/Source/DynamicMaterial/Private/Components/RenderTargetRenderers/DMRenderTargetTextRenderer.cpp
--- a/Engine/Plugins/Experimental/DynamicMaterial/Source/DynamicMaterial/Private/Components/RenderTargetRenderers/DMRenderTargetTextRenderer.cpp
+++ b/Engine/Plugins/Experimental/DynamicMaterial/Source/DynamicMaterial/Private/Components/RenderTargetRenderers/DMRenderTargetTextRenderer.cpp
@@ -5,6 +5,7 @@
 #include "Dom/JsonObject.h"
 #include "Dom/JsonValue.h"
 #include "Engine/Engine.h"
+#include "Engine/Font.h"
 #include "TextureResource.h"
 #include "Widgets/SBoxPanel.h"
 #include "Widgets/Text/STextBlock.h"
@@ -19,6 +20,12 @@ namespace UE::DynamicMaterial::Private
 {
 	const FIntPoint MinimumTextTextureSize = FIntPoint(5, 10);
 	const FIntPoint MaxTexScale = FIntPoint(64, 64);
+
+	UFont* GetDefaultFont()
+	{
+		static TSoftObjectPtr<UFont> DefaultFont = TSoftObjectPtr<UFont>(FSoftObjectPath(TEXT("/Script/Engine.Font'/Engine/EngineFonts/Roboto.Roboto'")));
+		return DefaultFont.LoadSynchronous();
+	}
 }
 
 #if WITH_EDITOR
@@ -66,6 +73,9 @@ struct FDMRenderTargetTextRenderer
 
 UDMRenderTargetTextRenderer::UDMRenderTargetTextRenderer()
 {
+	FontInfo.FontObject = UE::DynamicMaterial::Private::GetDefaultFont();
+	Text = LOCTEXT("Text", "Text");
+
 #if WITH_EDITOR
 	EditableProperties.Append(FDMRenderTargetTextRenderer::PropertyNames.Array());
 #endif
@@ -677,23 +687,22 @@ void UDMRenderTargetTextRenderer::UpdateTextLines()
 	TArray<FString> NewLines;
 	TextStr.ParseIntoArray(NewLines, TEXT("\n"), /* Cull empty */ false);
 
-	Lines.Empty();
-	Lines.Reserve(NewLines.Num());
-
 	for (FString& NewLine : NewLines)
 	{
 		if (NewLine.EndsWith(TEXT("\r")))
 		{
-			NewLine.LeftChopInline(1);
+			NewLine.LeftChopInline(1, EAllowShrinking::No);
 		}
+	}
 
-		TSharedRef<STextBlock> TextBlock = CreateTextWidget(FText::FromString(NewLine));
+	Lines.Empty();
+	Lines.Reserve(NewLines.Num());
 
-		FDMTextLine Line;
+	for (const FString& NewLine : NewLines)
+	{
+		FDMTextLine& Line = Lines.AddDefaulted_GetRef();
 		Line.Line = NewLine;
-		Line.Widget = TextBlock;
-
-		Lines.Add(Line);
+		Line.Widget = CreateTextWidget(FText::FromString(NewLine));
 	}	
 
 	bRecalculateTextSize = true;
@@ -704,7 +713,7 @@ TSharedRef<STextBlock> UDMRenderTargetTextRenderer::CreateTextWidget(const FText
 {
 	return SNew(STextBlock)
 		.Font(FontInfo)
-		.Text(Text)
+		.Text(InText)
 		.LineHeightPercentage(LineHeight)
 		.ColorAndOpacity(TextColor)
 		.HighlightColor(bHasHighlight ? HighlightColor : TAttribute<FLinearColor>())
@@ -759,7 +768,7 @@ void UDMRenderTargetTextRenderer::SetCustomTextureSize()
 	Widget->SlatePrepass(1.f);
 	const FVector2D Size = Widget->GetDesiredSize();
 
-	RenderTargetValue->SetTextureSize(FIntPoint(FMath::RoundToInt(Size.X), FMath::RoundToInt(Size.Y)));
+	RenderTargetValue->SetTextureSize(FIntPoint(FMath::CeilToInt32(Size.X), FMath::CeilToInt32(Size.Y)));
 	RenderTargetValue->FlushCreateRenderTarget();
 }
 
diff --git a/Engine/Plugins/Experimental/DynamicMaterial/Source/DynamicMaterial/Private/Components/RenderTargetRenderers/DMRenderTargetWidgetRendererBase.cpp b/Engine/Plugins/Experimental/DynamicMaterial/Source/DynamicMaterial/Private/Components/RenderTargetRenderers/DMRenderTargetWidgetRendererBase.cpp
--- a/Engine/Plugins/Experimental/DynamicMaterial/Source/DynamicMaterial/Private/Components/RenderTargetRenderers/DMRenderTargetWidgetRendererBase.cpp
+++ b/Engine/Plugins/Experimental/DynamicMaterial/Source/DynamicMaterial/Private/Components/RenderTargetRenderers/DMRenderTargetWidgetRendererBase.cpp
@@ -47,7 +47,7 @@ void UDMRenderTargetWidgetRendererBase::UpdateRenderTarget_Internal()
 	WidgetRenderer->DrawWidget(
 		RenderTarget,
 		Widget.ToSharedRef(),
-		{(double)RenderTarget->SizeX, (double)RenderTarget->SizeX},
+		{(double)RenderTarget->SizeX, (double)RenderTarget->SizeY},
 		/* Delta Time */ 0.f
 	);
 }
diff --git a/Engine/Plugins/Experimental/DynamicMaterial/Source/DynamicMaterial/Public/Components/RenderTargetRenderers/DMRenderTargetTextRenderer.h b/Engine/Plugins/Experimental/DynamicMaterial/Source/DynamicMaterial/Public/Components/RenderTargetRenderers/DMRenderTargetTextRenderer.h
--- a/Engine/Plugins/Experimental/DynamicMaterial/Source/DynamicMaterial/Public/Components/RenderTargetRenderers/DMRenderTargetTextRenderer.h
+++ b/Engine/Plugins/Experimental/DynamicMaterial/Source/DynamicMaterial/Public/Components/RenderTargetRenderers/DMRenderTargetTextRenderer.h
@@ -287,7 +287,7 @@ protected:
 	float PaddingBottom = 0.f;
 
 	/** When true, will change the size of the render target to fit the text. */
-	UPROPERTY(EditInstanceOnly, BlueprintReadWrite, Getter = IsOverridingRenderTargetSize, Setter = SetOverrideRenderTargetSize, BlueprintSetter = SetOverrideRenderTargetSize, Category = "Material Designer|Text",
+	UPROPERTY(EditInstanceOnly, BlueprintReadWrite, Getter = IsOverridingRenderTargetSize, Setter = SetOverrideRenderTargetSize, BlueprintSetter = SetOverrideRenderTargetSize, Category = "Material Designer|Render Target",
 		meta = (NotKeyframeable, AllowPrivateAccess = "true"))
 	bool bOverrideRenderTargetSize = true;
 
