Material Designer
- Reworked gradients.
- Converted to material functions.
- Added a 'start' and 'end' value.

#virtualized

diff --git a/Engine/Plugins/Experimental/DynamicMaterial/Source/DynamicMaterial/Public/Components/MaterialValues/DMMaterialValueFloat3RGB.h b/Engine/Plugins/Experimental/DynamicMaterial/Source/DynamicMaterial/Public/Components/MaterialValues/DMMaterialValueFloat3RGB.h
--- a/Engine/Plugins/Experimental/DynamicMaterial/Source/DynamicMaterial/Public/Components/MaterialValues/DMMaterialValueFloat3RGB.h
+++ b/Engine/Plugins/Experimental/DynamicMaterial/Source/DynamicMaterial/Public/Components/MaterialValues/DMMaterialValueFloat3RGB.h
@@ -62,7 +62,7 @@ protected:
  
 #if WITH_EDITORONLY_DATA
 	UPROPERTY(EditInstanceOnly, BlueprintReadWrite, BlueprintSetter = SetDefaultValue, Category = "Material Designer",
-		meta = (AllowPrivateAccess = "true"))
+		meta = (AllowPrivateAccess = "true", HideAlphaChannel))
 	FLinearColor DefaultValue;
 #endif
 };
diff --git a/Engine/Plugins/Experimental/DynamicMaterial/Source/DynamicMaterialEditor/Private/Components/DMMaterialStageGradient.cpp b/Engine/Plugins/Experimental/DynamicMaterial/Source/DynamicMaterialEditor/Private/Components/DMMaterialStageGradient.cpp
--- a/Engine/Plugins/Experimental/DynamicMaterial/Source/DynamicMaterialEditor/Private/Components/DMMaterialStageGradient.cpp
+++ b/Engine/Plugins/Experimental/DynamicMaterial/Source/DynamicMaterialEditor/Private/Components/DMMaterialStageGradient.cpp
@@ -1,8 +1,16 @@
 // Copyright Epic Games, Inc. All Rights Reserved.
 
 #include "Components/DMMaterialStageGradient.h"
+
 #include "Components/DMMaterialSlot.h"
 #include "Components/DMMaterialStage.h"
+#include "Components/DMMaterialStageFunction.h"
+#include "Components/MaterialStageInputs/DMMSIValue.h"
+#include "Components/MaterialValues/DMMaterialValueFloat3RGB.h"
+#include "Materials/MaterialExpressionMaterialFunctionCall.h"
+#include "Materials/MaterialFunctionInterface.h"
+#include "Model/DMMaterialBuildState.h"
+#include "Utils/DMUtils.h"
 
 #define LOCTEXT_NAMESPACE "DMMaterialProperty"
 
@@ -18,11 +26,32 @@ UDMMaterialStageGradient::UDMMaterialStageGradient(const FText& InName)
 {
 	bAllowNestedInputs = true;
 
-	InputConnectors.Add({0, LOCTEXT("UV", "UV"), EDMValueType::VT_Float2});
+	InputConnectors.Add({InputUV, LOCTEXT("UV", "UV"), EDMValueType::VT_Float2});
+	InputConnectors.Add({InputStart, LOCTEXT("Start", "Start"), EDMValueType::VT_Float3_RGB});
+	InputConnectors.Add({InputEnd, LOCTEXT("End", "End"), EDMValueType::VT_Float3_RGB});
 
 	OutputConnectors.Add({0, LOCTEXT("Value", "Value"), EDMValueType::VT_Float3_RGB});
 }
 
+bool UDMMaterialStageGradient::SetMaterialFunction(UMaterialFunctionInterface* InMaterialFunction)
+{
+	if (!IsComponentValid())
+	{
+		return false;
+	}
+
+	if (InMaterialFunction == MaterialFunction)
+	{
+		return false;
+	}
+
+	MaterialFunction = InMaterialFunction;
+
+	Update(EDMUpdateType::Structure | EDMUpdateType::AllowParentUpdate);
+
+	return true;
+}
+
 UDMMaterialStage* UDMMaterialStageGradient::CreateStage(TSubclassOf<UDMMaterialStageGradient> InMaterialStageGradientClass, UDMMaterialLayerObject* InLayer)
 {
 	check(InMaterialStageGradientClass);
@@ -76,12 +105,76 @@ UDMMaterialStageGradient* UDMMaterialStageGradient::ChangeStageSource_Gradient(U
 
 bool UDMMaterialStageGradient::CanChangeInputType(int32 InputIndex) const
 {
-	if (InputIndex == 0)
+	return false;
+}
+
+void UDMMaterialStageGradient::AddDefaultInput(int32 InInputIndex) const
+{
+	check(InputConnectors.IsValidIndex(InInputIndex));
+
+	UDMMaterialStage* Stage = GetStage();
+	check(Stage);
+
+	if (InInputIndex != InputStart && InInputIndex != InputEnd)
 	{
-		return false;
+		Super::AddDefaultInput(InInputIndex);
+		return;
+	}
+
+	UDMMaterialStageInputValue * InputValue = UDMMaterialStageInputValue::ChangeStageInput_NewLocalValue(
+		Stage,
+		InInputIndex,
+		FDMMaterialStageConnectorChannel::WHOLE_CHANNEL,
+		InputConnectors[InInputIndex].Type,
+		FDMMaterialStageConnectorChannel::WHOLE_CHANNEL
+	);
+	check(InputValue);
+
+	UDMMaterialValueFloat3RGB* Value = Cast<UDMMaterialValueFloat3RGB>(InputValue->GetValue());
+
+	switch (InInputIndex)
+	{
+		case InputStart:
+			Value->SetDefaultValue(FLinearColor::Black);
+			break;
+
+		case InputEnd:
+			Value->SetDefaultValue(FLinearColor::White);
+			break;
+	}
+
+	Value->ApplyDefaultValue();
+}
+
+void UDMMaterialStageGradient::GenerateExpressions(const TSharedRef<FDMMaterialBuildState>& InBuildState) const
+{
+	if (!IsComponentValid() || !IsComponentAdded())
+	{
+		return;
 	}
 
-	return Super::CanChangeInput(InputIndex);
+	if (InBuildState->HasStageSource(this))
+	{
+		return;
+	}
+
+	UMaterialFunctionInterface* ActualMaterialFunction = MaterialFunction;
+
+	if (!IsValid(ActualMaterialFunction))
+	{
+		ActualMaterialFunction = UDMMaterialStageFunction::NoOp.LoadSynchronous();
+	}
+
+	if (!ActualMaterialFunction)
+	{
+		return;
+	}
+
+	UMaterialExpressionMaterialFunctionCall* FunctionCall = InBuildState->GetBuildUtils().CreateExpression<UMaterialExpressionMaterialFunctionCall>(UE_DM_NodeComment_Default);
+	FunctionCall->SetMaterialFunction(ActualMaterialFunction);
+	FunctionCall->UpdateFromFunctionResource();
+
+	InBuildState->AddStageSourceExpressions(this, {FunctionCall});
 }
 
 void UDMMaterialStageGradient::GenerateGradientList()
diff --git a/Engine/Plugins/Experimental/DynamicMaterial/Source/DynamicMaterialEditor/Private/Components/MaterialStageGradients/DMMSGLinear.cpp b/Engine/Plugins/Experimental/DynamicMaterial/Source/DynamicMaterialEditor/Private/Components/MaterialStageGradients/DMMSGLinear.cpp
--- a/Engine/Plugins/Experimental/DynamicMaterial/Source/DynamicMaterialEditor/Private/Components/MaterialStageGradients/DMMSGLinear.cpp
+++ b/Engine/Plugins/Experimental/DynamicMaterial/Source/DynamicMaterialEditor/Private/Components/MaterialStageGradients/DMMSGLinear.cpp
@@ -1,123 +1,92 @@
 // Copyright Epic Games, Inc. All Rights Reserved.
 
 #include "Components/MaterialStageGradients/DMMSGLinear.h"
+
 #include "DMDefs.h"
-#include "Materials/MaterialExpressionComponentMask.h"
-#include "Materials/MaterialExpressionConstant.h"
-#include "Materials/MaterialExpressionFmod.h"
-#include "Materials/MaterialExpressionIf.h"
+#include "Components/DMMaterialStageFunction.h"
 #include "Materials/MaterialExpressionMaterialFunctionCall.h"
-#include "Materials/MaterialExpressionSubtract.h"
-#include "Model/DMMaterialBuildState.h"
-#include "Model/DMMaterialBuildUtils.h"
-#include "Utils/DMMaterialFunctionLibrary.h"
-#include "Utils/DMUtils.h"
+#include "Materials/MaterialFunctionInterface.h"
 
 #define LOCTEXT_NAMESPACE "DMMaterialStageGradientLinear"
 
+TSoftObjectPtr<UMaterialFunctionInterface> UDMMaterialStageGradientLinear::LinearGradientNoTileFunction = TSoftObjectPtr<UMaterialFunctionInterface>(FSoftObjectPath(TEXT(
+	"/Script/Engine.MaterialFunction'/DynamicMaterial/MaterialFunctions/Gradients/MF_DM_LinearGradient.MF_DM_LinearGradient'"
+)));
+
+TSoftObjectPtr<UMaterialFunctionInterface> UDMMaterialStageGradientLinear::LinearGradientTileFunction = TSoftObjectPtr<UMaterialFunctionInterface>(FSoftObjectPath(TEXT(
+	"/Script/Engine.MaterialFunction'/DynamicMaterial/MaterialFunctions/Gradients/MF_DM_LinearGradient_Tile.MF_DM_LinearGradient_Tile'"
+)));
+
+TSoftObjectPtr<UMaterialFunctionInterface> UDMMaterialStageGradientLinear::LinearGradientTileAndMirrorFunction = TSoftObjectPtr<UMaterialFunctionInterface>(FSoftObjectPath(TEXT(
+	"/Script/Engine.MaterialFunction'/DynamicMaterial/MaterialFunctions/Gradients/MF_DM_LinearGradient_TileAndMirror.MF_DM_LinearGradient_TileAndMirror'"
+)));
+
 UDMMaterialStageGradientLinear::UDMMaterialStageGradientLinear()
 	: UDMMaterialStageGradient(LOCTEXT("GradientLinear", "Linear Gradient"))
+	, Tiling(ELinearGradientTileType::NoTile)
 {
 	EditableProperties.Add(GET_MEMBER_NAME_CHECKED(UDMMaterialStageGradientLinear, Tiling));
+
+	MaterialFunction = GetMaterialFunctionForTilingType(Tiling);
 }
 
-void UDMMaterialStageGradientLinear::GenerateExpressions(const TSharedRef<FDMMaterialBuildState>& InBuildState) const
+void UDMMaterialStageGradientLinear::SetTilingType(ELinearGradientTileType InType)
 {
-	if (!IsComponentValid() || !IsComponentAdded())
+	if (Tiling == InType)
 	{
 		return;
 	}
 
-	if (InBuildState->HasStageSource(this))
-	{
-		return;
-	}
+	Tiling = InType;
+
+	OnTilingChanged();
+}
+
+void UDMMaterialStageGradientLinear::PostEditChangeProperty(FPropertyChangedEvent& PropertyChangedEvent)
+{
+	Super::PostEditChangeProperty(PropertyChangedEvent);
 
-	UMaterialExpressionComponentMask* Mask = InBuildState->GetBuildUtils().CreateExpression<UMaterialExpressionComponentMask>(UE_DM_NodeComment_Default);
+	const FName MemberName = PropertyChangedEvent.GetMemberPropertyName();
 
-	Mask->R = 1;
-	Mask->G = 0;
-	Mask->B = 0;
-	Mask->A = 0;
+	if (MemberName == GET_MEMBER_NAME_CHECKED(UDMMaterialStageGradientLinear, Tiling))
+	{
+		OnTilingChanged();
+	}
+}
 
-	TArray<UMaterialExpression*> OutputExpressions;
-	OutputExpressions.Add(Mask);
+void UDMMaterialStageGradientLinear::PostEditUndo()
+{
+	Super::PostEditUndo();
 
-	UMaterialExpression* Output = nullptr;
+	SetMaterialFunction(GetMaterialFunctionForTilingType(Tiling));
+}
 
+UMaterialFunctionInterface* UDMMaterialStageGradientLinear::GetMaterialFunctionForTilingType(ELinearGradientTileType) const
+{
 	switch (Tiling)
 	{
 		case ELinearGradientTileType::NoTile:
-			Output = Mask;
-			break;
+			return LinearGradientNoTileFunction.LoadSynchronous();
 
 		case ELinearGradientTileType::Tile:
-		{
-			UMaterialExpressionFmod* Fmod = InBuildState->GetBuildUtils().CreateExpression<UMaterialExpressionFmod>(UE_DM_NodeComment_Default);
-			UMaterialExpressionConstant* Divisor = InBuildState->GetBuildUtils().CreateExpression<UMaterialExpressionConstant>(UE_DM_NodeComment_Default);
-
-			Mask->ConnectExpression(&Fmod->A, 0);
-
-			Divisor->R = 1.f;
-			Divisor->ConnectExpression(&Fmod->B, 0);
-
-			Output = Fmod;
-			OutputExpressions.Append({Divisor, Fmod});
-			break;
-		}
+			return LinearGradientTileFunction.LoadSynchronous();
 
 		case ELinearGradientTileType::TileAndMirror:
-		{
-			UMaterialExpressionFmod* Fmod = InBuildState->GetBuildUtils().CreateExpression<UMaterialExpressionFmod>(UE_DM_NodeComment_Default);
-			UMaterialExpressionConstant* Divisor = InBuildState->GetBuildUtils().CreateExpression<UMaterialExpressionConstant>(UE_DM_NodeComment_Default);
-			UMaterialExpressionSubtract* Subtract = InBuildState->GetBuildUtils().CreateExpression<UMaterialExpressionSubtract>(UE_DM_NodeComment_Default);
-			UMaterialExpressionIf* IfExpr = InBuildState->GetBuildUtils().CreateExpression<UMaterialExpressionIf>(UE_DM_NodeComment_Default);
-
-			Mask->ConnectExpression(&Fmod->A, 0);
-
-			Divisor->R = 2.f;
-			Divisor->ConnectExpression(&Fmod->B, 0);
-
-			Fmod->ConnectExpression(&Subtract->B, 0);
-			Fmod->ConnectExpression(&IfExpr->A, 0);
-			Fmod->ConnectExpression(&IfExpr->AEqualsB, 0);
-			Fmod->ConnectExpression(&IfExpr->ALessThanB, 0);
-
-			Subtract->ConstA = 2.f;
-			Subtract->ConnectExpression(&IfExpr->AGreaterThanB, 0);
-
-			IfExpr->ConstB = 1.f;
-			IfExpr->EqualsThreshold = 0.f;
-
-			Output = IfExpr;
-			OutputExpressions.Append({Divisor, Fmod, Subtract, IfExpr});
-			break;
-		}
+			return LinearGradientTileAndMirrorFunction.LoadSynchronous();
 
 		default:
-			checkNoEntry();
-			break;
+			return UDMMaterialStageFunction::NoOp.LoadSynchronous();
 	}
-
-	UMaterialExpression* MakeFloat = FDMMaterialFunctionLibrary::Get().GetMakeFloat3(InBuildState->GetDynamicMaterial(), UE_DM_NodeComment_Default);
-	Output->ConnectExpression(MakeFloat->GetInput(0), 0);
-	Output->ConnectExpression(MakeFloat->GetInput(1), 0);
-	Output->ConnectExpression(MakeFloat->GetInput(2), 0);
-
-	OutputExpressions.Add(MakeFloat);	
-	InBuildState->AddStageSourceExpressions(this, OutputExpressions);
 }
 
-void UDMMaterialStageGradientLinear::SetTilingType(ELinearGradientTileType InType)
+void UDMMaterialStageGradientLinear::OnTilingChanged()
 {
-	if (Tiling == InType)
+	const bool bWasUpdateCalled = SetMaterialFunction(GetMaterialFunctionForTilingType(Tiling));
+
+	if (!bWasUpdateCalled)
 	{
-		return;
+		Update(EDMUpdateType::Structure);
 	}
-
-	Tiling = InType;
-
-	Update(EDMUpdateType::Structure);
 }
 
 #undef LOCTEXT_NAMESPACE
diff --git a/Engine/Plugins/Experimental/DynamicMaterial/Source/DynamicMaterialEditor/Private/Components/MaterialStageGradients/DMMSGRadial.cpp b/Engine/Plugins/Experimental/DynamicMaterial/Source/DynamicMaterialEditor/Private/Components/MaterialStageGradients/DMMSGRadial.cpp
--- a/Engine/Plugins/Experimental/DynamicMaterial/Source/DynamicMaterialEditor/Private/Components/MaterialStageGradients/DMMSGRadial.cpp
+++ b/Engine/Plugins/Experimental/DynamicMaterial/Source/DynamicMaterialEditor/Private/Components/MaterialStageGradients/DMMSGRadial.cpp
@@ -1,39 +1,20 @@
 // Copyright Epic Games, Inc. All Rights Reserved.
 
 #include "Components/MaterialStageGradients/DMMSGRadial.h"
-#include "Utils/DMMaterialFunctionLibrary.h"
+
 #include "Materials/MaterialExpressionMaterialFunctionCall.h"
-#include "Model/DMMaterialBuildState.h"
-#include "Model/DMMaterialBuildUtils.h"
-#include "Utils/DMUtils.h"
+#include "Materials/MaterialFunctionInterface.h"
 
 #define LOCTEXT_NAMESPACE "DMMaterialStageGradientRadial"
 
+TSoftObjectPtr<UMaterialFunctionInterface> UDMMaterialStageGradientRadial::RadialGradientFunction = TSoftObjectPtr<UMaterialFunctionInterface>(FSoftObjectPath(TEXT(
+	"/Script/Engine.MaterialFunction'/DynamicMaterial/MaterialFunctions/Gradients/MF_DM_RadialGradient.MF_DM_RadialGradient'"
+)));
+
 UDMMaterialStageGradientRadial::UDMMaterialStageGradientRadial()
 	: UDMMaterialStageGradient(LOCTEXT("GradientRadial", "Radial Gradient"))
 {
-}
-
-void UDMMaterialStageGradientRadial::GenerateExpressions(const TSharedRef<FDMMaterialBuildState>& InBuildState) const
-{
-	if (!IsComponentValid() || !IsComponentAdded())
-	{
-		return;
-	}
-
-	if (InBuildState->HasStageSource(this))
-	{
-		return;
-	}
-
-	UMaterialExpressionMaterialFunctionCall* Radial = FDMMaterialFunctionLibrary::Get().GetRadialGradientExponential(InBuildState->GetDynamicMaterial(), UE_DM_NodeComment_Default);
-	UMaterialExpression* MakeFloat = FDMMaterialFunctionLibrary::Get().GetMakeFloat3(InBuildState->GetDynamicMaterial(), UE_DM_NodeComment_Default);
-
-	Radial->ConnectExpression(MakeFloat->GetInput(0), 0);
-	Radial->ConnectExpression(MakeFloat->GetInput(1), 0);
-	Radial->ConnectExpression(MakeFloat->GetInput(2), 0);
-
-	InBuildState->AddStageSourceExpressions(this, {Radial, MakeFloat});
+	MaterialFunction = RadialGradientFunction.LoadSynchronous();
 }
 
 #undef LOCTEXT_NAMESPACE
diff --git a/Engine/Plugins/Experimental/DynamicMaterial/Source/DynamicMaterialEditor/Public/Components/DMMaterialStageFunction.h b/Engine/Plugins/Experimental/DynamicMaterial/Source/DynamicMaterialEditor/Public/Components/DMMaterialStageFunction.h
--- a/Engine/Plugins/Experimental/DynamicMaterial/Source/DynamicMaterialEditor/Public/Components/DMMaterialStageFunction.h
+++ b/Engine/Plugins/Experimental/DynamicMaterial/Source/DynamicMaterialEditor/Public/Components/DMMaterialStageFunction.h
@@ -17,6 +17,8 @@ class UDMMaterialStageFunction : public UDMMaterialStageThroughput
 public:
 	static constexpr int32 InputPreviousStage = 0;
 
+	DYNAMICMATERIALEDITOR_API static TSoftObjectPtr<UMaterialFunctionInterface> NoOp;
+
 	DYNAMICMATERIALEDITOR_API static UDMMaterialStage* CreateStage(UDMMaterialLayerObject* InLayer = nullptr);
 
 	UFUNCTION(BlueprintCallable, Category = "Material Designer")
@@ -61,8 +63,6 @@ public:
 	//~ End UObject
 
 protected:
-	static TSoftObjectPtr<UMaterialFunctionInterface> NoOp;
-
 	UPROPERTY(EditAnywhere, BlueprintReadWrite, Getter, Setter, BlueprintSetter = SetMaterialFunction, Category = "Material Designer",
 		meta = (DisplayThumbnail = true, AllowPrivateAccess = "true", HighPriority, NotKeyframeable, NoCreate))
 	TObjectPtr<UMaterialFunctionInterface> MaterialFunction;
diff --git a/Engine/Plugins/Experimental/DynamicMaterial/Source/DynamicMaterialEditor/Public/Components/DMMaterialStageGradient.h b/Engine/Plugins/Experimental/DynamicMaterial/Source/DynamicMaterialEditor/Public/Components/DMMaterialStageGradient.h
--- a/Engine/Plugins/Experimental/DynamicMaterial/Source/DynamicMaterialEditor/Public/Components/DMMaterialStageGradient.h
+++ b/Engine/Plugins/Experimental/DynamicMaterial/Source/DynamicMaterialEditor/Public/Components/DMMaterialStageGradient.h
@@ -3,12 +3,16 @@
 #pragma once
 
 #include "DMMaterialStageThroughput.h"
+
+#include "Math/Color.h"
 #include "UObject/StrongObjectPtr.h"
+
 #include "DMMaterialStageGradient.generated.h"
 
 class FMenuBuilder;
 class UDMMaterialLayerObject;
 class UDMMaterialStageInput;
+class UMaterialFunctionInterface;
 struct FDMMaterialBuildState;
 
 /**
@@ -20,6 +24,10 @@ class UDMMaterialStageGradient : public UDMMaterialStageThroughput
 	GENERATED_BODY()
 
 public:
+	static constexpr int32 InputUV = 0;
+	static constexpr int32 InputStart = 1;
+	static constexpr int32 InputEnd = 2;
+
 	UFUNCTION(BlueprintCallable, Category = "Material Designer")
 	static DYNAMICMATERIALEDITOR_API UDMMaterialStage* CreateStage(TSubclassOf<UDMMaterialStageGradient> InMaterialStageGradientClass, UDMMaterialLayerObject* InLayer = nullptr);
 
@@ -39,17 +47,26 @@ public:
 
 	//~ Begin UDMMaterialStageThroughput
 	DYNAMICMATERIALEDITOR_API virtual bool CanChangeInputType(int32 InputIndex) const override;
+	DYNAMICMATERIALEDITOR_API virtual void AddDefaultInput(int32 InInputIndex) const override;
 	//~ End UDMMaterialStageThroughput
 	
 	//~ Begin UDMMaterialStageSource
+	DYNAMICMATERIALEDITOR_API virtual void GenerateExpressions(const TSharedRef<FDMMaterialBuildState>& InBuildState) const override;
 	virtual bool SupportsLayerMaskTextureUVLink() const override { return true; }
 	virtual int32 GetLayerMaskTextureUVLinkInputIndex() const override { return 0; }
 	//~ End UDMMaterialStageSource
-	
+
 protected:
 	static TArray<TStrongObjectPtr<UClass>> Gradients;
 
 	static void GenerateGradientList();
 
+	UPROPERTY(VisibleAnywhere, BlueprintReadOnly, Category = "Material Designer",
+		meta = (DisplayThumbnail = true, AllowPrivateAccess = "true", NoCreate))
+	TObjectPtr<UMaterialFunctionInterface> MaterialFunction;
+
 	DYNAMICMATERIALEDITOR_API UDMMaterialStageGradient(const FText& InName);
+
+	/** Always set the function. Returns whether an update was called. */
+	bool SetMaterialFunction(UMaterialFunctionInterface* InMaterialFunction);
 };
\ No newline at end of file
diff --git a/Engine/Plugins/Experimental/DynamicMaterial/Source/DynamicMaterialEditor/Public/Components/MaterialStageGradients/DMMSGLinear.h b/Engine/Plugins/Experimental/DynamicMaterial/Source/DynamicMaterialEditor/Public/Components/MaterialStageGradients/DMMSGLinear.h
--- a/Engine/Plugins/Experimental/DynamicMaterial/Source/DynamicMaterialEditor/Public/Components/MaterialStageGradients/DMMSGLinear.h
+++ b/Engine/Plugins/Experimental/DynamicMaterial/Source/DynamicMaterialEditor/Public/Components/MaterialStageGradients/DMMSGLinear.h
@@ -30,12 +30,21 @@ public:
 	UFUNCTION(BlueprintCallable, Category = "Material Designer")
 	DYNAMICMATERIALEDITOR_API virtual void SetTilingType(ELinearGradientTileType InType);
 
-	//~ Begin UDMMaterialStageSource
-	DYNAMICMATERIALEDITOR_API virtual void GenerateExpressions(const TSharedRef<FDMMaterialBuildState>& InBuildState) const override;
-	//~ End UDMMaterialStageSource
+	//~ Begin UObject
+	DYNAMICMATERIALEDITOR_API virtual void PostEditChangeProperty(FPropertyChangedEvent& PropertyChangedEvent) override;
+	virtual void PostEditUndo() override;
+	//~ End UObject
 
 protected:
+	static TSoftObjectPtr<UMaterialFunctionInterface> LinearGradientNoTileFunction;
+	static TSoftObjectPtr<UMaterialFunctionInterface> LinearGradientTileFunction;
+	static TSoftObjectPtr<UMaterialFunctionInterface> LinearGradientTileAndMirrorFunction;
+
 	UPROPERTY(EditInstanceOnly, BlueprintReadWrite, Getter = GetTilingType, Setter=SetTilingType, BlueprintSetter = SetTilingType, 
 		Category = "Material Designer")
 	ELinearGradientTileType Tiling;
+
+	UMaterialFunctionInterface* GetMaterialFunctionForTilingType(ELinearGradientTileType) const;
+
+	void OnTilingChanged();
 };
diff --git a/Engine/Plugins/Experimental/DynamicMaterial/Source/DynamicMaterialEditor/Public/Components/MaterialStageGradients/DMMSGRadial.h b/Engine/Plugins/Experimental/DynamicMaterial/Source/DynamicMaterialEditor/Public/Components/MaterialStageGradients/DMMSGRadial.h
--- a/Engine/Plugins/Experimental/DynamicMaterial/Source/DynamicMaterialEditor/Public/Components/MaterialStageGradients/DMMSGRadial.h
+++ b/Engine/Plugins/Experimental/DynamicMaterial/Source/DynamicMaterialEditor/Public/Components/MaterialStageGradients/DMMSGRadial.h
@@ -13,10 +13,8 @@ class UDMMaterialStageGradientRadial : public UDMMaterialStageGradient
 {
 	GENERATED_BODY()
 
-public:
 	UDMMaterialStageGradientRadial();
 
-	//~ Begin UDMMaterialStageSource
-	DYNAMICMATERIALEDITOR_API virtual void GenerateExpressions(const TSharedRef<FDMMaterialBuildState>& InBuildState) const override;
-	//~ End UDMMaterialStageSource
+public:
+	static TSoftObjectPtr<UMaterialFunctionInterface> RadialGradientFunction;
 };
