Material Designer
- UI 2..3.. point 0?
- Added UPROPERTY-style specifiers for material function inputs in the Material Expression Desc property.

See also https://p4-swarm.epicgames.net/reviews/35125826 https://p4-swarm.epicgames.net/reviews/35125836

diff --git a/Engine/Plugins/Experimental/Avalanche/Source/AvalancheEditor/Private/MaterialDesigner/AvaMaterialDesignerExtension.cpp b/Engine/Plugins/Experimental/Avalanche/Source/AvalancheEditor/Private/MaterialDesigner/AvaMaterialDesignerExtension.cpp
--- a/Engine/Plugins/Experimental/Avalanche/Source/AvalancheEditor/Private/MaterialDesigner/AvaMaterialDesignerExtension.cpp
+++ b/Engine/Plugins/Experimental/Avalanche/Source/AvalancheEditor/Private/MaterialDesigner/AvaMaterialDesignerExtension.cpp
@@ -1,10 +1,11 @@
 // Copyright Epic Games, Inc. All Rights Reserved.
 
 #include "MaterialDesigner/AvaMaterialDesignerExtension.h"
+
 #include "AvaShapeActor.h"
+#include "DetailView/AvaDetailsExtension.h"
 #include "DMObjectMaterialProperty.h"
 #include "DMWorldSubsystem.h"
-#include "DetailView/AvaDetailsExtension.h"
 #include "DynamicMeshes/AvaShapeDynMeshBase.h"
 #include "Engine/World.h"
 #include "IDynamicMaterialEditorModule.h"
@@ -81,19 +82,19 @@ bool FAvaMaterialDesignerExtension::SetDynamicMaterialValue(const FDMObjectMater
 		return false;
 	}
 
-	if (const UActorComponent* const ActorComponent = Cast<UActorComponent>(InObjectMaterialProperty.OuterWeak.Get()))
+	if (const UActorComponent* const ActorComponent = Cast<UActorComponent>(InObjectMaterialProperty.GetOuter()))
 	{
 		if (const AAvaShapeActor* const ShapeActor = Cast<AAvaShapeActor>(ActorComponent->GetOwner()))
 		{
 			UAvaShapeDynamicMeshBase* const DynamicMesh = ShapeActor->GetDynamicMesh();
-			if (DynamicMesh && DynamicMesh->GetMeshesIndexes().Contains(InObjectMaterialProperty.Index))
+			if (DynamicMesh && DynamicMesh->GetMeshesIndexes().Contains(InObjectMaterialProperty.GetIndex()))
 			{
 				if (!InMaterial->IsAsset())
 				{
 					InMaterial->Rename(nullptr, DynamicMesh);
 				}
 
-				DynamicMesh->SetMaterial(InObjectMaterialProperty.Index, InMaterial);
+				DynamicMesh->SetMaterial(InObjectMaterialProperty.GetIndex(), InMaterial);
 				return true;
 			}
 		}
diff --git a/Engine/Plugins/Experimental/DynamicMaterial/Source/DynamicMaterial/Private/Components/DMMaterialComponentDynamic.cpp b/Engine/Plugins/Experimental/DynamicMaterial/Source/DynamicMaterial/Private/Components/DMMaterialComponentDynamic.cpp
--- a/Engine/Plugins/Experimental/DynamicMaterial/Source/DynamicMaterial/Private/Components/DMMaterialComponentDynamic.cpp
+++ b/Engine/Plugins/Experimental/DynamicMaterial/Source/DynamicMaterial/Private/Components/DMMaterialComponentDynamic.cpp
@@ -3,6 +3,7 @@
 #include "Components/DMMaterialComponentDynamic.h"
 
 #include "DMComponentPath.h"
+#include "DynamicMaterialModule.h"
 #include "Model/DynamicMaterialModel.h"
 #include "Model/DynamicMaterialModelDynamic.h"
 
@@ -49,7 +50,7 @@ void UDMMaterialComponentDynamic::BeginDestroy()
 {
 	Super::BeginDestroy();
 
-	if (!UObjectInitialized())
+	if (!FDynamicMaterialModule::AreUObjectsSafe())
 	{
 		return;
 	}
diff --git a/Engine/Plugins/Experimental/DynamicMaterial/Source/DynamicMaterial/Private/DynamicMaterialModule.cpp b/Engine/Plugins/Experimental/DynamicMaterial/Source/DynamicMaterial/Private/DynamicMaterialModule.cpp
--- a/Engine/Plugins/Experimental/DynamicMaterial/Source/DynamicMaterial/Private/DynamicMaterialModule.cpp
+++ b/Engine/Plugins/Experimental/DynamicMaterial/Source/DynamicMaterial/Private/DynamicMaterialModule.cpp
@@ -1,6 +1,8 @@
 // Copyright Epic Games, Inc. All Rights Reserved.
 
 #include "DynamicMaterialModule.h"
+
+#include "CoreGlobals.h"
 #include "HAL/IConsoleManager.h"
 #include "Misc/CoreDelegates.h"
 #include "Modules/ModuleManager.h"
@@ -8,8 +10,6 @@
 
 DEFINE_LOG_CATEGORY(LogDynamicMaterial);
 
-bool FDynamicMaterialModule::bIsEngineExiting = false;
-
 #if WITH_EDITOR
 FDMCreateEditorOnlyDataDelegate FDynamicMaterialModule::CreateEditorOnlyDataDelegate;
 #endif
@@ -27,7 +27,7 @@ bool FDynamicMaterialModule::IsMaterialExportEnabled()
 
 bool FDynamicMaterialModule::AreUObjectsSafe()
 {
-	return UObjectInitialized() && !bIsEngineExiting;
+	return UObjectInitialized() && !IsEngineExitRequested();
 }
 
 FDynamicMaterialModule& FDynamicMaterialModule::Get()
@@ -35,24 +35,6 @@ FDynamicMaterialModule& FDynamicMaterialModule::Get()
 	return FModuleManager::LoadModuleChecked<FDynamicMaterialModule>("DynamicMaterial");
 }
 
-void FDynamicMaterialModule::StartupModule()
-{
-	EnginePreExitHandle = FCoreDelegates::OnEnginePreExit.AddStatic(&FDynamicMaterialModule::HandleEnginePreExit);
-}
-
-void FDynamicMaterialModule::ShutdownModule()
-{
-	if (EnginePreExitHandle.IsValid())
-	{
-		FCoreDelegates::OnEnginePreExit.Remove(EnginePreExitHandle);
-	}
-}
-
-void FDynamicMaterialModule::HandleEnginePreExit()
-{
-	bIsEngineExiting = true;
-}
-
 #if WITH_EDITOR
 TScriptInterface<IDynamicMaterialModelEditorOnlyDataInterface> FDynamicMaterialModule::CreateEditorOnlyData(UDynamicMaterialModel* InMaterialModel)
 {
diff --git a/Engine/Plugins/Experimental/DynamicMaterial/Source/DynamicMaterial/Private/Material/DynamicMaterialInstance.cpp b/Engine/Plugins/Experimental/DynamicMaterial/Source/DynamicMaterial/Private/Material/DynamicMaterialInstance.cpp
--- a/Engine/Plugins/Experimental/DynamicMaterial/Source/DynamicMaterial/Private/Material/DynamicMaterialInstance.cpp
+++ b/Engine/Plugins/Experimental/DynamicMaterial/Source/DynamicMaterial/Private/Material/DynamicMaterialInstance.cpp
@@ -55,9 +55,9 @@ void UDynamicMaterialInstance::InitializeMIDPublic()
 	UpdateCachedData();
 }
 
-void UDynamicMaterialInstance::PostDuplicate(bool bDuplicateForPIE)
+void UDynamicMaterialInstance::PostDuplicate(bool bInDuplicateForPIE)
 {
-	Super::PostDuplicate(bDuplicateForPIE);
+	Super::PostDuplicate(bInDuplicateForPIE);
 
 	if (MaterialModelBase)
 	{
diff --git a/Engine/Plugins/Experimental/DynamicMaterial/Source/DynamicMaterial/Private/Model/DynamicMaterialModel.cpp b/Engine/Plugins/Experimental/DynamicMaterial/Source/DynamicMaterial/Private/Model/DynamicMaterialModel.cpp
--- a/Engine/Plugins/Experimental/DynamicMaterial/Source/DynamicMaterial/Private/Model/DynamicMaterialModel.cpp
+++ b/Engine/Plugins/Experimental/DynamicMaterial/Source/DynamicMaterial/Private/Model/DynamicMaterialModel.cpp
@@ -562,17 +562,20 @@ void UDynamicMaterialModel::PostEditImport()
 	}
 }
 
-void UDynamicMaterialModel::PostDuplicate(bool bDuplicateForPIE)
+void UDynamicMaterialModel::PostDuplicate(bool bInDuplicateForPIE)
 {
-	Super::PostDuplicate(bDuplicateForPIE);
+	Super::PostDuplicate(bInDuplicateForPIE);
 
-	FixGlobalVars();
-	PostEditorDuplicate();
-	ReinitComponents();
-
-	if (IDynamicMaterialModelEditorOnlyDataInterface* ModelEditorOnlyData = GetEditorOnlyData())
+	if (!bInDuplicateForPIE)
 	{
-		ModelEditorOnlyData->RequestMaterialBuild();
+		FixGlobalVars();
+		PostEditorDuplicate();
+		ReinitComponents();
+
+		if (IDynamicMaterialModelEditorOnlyDataInterface* ModelEditorOnlyData = GetEditorOnlyData())
+		{
+			ModelEditorOnlyData->RequestMaterialBuild();
+		}
 	}
 }
 
diff --git a/Engine/Plugins/Experimental/DynamicMaterial/Source/DynamicMaterial/Public/DynamicMaterialModule.h b/Engine/Plugins/Experimental/DynamicMaterial/Source/DynamicMaterial/Public/DynamicMaterialModule.h
--- a/Engine/Plugins/Experimental/DynamicMaterial/Source/DynamicMaterial/Public/DynamicMaterialModule.h
+++ b/Engine/Plugins/Experimental/DynamicMaterial/Source/DynamicMaterial/Public/DynamicMaterialModule.h
@@ -2,9 +2,10 @@
 
 #pragma once
 
+#include "Modules/ModuleInterface.h"
+
 #include "Delegates/Delegate.h"
 #include "Delegates/DelegateCombinations.h"
-#include "Modules/ModuleInterface.h"
 
 #if WITH_EDITOR
 #include "UObject/ScriptInterface.h"
@@ -40,19 +41,9 @@ public:
 	static FDMCreateEditorOnlyDataDelegate::RegistrationType& GetCreateEditorOnlyDataDelegate() { return CreateEditorOnlyDataDelegate; }
 #endif
 
-	//~ Begin IDynamicMaterialModule
-	virtual void StartupModule() override;
-	virtual void ShutdownModule() override;
-	//~ End IDynamicMaterialModule
-
 protected:
-	static bool bIsEngineExiting;
-
 #if WITH_EDITOR
 	DYNAMICMATERIAL_API static FDMCreateEditorOnlyDataDelegate CreateEditorOnlyDataDelegate;
 #endif
-	
-	static void HandleEnginePreExit();
 
-	FDelegateHandle EnginePreExitHandle;
 };
diff --git a/Engine/Plugins/Experimental/DynamicMaterial/Source/DynamicMaterial/Public/Material/DynamicMaterialInstance.h b/Engine/Plugins/Experimental/DynamicMaterial/Source/DynamicMaterial/Public/Material/DynamicMaterialInstance.h
--- a/Engine/Plugins/Experimental/DynamicMaterial/Source/DynamicMaterial/Public/Material/DynamicMaterialInstance.h
+++ b/Engine/Plugins/Experimental/DynamicMaterial/Source/DynamicMaterial/Public/Material/DynamicMaterialInstance.h
@@ -41,7 +41,7 @@ public:
 	DYNAMICMATERIAL_API void InitializeMIDPublic();
 
 	//~ Begin UObject
-	DYNAMICMATERIAL_API virtual void PostDuplicate(bool bDuplicateForPIE) override;
+	DYNAMICMATERIAL_API virtual void PostDuplicate(bool bInDuplicateForPIE) override;
 	DYNAMICMATERIAL_API virtual void PostEditImport() override;
 	//~ End UObject
 #endif
diff --git a/Engine/Plugins/Experimental/DynamicMaterial/Source/DynamicMaterial/Public/Model/DynamicMaterialModel.h b/Engine/Plugins/Experimental/DynamicMaterial/Source/DynamicMaterial/Public/Model/DynamicMaterialModel.h
--- a/Engine/Plugins/Experimental/DynamicMaterial/Source/DynamicMaterial/Public/Model/DynamicMaterialModel.h
+++ b/Engine/Plugins/Experimental/DynamicMaterial/Source/DynamicMaterial/Public/Model/DynamicMaterialModel.h
@@ -189,7 +189,7 @@ public:
 #if WITH_EDITOR
 	DYNAMICMATERIAL_API virtual void PostEditUndo() override;
 	DYNAMICMATERIAL_API virtual void PostEditImport() override;
-	DYNAMICMATERIAL_API virtual void PostDuplicate(bool bDuplicateForPIE) override;
+	DYNAMICMATERIAL_API virtual void PostDuplicate(bool bInDuplicateForPIE) override;
 	//~ End UObject
 
 	/** Called to ensure that the object hierarchy is correct. */
diff --git a/Engine/Plugins/Experimental/DynamicMaterial/Source/DynamicMaterialEditor/DynamicMaterialEditor.Build.cs b/Engine/Plugins/Experimental/DynamicMaterial/Source/DynamicMaterialEditor/DynamicMaterialEditor.Build.cs
--- a/Engine/Plugins/Experimental/DynamicMaterial/Source/DynamicMaterialEditor/DynamicMaterialEditor.Build.cs
+++ b/Engine/Plugins/Experimental/DynamicMaterial/Source/DynamicMaterialEditor/DynamicMaterialEditor.Build.cs
@@ -14,13 +14,15 @@ public class DynamicMaterialEditor : ModuleRules
 				"Core",
 				"CoreUObject",
 				"DynamicMaterial",
-				"Engine"
+				"Engine",
+				"MaterialEditor"
 			}
 		);
 
 		PrivateDependencyModuleNames.AddRange(
 			new string[]
 			{
+				"AdvancedPreviewScene",
 				"AppFramework",
 				"ApplicationCore",
 				"AssetDefinition",
diff --git a/Engine/Plugins/Experimental/DynamicMaterial/Source/DynamicMaterialEditor/Private/Components/DMMaterialLayer.cpp b/Engine/Plugins/Experimental/DynamicMaterial/Source/DynamicMaterialEditor/Private/Components/DMMaterialLayer.cpp
--- a/Engine/Plugins/Experimental/DynamicMaterial/Source/DynamicMaterialEditor/Private/Components/DMMaterialLayer.cpp
+++ b/Engine/Plugins/Experimental/DynamicMaterial/Source/DynamicMaterialEditor/Private/Components/DMMaterialLayer.cpp
@@ -1,13 +1,13 @@
 // Copyright Epic Games, Inc. All Rights Reserved.
 
 #include "Components/DMMaterialLayer.h"
+
 #include "Components/DMMaterialEffect.h"
 #include "Components/DMMaterialEffectStack.h"
 #include "Components/DMMaterialSlot.h"
 #include "Components/DMMaterialStage.h"
 #include "DMComponentPath.h"
 #include "Dom/JsonObject.h"
-#include "DynamicMaterialEditorModule.h"
 #include "Factories.h"
 #include "JsonObjectConverter.h"
 #include "Misc/ReverseIterate.h"
@@ -379,19 +379,6 @@ bool UDMMaterialLayerObject::SetStage(EDMMaterialLayerStage InStageType, UDMMate
 	return true;
 }
 
-UDMMaterialStage* UDMMaterialLayerObject::GetFirstStageBeingEdited(EDMMaterialLayerStage InStageScope) const
-{
-	for (UDMMaterialStage* Stage : GetStages(InStageScope))
-	{
-		if (IsValid(Stage) && Stage->IsBeingEdited())
-		{
-			return Stage;
-		}
-	}
-
-	return nullptr;
-}
-
 void UDMMaterialLayerObject::ForEachValidStage(EDMMaterialLayerStage InStageScope, FStageCallbackFunc InCallback) const
 {
 	for (UDMMaterialStage* Stage : GetStages(InStageScope))
@@ -507,16 +494,6 @@ bool UDMMaterialLayerObject::IsStageEnabled(EDMMaterialLayerStage InStageType) c
 	return false;
 }
 
-bool UDMMaterialLayerObject::IsStageBeingEdited(EDMMaterialLayerStage InStageType) const
-{
-	if (UDMMaterialStage* Stage = GetStage(InStageType))
-	{
-		return Stage->IsBeingEdited();
-	}
-
-	return false;
-}
-
 bool UDMMaterialLayerObject::CanMoveLayerAbove(UDMMaterialLayerObject* InLayer) const
 {
 	if (!IsComponentValid())
diff --git a/Engine/Plugins/Experimental/DynamicMaterial/Source/DynamicMaterialEditor/Private/Components/DMMaterialProperty.cpp b/Engine/Plugins/Experimental/DynamicMaterial/Source/DynamicMaterialEditor/Private/Components/DMMaterialProperty.cpp
--- a/Engine/Plugins/Experimental/DynamicMaterial/Source/DynamicMaterialEditor/Private/Components/DMMaterialProperty.cpp
+++ b/Engine/Plugins/Experimental/DynamicMaterial/Source/DynamicMaterialEditor/Private/Components/DMMaterialProperty.cpp
@@ -13,7 +13,6 @@
 #include "Components/MaterialStageInputs/DMMSIExpression.h"
 #include "Components/MaterialStageInputs/DMMSIValue.h"
 #include "Components/MaterialValues/DMMaterialValueFloat1.h"
-#include "Components/MaterialValues/DMMaterialValueFloat3RGB.h"
 #include "Components/MaterialValues/DMMaterialValueTexture.h"
 #include "DMComponentPath.h"
 #include "DMValueDefinition.h"
@@ -31,7 +30,6 @@
 #include "Math/Vector4.h"
 #include "Model/DMMaterialBuildState.h"
 #include "Model/DMMaterialBuildUtils.h"
-#include "Model/DynamicMaterialModel.h"
 #include "Model/DynamicMaterialModelEditorOnlyData.h"
 #include "Utils/DMMaterialFunctionLibrary.h"
 #include "Utils/DMMaterialUtils.h"
@@ -49,6 +47,7 @@ UDMMaterialProperty::UDMMaterialProperty()
 
 UDMMaterialProperty::UDMMaterialProperty(EDMMaterialPropertyType InMaterialProperty, EDMValueType InInputConnectorType)
 	: MaterialProperty(InMaterialProperty)
+	, bEnabled(true)
 	, InputConnectorType(InInputConnectorType)
 {
 }
@@ -76,6 +75,18 @@ UDynamicMaterialModelEditorOnlyData* UDMMaterialProperty::GetMaterialModelEditor
 	return Cast<UDynamicMaterialModelEditorOnlyData>(GetOuterSafe());
 }
 
+void UDMMaterialProperty::SetEnabled(bool bInEnabled)
+{
+	if (bEnabled == bInEnabled)
+	{
+		return;
+	}
+
+	bEnabled = bInEnabled;
+
+	Update(EDMUpdateType::Structure | EDMUpdateType::AllowParentUpdate);
+}
+
 FText UDMMaterialProperty::GetDescription() const
 {
 	return StaticEnum<EDMMaterialPropertyType>()->GetDisplayNameTextByValue(static_cast<int64>(MaterialProperty));
diff --git a/Engine/Plugins/Experimental/DynamicMaterial/Source/DynamicMaterialEditor/Private/Components/DMMaterialSlot.cpp b/Engine/Plugins/Experimental/DynamicMaterial/Source/DynamicMaterialEditor/Private/Components/DMMaterialSlot.cpp
--- a/Engine/Plugins/Experimental/DynamicMaterial/Source/DynamicMaterialEditor/Private/Components/DMMaterialSlot.cpp
+++ b/Engine/Plugins/Experimental/DynamicMaterial/Source/DynamicMaterialEditor/Private/Components/DMMaterialSlot.cpp
@@ -1,6 +1,7 @@
 // Copyright Epic Games, Inc. All Rights Reserved.
 
 #include "Components/DMMaterialSlot.h"
+
 #include "Components/DMMaterialLayer.h"
 #include "Components/DMMaterialLayer_Deprecated.h"
 #include "Components/DMMaterialProperty.h"
@@ -363,7 +364,6 @@ void UDMMaterialSlot::UpdateMaterialProperties()
 
 UDMMaterialSlot::UDMMaterialSlot()
 	: Index(INDEX_NONE)
-	, bIsEditingLayers(true)
 {
 }
 
@@ -619,7 +619,7 @@ bool UDMMaterialSlot::PasteLayer(UDMMaterialLayerObject* InLayer)
 	return true;
 }
 
-bool UDMMaterialSlot::CanRemoveLayer(const UDMMaterialLayerObject* InLayer) const
+bool UDMMaterialSlot::CanRemoveLayer(UDMMaterialLayerObject* InLayer) const
 {
 	if (!IsComponentValid())
 	{
diff --git a/Engine/Plugins/Experimental/DynamicMaterial/Source/DynamicMaterialEditor/Private/Components/DMMaterialStage.cpp b/Engine/Plugins/Experimental/DynamicMaterial/Source/DynamicMaterialEditor/Private/Components/DMMaterialStage.cpp
--- a/Engine/Plugins/Experimental/DynamicMaterial/Source/DynamicMaterialEditor/Private/Components/DMMaterialStage.cpp
+++ b/Engine/Plugins/Experimental/DynamicMaterial/Source/DynamicMaterialEditor/Private/Components/DMMaterialStage.cpp
@@ -1,6 +1,7 @@
 // Copyright Epic Games, Inc. All Rights Reserved.
 
 #include "Components/DMMaterialStage.h"
+
 #include "Components/DMMaterialLayer.h"
 #include "Components/DMMaterialSlot.h"
 #include "Components/DMMaterialStageInput.h"
@@ -34,7 +35,6 @@ UDMMaterialStage::UDMMaterialStage()
 	: Source(nullptr)
 	, bEnabled(true)
 	, bCanChangeSource(true)
-	, bIsBeingEdited(false)
 {
 	EditableProperties.Add(GET_MEMBER_NAME_CHECKED(UDMMaterialStage, Source));
 	EditableProperties.Add(GET_MEMBER_NAME_CHECKED(UDMMaterialStage, Inputs));
@@ -308,7 +308,7 @@ EDMValueType UDMMaterialStage::GetSourceType(const FDMMaterialStageConnectorChan
 		EDMMaterialPropertyType StageProperty = Layer->GetMaterialProperty();
 		check(StageProperty != EDMMaterialPropertyType::None && StageProperty != EDMMaterialPropertyType::Any);
 
-		const UDMMaterialLayerObject* PreviousLayer = Layer->GetPreviousLayer(StageProperty, EDMMaterialLayerStage::Base);
+		UDMMaterialLayerObject* PreviousLayer = Layer->GetPreviousLayer(StageProperty, EDMMaterialLayerStage::Base);
 		check(PreviousLayer);
 		check(PreviousLayer->GetStage(EDMMaterialLayerStage::Mask));
 		check(PreviousLayer->GetStage(EDMMaterialLayerStage::Mask)->GetSource());
@@ -530,18 +530,6 @@ void UDMMaterialStage::GenerateExpressions(const TSharedRef<FDMMaterialBuildStat
 	InBuildState->AddStageExpressions(this, StageExpressions);
 }
 
-bool UDMMaterialStage::SetBeingEdited(bool bInBeingEdited)
-{ 
-	if (bIsBeingEdited == bInBeingEdited)
-	{
-		return false;
-	}
-
-	bIsBeingEdited = bInBeingEdited; 
-
-	return true;
-}
-
 TMap<EDMMaterialPropertyType, UDMMaterialLayerObject*> UDMMaterialStage::GetPreviousStagesPropertyMap()
 {
 	UDMMaterialLayerObject* Layer = GetLayer();
@@ -724,7 +712,7 @@ UDMMaterialStageSource* UDMMaterialStage::ChangeInput_PreviousStage(int32 InInpu
 	EDMMaterialPropertyType StageProperty = Layer->GetMaterialProperty();
 	check(StageProperty != EDMMaterialPropertyType::None);
 
-	const UDMMaterialLayerObject* PreviousLayer = Layer->GetPreviousLayer(StageProperty, EDMMaterialLayerStage::Base);
+	UDMMaterialLayerObject* PreviousLayer = Layer->GetPreviousLayer(StageProperty, EDMMaterialLayerStage::Base);
 	UDMMaterialStageSource* PreviousSource = nullptr;
 	
 	if (PreviousLayer)
@@ -789,7 +777,7 @@ void UDMMaterialStage::UpdateInputMap(int32 InInputIdx, int32 InSourceIndex, int
 		UDMMaterialLayerObject* Layer = GetLayer();
 		check(Layer);
 
-		const UDMMaterialLayerObject* PreviousLayer = Layer->GetPreviousLayer(InStageProperty, EDMMaterialLayerStage::Base);
+		UDMMaterialLayerObject* PreviousLayer = Layer->GetPreviousLayer(InStageProperty, EDMMaterialLayerStage::Base);
 		UDMMaterialStageSource* PreviousSource = nullptr;
 
 		if (PreviousLayer)
@@ -1044,7 +1032,7 @@ bool UDMMaterialStage::VerifyInputMap(int32 InInputIdx)
 		// Check previous stage
 		if (Channel.SourceIndex == FDMMaterialStageConnectorChannel::PREVIOUS_STAGE)
 		{
-			if (const UDMMaterialLayerObject* PreviousLayerAndStage = Layer->GetPreviousLayer(StageProperty, EDMMaterialLayerStage::Base))
+			if (UDMMaterialLayerObject* PreviousLayerAndStage = Layer->GetPreviousLayer(StageProperty, EDMMaterialLayerStage::Base))
 			{
 				check(PreviousLayerAndStage->GetStage(EDMMaterialLayerStage::Mask));
 
diff --git a/Engine/Plugins/Experimental/DynamicMaterial/Source/DynamicMaterialEditor/Private/Components/DMMaterialStageBlend.cpp b/Engine/Plugins/Experimental/DynamicMaterial/Source/DynamicMaterialEditor/Private/Components/DMMaterialStageBlend.cpp
--- a/Engine/Plugins/Experimental/DynamicMaterial/Source/DynamicMaterialEditor/Private/Components/DMMaterialStageBlend.cpp
+++ b/Engine/Plugins/Experimental/DynamicMaterial/Source/DynamicMaterialEditor/Private/Components/DMMaterialStageBlend.cpp
@@ -1,6 +1,7 @@
 // Copyright Epic Games, Inc. All Rights Reserved.
  
 #include "Components/DMMaterialStageBlend.h"
+
 #include "Components/DMMaterialLayer.h"
 #include "Components/DMMaterialSlot.h"
 #include "Components/DMMaterialStage.h"
@@ -456,7 +457,7 @@ void UDMMaterialStageBlend::AddDefaultInput(int32 InInputIndex) const
 			EDMMaterialPropertyType StageProperty = Layer->GetMaterialProperty();
 			check(StageProperty != EDMMaterialPropertyType::None);
  
-			const UDMMaterialLayerObject* PreviousLayer = Layer->GetPreviousLayer(StageProperty, EDMMaterialLayerStage::Base);
+			UDMMaterialLayerObject* PreviousLayer = Layer->GetPreviousLayer(StageProperty, EDMMaterialLayerStage::Base);
  
 			if (PreviousLayer)
 			{
diff --git a/Engine/Plugins/Experimental/DynamicMaterial/Source/DynamicMaterialEditor/Private/Components/DMMaterialStageExpression.cpp b/Engine/Plugins/Experimental/DynamicMaterial/Source/DynamicMaterialEditor/Private/Components/DMMaterialStageExpression.cpp
--- a/Engine/Plugins/Experimental/DynamicMaterial/Source/DynamicMaterialEditor/Private/Components/DMMaterialStageExpression.cpp
+++ b/Engine/Plugins/Experimental/DynamicMaterial/Source/DynamicMaterialEditor/Private/Components/DMMaterialStageExpression.cpp
@@ -1,12 +1,12 @@
 // Copyright Epic Games, Inc. All Rights Reserved.
 
 #include "Components/DMMaterialStageExpression.h"
+
 #include "Components/DMMaterialSlot.h"
 #include "Components/DMMaterialStage.h"
 #include "DMDefs.h"
 #include "EditorClassUtils.h"
 #include "Materials/MaterialExpression.h"
-#include "Materials/MaterialExpressionComponentMask.h"
 #include "Model/DMMaterialBuildState.h"
 #include "Model/DMMaterialBuildUtils.h"
 #include "Utils/DMUtils.h"
@@ -175,8 +175,8 @@ UDMMaterialStageExpression::UDMMaterialStageExpression(const FText& InName, TSub
 	bInputRequired = false;
 	bAllowNestedInputs = false;
 
-	ensureAlwaysMsgf(!InName.IsEmpty(), TEXT("Material Designer MSE Class with invalid Name: %s"),             *StaticClass()->GetName());
-	ensureAlwaysMsgf(InClass,           TEXT("Material Designer MSE Class with invalid Expression Class: %s"), *StaticClass()->GetName());
+	ensureAlwaysMsgf(!InName.IsEmpty(), TEXT("MSE Class with invalid Name: %s"),             *StaticClass()->GetName());
+	ensureAlwaysMsgf(InClass,           TEXT("MSE Class with invalid Expression Class: %s"), *StaticClass()->GetName());
 }
 
 void UDMMaterialStageExpression::GenerateExpressions(const TSharedRef<FDMMaterialBuildState>& InBuildState) const
diff --git a/Engine/Plugins/Experimental/DynamicMaterial/Source/DynamicMaterialEditor/Private/Components/DMMaterialStageFunction.cpp b/Engine/Plugins/Experimental/DynamicMaterial/Source/DynamicMaterialEditor/Private/Components/DMMaterialStageFunction.cpp
--- a/Engine/Plugins/Experimental/DynamicMaterial/Source/DynamicMaterialEditor/Private/Components/DMMaterialStageFunction.cpp
+++ b/Engine/Plugins/Experimental/DynamicMaterial/Source/DynamicMaterialEditor/Private/Components/DMMaterialStageFunction.cpp
@@ -12,7 +12,6 @@
 #include "Components/MaterialValues/DMMaterialValueFloat3RGB.h"
 #include "Components/MaterialValues/DMMaterialValueFloat3XYZ.h"
 #include "Components/MaterialValues/DMMaterialValueFloat4.h"
-#include "DynamicMaterialEditorModule.h"
 #include "DynamicMaterialEditorSettings.h"
 #include "Materials/MaterialExpressionFunctionInput.h"
 #include "Materials/MaterialExpressionFunctionOutput.h"
@@ -139,7 +138,7 @@ void UDMMaterialStageFunction::AddDefaultInput(int32 InInputIndex) const
 	EDMMaterialPropertyType StageProperty = Layer->GetMaterialProperty();
 	check(StageProperty != EDMMaterialPropertyType::None);
 
-	const UDMMaterialLayerObject* PreviousLayer = Layer->GetPreviousLayer(StageProperty, EDMMaterialLayerStage::Base);
+	UDMMaterialLayerObject* PreviousLayer = Layer->GetPreviousLayer(StageProperty, EDMMaterialLayerStage::Base);
 
 	if (PreviousLayer)
 	{
diff --git a/Engine/Plugins/Experimental/DynamicMaterial/Source/DynamicMaterialEditor/Private/Components/DMMaterialStageThroughput.cpp b/Engine/Plugins/Experimental/DynamicMaterial/Source/DynamicMaterialEditor/Private/Components/DMMaterialStageThroughput.cpp
--- a/Engine/Plugins/Experimental/DynamicMaterial/Source/DynamicMaterialEditor/Private/Components/DMMaterialStageThroughput.cpp
+++ b/Engine/Plugins/Experimental/DynamicMaterial/Source/DynamicMaterialEditor/Private/Components/DMMaterialStageThroughput.cpp
@@ -1,6 +1,7 @@
 // Copyright Epic Games, Inc. All Rights Reserved.
 
 #include "Components/DMMaterialStageThroughput.h"
+
 #include "Components/DMMaterialLayer.h"
 #include "Components/DMMaterialSlot.h"
 #include "Components/DMMaterialStage.h"
@@ -166,7 +167,7 @@ bool UDMMaterialStageThroughput::ShouldKeepInput(int32 InThroughputInputIndex)
 
 			case FDMMaterialStageConnectorChannel::PREVIOUS_STAGE:
 			{
-				if (const UDMMaterialLayerObject* PreviousLayer = Layer->GetPreviousLayer(StageProperty, EDMMaterialLayerStage::Base))
+				if (UDMMaterialLayerObject* PreviousLayer = Layer->GetPreviousLayer(StageProperty, EDMMaterialLayerStage::Base))
 				{
 					if (UDMMaterialStage* MaskStage = PreviousLayer->GetStage(EDMMaterialLayerStage::Mask))
 					{
@@ -474,7 +475,7 @@ void UDMMaterialStageThroughput::AddDefaultInput(int32 InInputIndex) const
 
 		case EDMValueType::VT_ColorAtlas:
 		{
-			if (const UDMMaterialLayerObject* Layer = Stage->GetLayer())
+			if (UDMMaterialLayerObject* Layer = Stage->GetLayer())
 			{
 				if (Layer->GetStageType(Stage) == EDMMaterialLayerStage::Mask)
 				{
@@ -514,7 +515,7 @@ int32 UDMMaterialStageThroughput::ResolveInput(const TSharedRef<FDMMaterialBuild
 		ParentMostStage = SubStage->GetParentMostStage();
 	}
 
-	const UDMMaterialLayerObject* Layer = ParentMostStage->GetLayer();
+	UDMMaterialLayerObject* Layer = ParentMostStage->GetLayer();
 	check(Layer);
 
 	if (Layer->IsTextureUVLinkEnabled() 
@@ -628,7 +629,7 @@ int32 UDMMaterialStageThroughput::ResolveLayerMaskTextureUVLinkInput(const TShar
 	UDMMaterialStage* Stage = GetStage();
 	check(Stage);
 
-	const UDMMaterialLayerObject* Layer = Stage->GetLayer();
+	UDMMaterialLayerObject* Layer = Stage->GetLayer();
 	check(Layer);
 	check(Layer->GetStage(EDMMaterialLayerStage::Base));
 
@@ -709,10 +710,10 @@ int32 UDMMaterialStageThroughput::ResolveInputChannel(const TSharedRef<FDMMateri
 
 	if (OutChannel.SourceIndex == FDMMaterialStageConnectorChannel::PREVIOUS_STAGE)
 	{
-		const UDMMaterialLayerObject* Layer = Stage->GetLayer();
+		UDMMaterialLayerObject* Layer = Stage->GetLayer();
 		check(Layer);
 
-		if (const UDMMaterialLayerObject* PreviousLayer = Layer->GetPreviousLayer(OutChannel.MaterialProperty, EDMMaterialLayerStage::Base))
+		if (UDMMaterialLayerObject* PreviousLayer = Layer->GetPreviousLayer(OutChannel.MaterialProperty, EDMMaterialLayerStage::Base))
 		{
 			PreviousLayer->GenerateExpressions(InBuildState);
 
diff --git a/Engine/Plugins/Experimental/DynamicMaterial/Source/DynamicMaterialEditor/Private/Components/DMMaterialStageThroughputLayerBlend.cpp b/Engine/Plugins/Experimental/DynamicMaterial/Source/DynamicMaterialEditor/Private/Components/DMMaterialStageThroughputLayerBlend.cpp
--- a/Engine/Plugins/Experimental/DynamicMaterial/Source/DynamicMaterialEditor/Private/Components/DMMaterialStageThroughputLayerBlend.cpp
+++ b/Engine/Plugins/Experimental/DynamicMaterial/Source/DynamicMaterialEditor/Private/Components/DMMaterialStageThroughputLayerBlend.cpp
@@ -109,7 +109,7 @@ void UDMMaterialStageThroughputLayerBlend::AddDefaultInput(int32 InInputIndex) c
 	UDMMaterialStage* Stage = GetStage();
 	check(Stage);
  
-	const UDMMaterialLayerObject* Layer = Stage->GetLayer();
+	UDMMaterialLayerObject* Layer = Stage->GetLayer();
 	check(Layer);
  
 	switch (InInputIndex)
@@ -202,7 +202,7 @@ void UDMMaterialStageThroughputLayerBlend::Update(EDMUpdateType InUpdateType)
 		InputConnectors[1].Type = EDMValueType::VT_Float_Any;
 		OutputConnectors[0].Type = EDMValueType::VT_Float_Any;
  
-		if (const UDMMaterialLayerObject* Layer = Stage->GetLayer())
+		if (UDMMaterialLayerObject* Layer = Stage->GetLayer())
 		{
 			if (const UDMMaterialStageSource* LayerSource = Layer->GetStage(EDMMaterialLayerStage::Base)->GetSource())
 			{
@@ -233,13 +233,13 @@ int32 UDMMaterialStageThroughputLayerBlend::ResolveInput(const TSharedRef<FDMMat
 			const UDMMaterialStage* Stage = GetStage();
 			check(Stage);
 
-			const UDMMaterialLayerObject* Layer = Stage->GetLayer();
+			UDMMaterialLayerObject* Layer = Stage->GetLayer();
 			check(Layer);
 
 			EDMMaterialPropertyType StageProperty = Layer->GetMaterialProperty();
 			check(StageProperty != EDMMaterialPropertyType::None);
 
-			if (const UDMMaterialLayerObject* PreviousLayer = Layer->GetPreviousLayer(StageProperty, EDMMaterialLayerStage::Base))
+			if (UDMMaterialLayerObject* PreviousLayer = Layer->GetPreviousLayer(StageProperty, EDMMaterialLayerStage::Base))
 			{
 				PreviousLayer->GenerateExpressions(InBuildState);
 				OutExpressions.Append(InBuildState->GetLayerExpressions(PreviousLayer));
@@ -261,7 +261,7 @@ int32 UDMMaterialStageThroughputLayerBlend::ResolveInput(const TSharedRef<FDMMat
 			const UDMMaterialStage* Stage = GetStage();
 			check(Stage);
  
-			const UDMMaterialLayerObject* Layer = Stage->GetLayer();
+			UDMMaterialLayerObject* Layer = Stage->GetLayer();
 			check(Layer);
 
 			UDMMaterialStage* BaseStage = Layer->GetStage(EDMMaterialLayerStage::Base, /* Enabled Only */ true);
@@ -347,7 +347,7 @@ void UDMMaterialStageThroughputLayerBlend::GenerateMainExpressions(const TShared
 
 	if (UDMMaterialStage* Stage = GetStage())
 	{
-		if (const UDMMaterialLayerObject* Layer = Stage->GetLayer())
+		if (UDMMaterialLayerObject* Layer = Stage->GetLayer())
 		{
 			switch (Layer->GetMaterialProperty())
 			{
@@ -475,7 +475,7 @@ void UDMMaterialStageThroughputLayerBlend::UpdateLinkedInputStage(EDMUpdateType
 		MaskStage = SubStage->GetParentMostStage();
 	}
 
-	const UDMMaterialLayerObject* Layer = MaskStage->GetLayer();
+	UDMMaterialLayerObject* Layer = MaskStage->GetLayer();
 
 	if (!Layer)
 	{
@@ -802,7 +802,7 @@ int32 UDMMaterialStageThroughputLayerBlend::ResolveMaskInput(const TSharedRef<FD
 	UDMMaterialStage* ThisStage = GetStage();
 	check(ThisStage);
  
-	const UDMMaterialLayerObject* ThisLayer = ThisStage->GetLayer();
+	UDMMaterialLayerObject* ThisLayer = ThisStage->GetLayer();
 	check(ThisLayer);
 	check(ThisLayer->GetStage(EDMMaterialLayerStage::Mask));
 
@@ -814,7 +814,7 @@ int32 UDMMaterialStageThroughputLayerBlend::ResolveMaskInput(const TSharedRef<FD
 		OutChannel.OutputIndex
 	);
 		
-	while (const UDMMaterialLayerObject* NextLayer = ThisLayer->GetNextLayer(ThisLayer->GetMaterialProperty(), EDMMaterialLayerStage::Mask))
+	while (UDMMaterialLayerObject* NextLayer = ThisLayer->GetNextLayer(ThisLayer->GetMaterialProperty(), EDMMaterialLayerStage::Mask))
 	{
 		UDMMaterialStage* NextMaskStage = NextLayer->GetStage(EDMMaterialLayerStage::Mask, /* Enabled Only */ true);
 
@@ -971,7 +971,7 @@ void UDMMaterialStageThroughputLayerBlend::UpdateAlphaOnlyMaskStatus()
 		return;
 	}
  
-	const UDMMaterialLayerObject* Layer = Stage->GetLayer();
+	UDMMaterialLayerObject* Layer = Stage->GetLayer();
 	
 	if (!Layer)
 	{
@@ -1019,7 +1019,7 @@ void UDMMaterialStageThroughputLayerBlend::UpdateAlphaOnlyMasks(EDMUpdateType In
 		return;
 	}
  
-	const UDMMaterialLayerObject* CurrentLayer = Stage->GetLayer();
+	UDMMaterialLayerObject* CurrentLayer = Stage->GetLayer();
 	
 	if (!CurrentLayer)
 	{
@@ -1033,7 +1033,7 @@ void UDMMaterialStageThroughputLayerBlend::UpdateAlphaOnlyMasks(EDMUpdateType In
 		return;
 	}
  
-	const UDMMaterialLayerObject* PreviousLayer = nullptr;
+	UDMMaterialLayerObject* PreviousLayer = nullptr;
  
 	while (true)
 	{
diff --git a/Engine/Plugins/Experimental/DynamicMaterial/Source/DynamicMaterialEditor/Private/Components/MaterialStageExpressions/DMMSEMathBase.cpp b/Engine/Plugins/Experimental/DynamicMaterial/Source/DynamicMaterialEditor/Private/Components/MaterialStageExpressions/DMMSEMathBase.cpp
--- a/Engine/Plugins/Experimental/DynamicMaterial/Source/DynamicMaterialEditor/Private/Components/MaterialStageExpressions/DMMSEMathBase.cpp
+++ b/Engine/Plugins/Experimental/DynamicMaterial/Source/DynamicMaterialEditor/Private/Components/MaterialStageExpressions/DMMSEMathBase.cpp
@@ -1,6 +1,7 @@
 // Copyright Epic Games, Inc. All Rights Reserved.
 
 #include "Components/MaterialStageExpressions/DMMSEMathBase.h"
+
 #include "Components/DMMaterialLayer.h"
 #include "Components/DMMaterialSlot.h"
 #include "Components/DMMaterialStage.h"
@@ -68,10 +69,10 @@ bool UDMMaterialStageExpressionMathBase::CanInputAcceptType(int32 InInputIndex,
 
 				case FDMMaterialStageConnectorChannel::PREVIOUS_STAGE:
 				{
-					const UDMMaterialLayerObject* Layer = Stage->GetLayer();
+					UDMMaterialLayerObject* Layer = Stage->GetLayer();
 					check(Layer);
 
-					if (const UDMMaterialLayerObject* PreviousLayer = Layer->GetPreviousLayer(Channel.MaterialProperty, EDMMaterialLayerStage::Base))
+					if (UDMMaterialLayerObject* PreviousLayer = Layer->GetPreviousLayer(Channel.MaterialProperty, EDMMaterialLayerStage::Base))
 					{
 						UDMMaterialStage* PreviousStage = PreviousLayer->GetLastEnabledStage(EDMMaterialLayerStage::All);
 						check(PreviousStage);
diff --git a/Engine/Plugins/Experimental/DynamicMaterial/Source/DynamicMaterialEditor/Private/Components/MaterialStageExpressions/DMMSETextureSampleBase.cpp b/Engine/Plugins/Experimental/DynamicMaterial/Source/DynamicMaterialEditor/Private/Components/MaterialStageExpressions/DMMSETextureSampleBase.cpp
--- a/Engine/Plugins/Experimental/DynamicMaterial/Source/DynamicMaterialEditor/Private/Components/MaterialStageExpressions/DMMSETextureSampleBase.cpp
+++ b/Engine/Plugins/Experimental/DynamicMaterial/Source/DynamicMaterialEditor/Private/Components/MaterialStageExpressions/DMMSETextureSampleBase.cpp
@@ -103,7 +103,7 @@ bool UDMMaterialStageExpressionTextureSampleBase::IsPropertyVisible(FName Proper
 			ParentMostStage = SubStage->GetParentMostStage();
 		}
 
-		const UDMMaterialLayerObject* Layer = ParentMostStage->GetLayer();
+		UDMMaterialLayerObject* Layer = ParentMostStage->GetLayer();
 		check(Layer);
 
 		if (UDMMaterialStage* BaseStage = Layer->GetStage(EDMMaterialLayerStage::Base, /* Enabled Only */ true))
@@ -152,7 +152,7 @@ void UDMMaterialStageExpressionTextureSampleBase::AddExpressionProperties(const
 		ParentMostStage = SubStage->GetParentMostStage();
 	}
 
-	const UDMMaterialLayerObject* Layer = ParentMostStage->GetLayer();
+	UDMMaterialLayerObject* Layer = ParentMostStage->GetLayer();
 	check(Layer);
 
 	bool bClampTextureActual = IsClampTextureEnabled();
@@ -276,7 +276,7 @@ void UDMMaterialStageExpressionTextureSampleBase::UpdateMask()
 		ParentMostStage = SubStage->GetParentMostStage();
 	}
 
-	const UDMMaterialLayerObject* Layer = ParentMostStage->GetLayer();
+	UDMMaterialLayerObject* Layer = ParentMostStage->GetLayer();
 	check(Layer);
 
 	const UDMMaterialSlot* Slot = Layer->GetSlot();
diff --git a/Engine/Plugins/Experimental/DynamicMaterial/Source/DynamicMaterialEditor/Private/Components/MaterialStageInputs/DMMSISlot.cpp b/Engine/Plugins/Experimental/DynamicMaterial/Source/DynamicMaterialEditor/Private/Components/MaterialStageInputs/DMMSISlot.cpp
--- a/Engine/Plugins/Experimental/DynamicMaterial/Source/DynamicMaterialEditor/Private/Components/MaterialStageInputs/DMMSISlot.cpp
+++ b/Engine/Plugins/Experimental/DynamicMaterial/Source/DynamicMaterialEditor/Private/Components/MaterialStageInputs/DMMSISlot.cpp
@@ -1,6 +1,7 @@
 // Copyright Epic Games, Inc. All Rights Reserved.
 
 #include "Components/MaterialStageInputs/DMMSISlot.h"
+
 #include "Components/DMMaterialLayer.h"
 #include "Components/DMMaterialProperty.h"
 #include "Components/DMMaterialSlot.h"
@@ -9,9 +10,7 @@
 #include "DMComponentPath.h"
 #include "DMValueDefinition.h"
 #include "Model/DMMaterialBuildState.h"
-#include "Model/DynamicMaterialModel.h"
 #include "Model/DynamicMaterialModelEditorOnlyData.h"
-#include "Utils/DMPrivate.h"
 
 #define LOCTEXT_NAMESPACE "DMMaterialStageInputSlot"
 
@@ -361,7 +360,7 @@ void UDMMaterialStageInputSlot::OnSlotConnectorsUpdated(UDMMaterialSlot* InSlot)
 	UDMMaterialStage* Stage = GetStage();
 	check(Stage);
 
-	const UDMMaterialLayerObject* Layer = Stage->GetLayer();
+	UDMMaterialLayerObject* Layer = Stage->GetLayer();
 	check(Layer);
 
 	UDMMaterialSlot* StageSlot = Layer->GetSlot();
@@ -379,7 +378,7 @@ void UDMMaterialStageInputSlot::OnSlotConnectorsUpdated(UDMMaterialSlot* InSlot)
 		EDMMaterialPropertyType StageProperty = Layer->GetMaterialProperty();
 		check(StageProperty != EDMMaterialPropertyType::None);
 
-		if (const UDMMaterialLayerObject* NextLayer = Layer->GetNextLayer(StageProperty, EDMMaterialLayerStage::Base))
+		if (UDMMaterialLayerObject* NextLayer = Layer->GetNextLayer(StageProperty, EDMMaterialLayerStage::Base))
 		{
 			NextLayer->GetStage(EDMMaterialLayerStage::Base)->ResetInputConnectionMap();
 		}
diff --git a/Engine/Plugins/Experimental/DynamicMaterial/Source/DynamicMaterialEditor/Private/Components/MaterialStageInputs/DMMSIValue.cpp b/Engine/Plugins/Experimental/DynamicMaterial/Source/DynamicMaterialEditor/Private/Components/MaterialStageInputs/DMMSIValue.cpp
--- a/Engine/Plugins/Experimental/DynamicMaterial/Source/DynamicMaterialEditor/Private/Components/MaterialStageInputs/DMMSIValue.cpp
+++ b/Engine/Plugins/Experimental/DynamicMaterial/Source/DynamicMaterialEditor/Private/Components/MaterialStageInputs/DMMSIValue.cpp
@@ -1,6 +1,7 @@
 // Copyright Epic Games, Inc. All Rights Reserved.
 
 #include "Components/MaterialStageInputs/DMMSIValue.h"
+
 #include "Components/DMMaterialLayer.h"
 #include "Components/DMMaterialParameter.h"
 #include "Components/DMMaterialSlot.h"
@@ -631,7 +632,7 @@ bool UDMMaterialStageInputValue::IsSharedStageValue() const
 		MyStage = SubStage->GetParentMostStage();
 	}
 
-	const UDMMaterialLayerObject* Layer = MyStage->GetLayer();
+	UDMMaterialLayerObject* Layer = MyStage->GetLayer();
 
 	if (!Layer)
 	{
diff --git a/Engine/Plugins/Experimental/DynamicMaterial/Source/DynamicMaterialEditor/Private/DMContentBrowserIntegration.cpp b/Engine/Plugins/Experimental/DynamicMaterial/Source/DynamicMaterialEditor/Private/DMContentBrowserIntegration.cpp
--- a/Engine/Plugins/Experimental/DynamicMaterial/Source/DynamicMaterialEditor/Private/DMContentBrowserIntegration.cpp
+++ b/Engine/Plugins/Experimental/DynamicMaterial/Source/DynamicMaterialEditor/Private/DMContentBrowserIntegration.cpp
@@ -203,7 +203,7 @@ void FDMContentBrowserIntegration::OnUpdateMaterialDesignerInstanceFromTextureSe
 		return;
 	}
 
-	FScopedTransaction Transaction(LOCTEXT("AddTextureSet", "Material Designer Add Texture Set"));
+	FScopedTransaction Transaction(LOCTEXT("AddTextureSet", "Add Texture Set"));
 	EditorOnlyData->Modify();
 
 	const bool bSuccess = EditorOnlyData->AddTextureSet(InTextureSet, bInReplace);
diff --git a/Engine/Plugins/Experimental/DynamicMaterial/Source/DynamicMaterialEditor/Private/DMObjectMaterialProperty.cpp b/Engine/Plugins/Experimental/DynamicMaterial/Source/DynamicMaterialEditor/Private/DMObjectMaterialProperty.cpp
--- a/Engine/Plugins/Experimental/DynamicMaterial/Source/DynamicMaterialEditor/Private/DMObjectMaterialProperty.cpp
+++ b/Engine/Plugins/Experimental/DynamicMaterial/Source/DynamicMaterialEditor/Private/DMObjectMaterialProperty.cpp
@@ -1,6 +1,7 @@
 // Copyright Epic Games, Inc. All Rights Reserved.
 
 #include "DMObjectMaterialProperty.h"
+
 #include "Components/PrimitiveComponent.h"
 #include "Containers/Array.h"
 #include "Containers/Map.h"
@@ -14,17 +15,11 @@
 #define LOCTEXT_NAMESPACE "DMObjectMaterialProperty"
 
 FDMObjectMaterialProperty::FDMObjectMaterialProperty()
-	: OuterWeak(nullptr)
-	, Property(nullptr)
-	, PropertyName(NAME_None)
-	, Index(INDEX_NONE)
 {
 }
 
 FDMObjectMaterialProperty::FDMObjectMaterialProperty(UPrimitiveComponent* InOuter, int32 InIndex)
 	: OuterWeak(InOuter)
-	, Property(nullptr)
-	, PropertyName(NAME_None)
 	, Index(InIndex)
 {
 }
@@ -32,11 +27,25 @@ FDMObjectMaterialProperty::FDMObjectMaterialProperty(UPrimitiveComponent* InOute
 FDMObjectMaterialProperty::FDMObjectMaterialProperty(UObject* InOuter, FProperty* InProperty, int32 InIndex)
 	: OuterWeak(InOuter)
 	, Property(InProperty)
-	, PropertyName(InProperty ? InProperty->GetFName() : NAME_None)
 	, Index(InIndex)
 {
 }
 
+UObject* FDMObjectMaterialProperty::GetOuter() const
+{
+	return OuterWeak.Get();
+}
+
+FProperty* FDMObjectMaterialProperty::GetProperty() const
+{
+	return Property;
+}
+
+int32 FDMObjectMaterialProperty::GetIndex() const
+{
+	return Index;
+}
+
 UDynamicMaterialModelBase* FDMObjectMaterialProperty::GetMaterialModelBase() const
 {
 	if (UDynamicMaterialInstance* Instance = GetMaterial())
@@ -71,7 +80,7 @@ UDynamicMaterialInstance* FDMObjectMaterialProperty::GetMaterial() const
 				if (ArrayHelper.IsValidIndex(Index))
 				{
 					void* Value = ArrayHelper.GetRawPtr(Index);
-					Material = *reinterpret_cast<UMaterialInterface**>(Value);
+					Material = *static_cast<UMaterialInterface**>(Value);
 				}
 			}
 		}
@@ -219,7 +228,7 @@ FText FDMObjectMaterialProperty::GetPropertyName(bool bInIgnoreNewStatus) const
 	{
 		FText PropertyNameText = Property->GetDisplayNameText();
 
-		if (FArrayProperty* ArrayProperty = CastField<FArrayProperty>(Property))
+		if (CastField<FArrayProperty>(Property))
 		{
 			PropertyNameText = FText::Format(
 				LOCTEXT("PropertyNameFormatArray", "{0} [{1}]"),
@@ -276,8 +285,17 @@ void FDMObjectMaterialProperty::Reset()
 {
 	OuterWeak = nullptr;
 	Property = nullptr;
-	PropertyName = NAME_None;
 	Index = INDEX_NONE;
 }
 
+bool FDMObjectMaterialProperty::IsProperty() const
+{
+	return !!Property;
+}
+
+bool FDMObjectMaterialProperty::IsElement() const
+{
+	return !Property;
+}
+
 #undef LOCTEXT_NAMESPACE
diff --git a/Engine/Plugins/Experimental/DynamicMaterial/Source/DynamicMaterialEditor/Private/DetailsPanel/DMMaterialInterfaceTypeCustomizer.cpp b/Engine/Plugins/Experimental/DynamicMaterial/Source/DynamicMaterialEditor/Private/DetailsPanel/DMMaterialInterfaceTypeCustomizer.cpp
--- a/Engine/Plugins/Experimental/DynamicMaterial/Source/DynamicMaterialEditor/Private/DetailsPanel/DMMaterialInterfaceTypeCustomizer.cpp
+++ b/Engine/Plugins/Experimental/DynamicMaterial/Source/DynamicMaterialEditor/Private/DetailsPanel/DMMaterialInterfaceTypeCustomizer.cpp
@@ -1,12 +1,17 @@
 // Copyright Epic Games, Inc. All Rights Reserved.
 
 #include "DetailsPanel/DMMaterialInterfaceTypeCustomizer.h"
+
+#include "DetailsPanel/Widgets/SDMDetailsPanelMaterialInterfaceWidget.h"
 #include "DetailWidgetRow.h"
-#include "DetailsPanel/Slate/SDMDetailsPanelMaterialInterfaceWidget.h"
-#include "IDetailChildrenBuilder.h"
 #include "IPropertyTypeCustomization.h"
 #include "PropertyHandle.h"
 
+bool FDMMaterialInterfaceTypeIdentifier::IsPropertyTypeCustomized(const IPropertyHandle& InPropertyHandle) const
+{
+	return InPropertyHandle.GetNumOuterObjects() > 0;
+}
+
 TSharedRef<IPropertyTypeCustomization> FDMMaterialInterfaceTypeCustomizer::MakeInstance()
 {
 	return MakeShared<FDMMaterialInterfaceTypeCustomizer>();
diff --git a/Engine/Plugins/Experimental/DynamicMaterial/Source/DynamicMaterialEditor/Private/DetailsPanel/DMMaterialInterfaceTypeCustomizer.h b/Engine/Plugins/Experimental/DynamicMaterial/Source/DynamicMaterialEditor/Private/DetailsPanel/DMMaterialInterfaceTypeCustomizer.h
--- a/Engine/Plugins/Experimental/DynamicMaterial/Source/DynamicMaterialEditor/Private/DetailsPanel/DMMaterialInterfaceTypeCustomizer.h
+++ b/Engine/Plugins/Experimental/DynamicMaterial/Source/DynamicMaterialEditor/Private/DetailsPanel/DMMaterialInterfaceTypeCustomizer.h
@@ -4,16 +4,15 @@
 
 #include "IPropertyTypeCustomization.h"
 #include "PropertyEditorModule.h"
-#include "PropertyHandle.h"
+
 #include "Templates/SharedPointer.h"
 
 class FDMMaterialInterfaceTypeIdentifier : public IPropertyTypeIdentifier
 {
 public:
-	virtual bool IsPropertyTypeCustomized(const IPropertyHandle& InPropertyHandle) const override
-	{
-		return InPropertyHandle.GetNumOuterObjects() > 0;
-	}
+	//~ Begin IPropertyTypeIdentifier
+	virtual bool IsPropertyTypeCustomized(const IPropertyHandle& InPropertyHandle) const override;
+	//~ End IPropertyTypeIdentifier
 };
 
 class FDMMaterialInterfaceTypeCustomizer : public IPropertyTypeCustomization
diff --git a/Engine/Plugins/Experimental/DynamicMaterial/Source/DynamicMaterialEditor/Private/DetailsPanel/DMPropertyTypeCustomizer.cpp b/Engine/Plugins/Experimental/DynamicMaterial/Source/DynamicMaterialEditor/Private/DetailsPanel/DMPropertyTypeCustomizer.cpp
--- a/Engine/Plugins/Experimental/DynamicMaterial/Source/DynamicMaterialEditor/Private/DetailsPanel/DMPropertyTypeCustomizer.cpp
+++ b/Engine/Plugins/Experimental/DynamicMaterial/Source/DynamicMaterialEditor/Private/DetailsPanel/DMPropertyTypeCustomizer.cpp
@@ -1,9 +1,9 @@
 // Copyright Epic Games, Inc. All Rights Reserved.
 
-#include "DMPropertyTypeCustomizer.h"
+#include "DetailsPanel/DMPropertyTypeCustomizer.h"
+
+#include "DetailsPanel/Widgets/SDMDetailsPanelTabSpawner.h"
 #include "DetailWidgetRow.h"
-#include "DetailsPanel/Slate/SDMDetailsPanelTabSpawner.h"
-#include "IDetailChildrenBuilder.h"
 #include "IPropertyTypeCustomization.h"
 #include "PropertyHandle.h"
 
diff --git a/Engine/Plugins/Experimental/DynamicMaterial/Source/DynamicMaterialEditor/Private/DetailsPanel/DMPropertyTypeCustomizer.h b/Engine/Plugins/Experimental/DynamicMaterial/Source/DynamicMaterialEditor/Private/DetailsPanel/DMPropertyTypeCustomizer.h
--- a/Engine/Plugins/Experimental/DynamicMaterial/Source/DynamicMaterialEditor/Private/DetailsPanel/DMPropertyTypeCustomizer.h
+++ b/Engine/Plugins/Experimental/DynamicMaterial/Source/DynamicMaterialEditor/Private/DetailsPanel/DMPropertyTypeCustomizer.h
@@ -3,12 +3,8 @@
 #pragma once
 
 #include "IPropertyTypeCustomization.h"
-#include "Templates/SharedPointer.h"
 
-class FDetailWidgetRow;
-class IDetailChildrenBuilder;
-class IPropertyHandle;
-class IPropertyTypeCustomizationUtils;
+#include "Templates/SharedPointer.h"
 
 class FDMPropertyTypeCustomizer : public IPropertyTypeCustomization
 {
diff --git a/Engine/Plugins/Experimental/DynamicMaterial/Source/DynamicMaterialEditor/Private/DetailsPanel/DMValueDetailsRowExtensions.cpp b/Engine/Plugins/Experimental/DynamicMaterial/Source/DynamicMaterialEditor/Private/DetailsPanel/DMValueDetailsRowExtensions.cpp
--- a/Engine/Plugins/Experimental/DynamicMaterial/Source/DynamicMaterialEditor/Private/DetailsPanel/DMValueDetailsRowExtensions.cpp
+++ b/Engine/Plugins/Experimental/DynamicMaterial/Source/DynamicMaterialEditor/Private/DetailsPanel/DMValueDetailsRowExtensions.cpp
@@ -1,7 +1,7 @@
 // Copyright Epic Games, Inc. All Rights Reserved.
 
 #include "DetailsPanel/DMValueDetailsRowExtensions.h"
-#include "Components/DMMaterialParameter.h"
+
 #include "Components/DMMaterialValue.h"
 #include "Components/DMTextureUV.h"
 #include "Delegates/IDelegateInstance.h"
@@ -402,7 +402,7 @@ void FDMValueDetailsRowExtensions::SetValueParameterName(TWeakObjectPtr<UDMMater
 		return;
 	}
 
-	FScopedTransaction Transaction(LOCTEXT("SetParameterName", "Material Designer Set Parameter Name"));
+	FScopedTransaction Transaction(LOCTEXT("SetParameterName", "Set Parameter Name"));
 	Value->Modify();
 	Value->SetParameterName(InName);
 
diff --git a/Engine/Plugins/Experimental/DynamicMaterial/Source/DynamicMaterialEditor/Private/DynamicMaterialEditorCommands.cpp b/Engine/Plugins/Experimental/DynamicMaterial/Source/DynamicMaterialEditor/Private/DynamicMaterialEditorCommands.cpp
--- a/Engine/Plugins/Experimental/DynamicMaterial/Source/DynamicMaterialEditor/Private/DynamicMaterialEditorCommands.cpp
+++ b/Engine/Plugins/Experimental/DynamicMaterial/Source/DynamicMaterialEditor/Private/DynamicMaterialEditorCommands.cpp
@@ -57,6 +57,12 @@ void FDynamicMaterialEditorCommands::RegisterCommands()
 		, "Moves a layer down in the order, moving it closer to the bottom of the layer stack."
 		, EUserInterfaceActionType::Button
 		, FInputChord(EKeys::RightBracket, EModifierKey::Alt));
+
+	UI_COMMAND(SetShaderBallPreview
+		, "Shader Ball"
+		, "Sets the preview mesh to the Shader Ball."
+		, EUserInterfaceActionType::ToggleButton
+		, FInputChord());
 }
 
 #undef LOCTEXT_NAMESPACE
diff --git a/Engine/Plugins/Experimental/DynamicMaterial/Source/DynamicMaterialEditor/Private/DynamicMaterialEditorCommands.h b/Engine/Plugins/Experimental/DynamicMaterial/Source/DynamicMaterialEditor/Private/DynamicMaterialEditorCommands.h
--- a/Engine/Plugins/Experimental/DynamicMaterial/Source/DynamicMaterialEditor/Private/DynamicMaterialEditorCommands.h
+++ b/Engine/Plugins/Experimental/DynamicMaterial/Source/DynamicMaterialEditor/Private/DynamicMaterialEditorCommands.h
@@ -20,4 +20,6 @@ public:
 	TSharedPtr<FUICommandInfo> SelectLayerMaskStage;
 	TSharedPtr<FUICommandInfo> MoveLayerUp;
 	TSharedPtr<FUICommandInfo> MoveLayerDown;
+
+	TSharedPtr<FUICommandInfo> SetShaderBallPreview;
 };
diff --git a/Engine/Plugins/Experimental/DynamicMaterial/Source/DynamicMaterialEditor/Private/DynamicMaterialEditorModule.cpp b/Engine/Plugins/Experimental/DynamicMaterial/Source/DynamicMaterialEditor/Private/DynamicMaterialEditorModule.cpp
--- a/Engine/Plugins/Experimental/DynamicMaterial/Source/DynamicMaterialEditor/Private/DynamicMaterialEditorModule.cpp
+++ b/Engine/Plugins/Experimental/DynamicMaterial/Source/DynamicMaterialEditor/Private/DynamicMaterialEditorModule.cpp
@@ -11,19 +11,12 @@
 #include "Components/DMTextureUVDynamic.h"
 #include "Components/MaterialStageInputs/DMMSIFunction.h"
 #include "Components/MaterialStageInputs/DMMSIThroughput.h"
-#include "Components/MaterialValues/DMMaterialValueBool.h"
-#include "Components/MaterialValues/DMMaterialValueFloat1.h"
-#include "Components/MaterialValues/DMMaterialValueFloat2.h"
-#include "Components/MaterialValues/DMMaterialValueFloat3RGB.h"
-#include "Components/MaterialValues/DMMaterialValueFloat3RPY.h"
-#include "Components/MaterialValues/DMMaterialValueFloat3XYZ.h"
-#include "Components/MaterialValues/DMMaterialValueFloat4.h"
 #include "Components/MaterialValues/DMMaterialValueTexture.h"
 #include "Components/PrimitiveComponent.h"
 #include "DetailsPanel/DMMaterialInterfaceTypeCustomizer.h"
 #include "DetailsPanel/DMPropertyTypeCustomizer.h"
 #include "DetailsPanel/DMValueDetailsRowExtensions.h"
-#include "DetailsPanel/Slate/SDMMaterialListExtensionWidget.h"
+#include "DetailsPanel/Widgets/SDMMaterialListExtensionWidget.h"
 #include "DMContentBrowserIntegration.h"
 #include "DMWorldSubsystem.h"
 #include "DynamicMaterialEditorCommands.h"
@@ -41,25 +34,17 @@
 #include "Model/DynamicMaterialModelEditorOnlyData.h"
 #include "Modules/ModuleManager.h"
 #include "PropertyEditorModule.h"
-#include "Slate/Properties/Editors/SDMPropertyEditBoolValue.h"
-#include "Slate/Properties/Editors/SDMPropertyEditFloat1Value.h"
-#include "Slate/Properties/Editors/SDMPropertyEditFloat2Value.h"
-#include "Slate/Properties/Editors/SDMPropertyEditFloat3RGBValue.h"
-#include "Slate/Properties/Editors/SDMPropertyEditFloat3RPYValue.h"
-#include "Slate/Properties/Editors/SDMPropertyEditFloat3XYZValue.h"
-#include "Slate/Properties/Editors/SDMPropertyEditFloat4Value.h"
-#include "Slate/Properties/Editors/SDMPropertyEditTextureValue.h"
-#include "Slate/Properties/Generators/DMComponentPropertyRowGenerator.h"
-#include "Slate/Properties/Generators/DMInputThroughputPropertyRowGenerator.h"
-#include "Slate/Properties/Generators/DMMaterialEffectFunctionPropertyRowGenerator.h"
-#include "Slate/Properties/Generators/DMMaterialStageFunctionPropertyRowGenerator.h"
-#include "Slate/Properties/Generators/DMMaterialValueDynamicPropertyRowGenerator.h"
-#include "Slate/Properties/Generators/DMMaterialValuePropertyRowGenerator.h"
-#include "Slate/Properties/Generators/DMStagePropertyRowGenerator.h"
-#include "Slate/Properties/Generators/DMTextureUVPropertyRowGenerator.h"
-#include "Slate/Properties/Generators/DMTextureUVDynamicPropertyRowGenerator.h"
-#include "Slate/Properties/Generators/DMThroughputPropertyRowGenerator.h"
-#include "Slate/SDMEditor.h"
+#include "UI/PropertyGenerators/DMComponentPropertyRowGenerator.h"
+#include "UI/PropertyGenerators/DMInputThroughputPropertyRowGenerator.h"
+#include "UI/PropertyGenerators/DMMaterialEffectFunctionPropertyRowGenerator.h"
+#include "UI/PropertyGenerators/DMMaterialStageFunctionPropertyRowGenerator.h"
+#include "UI/PropertyGenerators/DMMaterialValueDynamicPropertyRowGenerator.h"
+#include "UI/PropertyGenerators/DMMaterialValuePropertyRowGenerator.h"
+#include "UI/PropertyGenerators/DMStagePropertyRowGenerator.h"
+#include "UI/PropertyGenerators/DMTextureUVDynamicPropertyRowGenerator.h"
+#include "UI/PropertyGenerators/DMTextureUVPropertyRowGenerator.h"
+#include "UI/PropertyGenerators/DMThroughputPropertyRowGenerator.h"
+#include "UI/Widgets/SDMMaterialDesigner.h"
 
 DEFINE_LOG_CATEGORY(LogDynamicMaterialEditor);
 
@@ -84,46 +69,11 @@ namespace UE::DynamicMaterialEditor::Private
 }
 
 const FName FDynamicMaterialEditorModule::TabId = TEXT("MaterialDesigner");
-TMap<UClass*, FDMCreateValueEditWidgetDelegate> FDynamicMaterialEditorModule::ValueEditWidgetDelegates;
 TMap<UClass*, FDMComponentPropertyRowGeneratorDelegate> FDynamicMaterialEditorModule::ComponentPropertyRowGenerators;
 TMap<UClass*, FDMGetObjectMaterialPropertiesDelegate> FDynamicMaterialEditorModule::CustomMaterialPropertyGenerators;
 FDMOnUIValueUpdate FDynamicMaterialEditorModule::OnUIValueUpdate;
 TArray<TSharedRef<IDMOnWizardCompleteCallback>> FDynamicMaterialEditorModule::OnWizardCompleteCallbacks;
 
-void FDynamicMaterialEditorModule::RegisterValueEditWidgetDelegate(UClass* InClass, FDMCreateValueEditWidgetDelegate ValueEditBodyDelegate)
-{
-	ValueEditWidgetDelegates.Emplace(InClass, ValueEditBodyDelegate);
-}
-
-FDMCreateValueEditWidgetDelegate FDynamicMaterialEditorModule::GetValueEditWidgetDelegate(UClass* InClass)
-{
-	FDMCreateValueEditWidgetDelegate* Delegate = ValueEditWidgetDelegates.Find(InClass);
-
-	if (Delegate)
-	{
-		return *Delegate;
-	}
-
-	return FDMCreateValueEditWidgetDelegate::CreateLambda([](TSharedPtr<SDMComponentEdit>, UDMMaterialValue*)->TSharedPtr<SWidget> { return nullptr; });
-}
-
-TSharedPtr<SWidget> FDynamicMaterialEditorModule::CreateEditWidgetForValue(const TSharedPtr<SDMComponentEdit>& InComponentEditWidget, UDMMaterialValue* InValue)
-{
-	if (!IsValid(InValue))
-	{
-		return SNullWidget::NullWidget;
-	}
-
-	FDMCreateValueEditWidgetDelegate ValueDelegate = GetValueEditWidgetDelegate(InValue->GetClass());
-
-	if (ValueDelegate.IsBound())
-	{
-		return ValueDelegate.Execute(InComponentEditWidget, InValue);
-	}
-
-	return SNullWidget::NullWidget;
-}
-
 void FDynamicMaterialEditorModule::RegisterComponentPropertyRowGeneratorDelegate(UClass* InClass, FDMComponentPropertyRowGeneratorDelegate InComponentPropertyRowGeneratorDelegate)
 {
 	ComponentPropertyRowGenerators.Emplace(InClass, InComponentPropertyRowGeneratorDelegate);
@@ -155,7 +105,7 @@ FDMComponentPropertyRowGeneratorDelegate FDynamicMaterialEditorModule::GetCompon
 	{
 		// Return an invalid lambda
 		static FDMComponentPropertyRowGeneratorDelegate NullLambda = FDMComponentPropertyRowGeneratorDelegate::CreateLambda(
-			[](const TSharedRef<SDMComponentEdit>& InComponentEditWidget, UDMMaterialComponent* InComponent, TArray<FDMPropertyHandle>& InOutPropertyRows,
+			[](const TSharedRef<SDMMaterialComponentEditor>& InComponentEditorWidget, UDMMaterialComponent* InComponent, TArray<FDMPropertyHandle>& InOutPropertyRows,
 				TSet<UDMMaterialComponent*>& InOutProcessedObjects)
 			{
 			});
@@ -166,7 +116,7 @@ FDMComponentPropertyRowGeneratorDelegate FDynamicMaterialEditorModule::GetCompon
 	return ComponentPropertyRowGenerators[FoundClass];
 }
 
-void FDynamicMaterialEditorModule::GeneratorComponentPropertyRows(const TSharedRef<SDMComponentEdit>& InComponentEditWidget, UDMMaterialComponent* InComponent,
+void FDynamicMaterialEditorModule::GeneratorComponentPropertyRows(const TSharedRef<SDMMaterialComponentEditor>& InComponentEditorWidget, UDMMaterialComponent* InComponent,
 	TArray<FDMPropertyHandle>& InOutPropertyRows, TSet<UDMMaterialComponent*>& InOutProcessedObjects)
 {
 	if (!IsValid(InComponent))
@@ -180,7 +130,7 @@ void FDynamicMaterialEditorModule::GeneratorComponentPropertyRows(const TSharedR
 	}
 
 	FDMComponentPropertyRowGeneratorDelegate RowGenerator = GetComponentPropertyRowGeneratorDelegate(InComponent->GetClass());
-	RowGenerator.ExecuteIfBound(InComponentEditWidget, InComponent, InOutPropertyRows, InOutProcessedObjects);
+	RowGenerator.ExecuteIfBound(InComponentEditorWidget, InComponent, InOutPropertyRows, InOutProcessedObjects);
 }
 
 void FDynamicMaterialEditorModule::RegisterCustomMaterialPropertyGenerator(UClass* InClass, FDMGetObjectMaterialPropertiesDelegate InGenerator)
@@ -282,15 +232,6 @@ void FDynamicMaterialEditorModule::StartupModule()
 
 	FDMMaterialModelDefaults::RegisterDefaultsDelegates();
 
-	RegisterValueEditWidgetDelegate<UDMMaterialValueBool,      SDMPropertyEditBoolValue>();
-	RegisterValueEditWidgetDelegate<UDMMaterialValueFloat1,    SDMPropertyEditFloat1Value>();
-	RegisterValueEditWidgetDelegate<UDMMaterialValueFloat2,    SDMPropertyEditFloat2Value>();
-	RegisterValueEditWidgetDelegate<UDMMaterialValueFloat3RGB, SDMPropertyEditFloat3RGBValue>();
-	RegisterValueEditWidgetDelegate<UDMMaterialValueFloat3RPY, SDMPropertyEditFloat3RPYValue>();
-	RegisterValueEditWidgetDelegate<UDMMaterialValueFloat3XYZ, SDMPropertyEditFloat3XYZValue>();
-	RegisterValueEditWidgetDelegate<UDMMaterialValueFloat4,    SDMPropertyEditFloat4Value>();
-	RegisterValueEditWidgetDelegate<UDMMaterialValueTexture,   SDMPropertyEditTextureValue>();
-
 	RegisterComponentPropertyRowGeneratorDelegate<UDMMaterialComponent,            FDMComponentPropertyRowGenerator>();
 	RegisterComponentPropertyRowGeneratorDelegate<UDMMaterialStage,                FDMStagePropertyRowGenerator>();
 	RegisterComponentPropertyRowGeneratorDelegate<UDMMaterialValue,                FDMMaterialValuePropertyRowGenerator>();
@@ -330,7 +271,7 @@ void FDynamicMaterialEditorModule::ShutdownModule()
 	FDynamicMaterialEditorStyle::Shutdown();
 	FDMContentBrowserIntegration::Disintegrate();
 
-	if (UObjectInitialized() && !IsEngineExitRequested() && FModuleManager::Get().IsModuleLoaded("PropertyEditor"))
+	if (FDynamicMaterialModule::AreUObjectsSafe() && FModuleManager::Get().IsModuleLoaded("PropertyEditor"))
 	{
 		FPropertyEditorModule& PropertyModule = FModuleManager::GetModuleChecked<FPropertyEditorModule>("PropertyEditor");
 		PropertyModule.UnregisterCustomPropertyTypeLayout(UDynamicMaterialModelEditorOnlyData::StaticClass()->GetFName());
@@ -354,14 +295,14 @@ void FDynamicMaterialEditorModule::OpenMaterialModel(UDynamicMaterialModelBase*
 {
 	using namespace UE::DynamicMaterialEditor::Private;
 
-	if (TSharedPtr<SDMEditor> Editor = FDMLevelEditorIntegration::GetEditorForWorld(InWorld))
+	if (TSharedPtr<SDMMaterialDesigner> Designer = FDMLevelEditorIntegration::GetMaterialDesignerForWorld(InWorld))
 	{
 		if (bInInvokeTab)
 		{
 			FDMLevelEditorIntegration::InvokeTabForWorld(InWorld);
 		}
 
-		Editor->SetMaterialModelBase(InMaterialModel);
+		Designer->OpenMaterialModelBase(InMaterialModel);
 	}
 	else if (IsValid(InWorld))
 	{
@@ -382,14 +323,14 @@ void FDynamicMaterialEditorModule::OpenMaterialObjectProperty(const FDMObjectMat
 {
 	using namespace UE::DynamicMaterialEditor::Private;
 
-	if (TSharedPtr<SDMEditor> Editor = FDMLevelEditorIntegration::GetEditorForWorld(InWorld))
+	if (TSharedPtr<SDMMaterialDesigner> Designer = FDMLevelEditorIntegration::GetMaterialDesignerForWorld(InWorld))
 	{
 		if (bInInvokeTab)
 		{
 			FDMLevelEditorIntegration::InvokeTabForWorld(InWorld);
 		}
 
-		Editor->SetMaterialObjectProperty(InObjectProperty);
+		Designer->OpenObjectMaterialProperty(InObjectProperty);
 	}
 	else if (IsValid(InWorld))
 	{
@@ -421,14 +362,14 @@ void FDynamicMaterialEditorModule::OnActorSelected(AActor* InActor, UWorld* InWo
 {
 	using namespace UE::DynamicMaterialEditor::Private;
 
-	if (TSharedPtr<SDMEditor> Editor = FDMLevelEditorIntegration::GetEditorForWorld(InWorld))
+	if (TSharedPtr<SDMMaterialDesigner> Designer = FDMLevelEditorIntegration::GetMaterialDesignerForWorld(InWorld))
 	{
 		if (bInInvokeTab)
 		{
 			FDMLevelEditorIntegration::InvokeTabForWorld(InWorld);
 		}
 
-		Editor->OnActorSelected(InActor);
+		Designer->OnActorSelected(InActor);
 	}
 	else if (IsValid(InWorld))
 	{
@@ -451,8 +392,8 @@ void FDynamicMaterialEditorModule::ClearDynamicMaterialModel(UWorld* InWorld) co
 
 TSharedRef<SWidget> FDynamicMaterialEditorModule::CreateEditor(UDynamicMaterialModelBase* InMaterialModelBase, UWorld* InAssetEditorWorld)
 {
-	TSharedRef<SDMEditor> NewEditor = SNew(SDMEditor);
-	NewEditor->SetMaterialModelBase(InMaterialModelBase);
+	TSharedRef<SDMMaterialDesigner> NewDesigner = SNew(SDMMaterialDesigner);
+	NewDesigner->OpenMaterialModelBase(InMaterialModelBase);
 
 	if (IsValid(InAssetEditorWorld))
 	{
@@ -460,19 +401,28 @@ TSharedRef<SWidget> FDynamicMaterialEditorModule::CreateEditor(UDynamicMaterialM
 
 		if (IsValid(WorldSubsystem))
 		{
-			WorldSubsystem->GetGetCustomEditorModelDelegate().BindSP(NewEditor, &SDMEditor::GetMaterialModelBase);
-			WorldSubsystem->GetSetCustomEditorModelDelegate().BindSP(NewEditor, &SDMEditor::SetMaterialModelBase);
-			WorldSubsystem->GetCustomObjectPropertyEditorDelegate().BindSP(NewEditor, &SDMEditor::SetMaterialObjectProperty);
-			WorldSubsystem->GetSetCustomEditorActorDelegate().BindSP(NewEditor, &SDMEditor::OnActorSelected);
+			WorldSubsystem->GetGetCustomEditorModelDelegate().BindSP(NewDesigner, &SDMMaterialDesigner::GetMaterialModelBase);
+			WorldSubsystem->GetSetCustomEditorActorDelegate().BindSP(NewDesigner, &SDMMaterialDesigner::OnActorSelected);
+
+			TWeakPtr<SDMMaterialDesigner> NewDesignerWeak = NewDesigner;
+
+			WorldSubsystem->GetSetCustomEditorModelDelegate().BindSPLambda(
+				NewDesigner, 
+				[NewDesignerWeak](UDynamicMaterialModelBase* InMaterialModelBase)
+				{
+					NewDesignerWeak.Pin()->OpenMaterialModelBase(InMaterialModelBase);
+				});
+
+			WorldSubsystem->GetCustomObjectPropertyEditorDelegate().BindSPLambda(
+				NewDesigner,
+				[NewDesignerWeak](const FDMObjectMaterialProperty& InObjectProperty)
+				{
+					NewDesignerWeak.Pin()->OpenObjectMaterialProperty(InObjectProperty);
+				});
 		}
 	}
 
-	return NewEditor;
-}
-
-TSharedRef<SWidget> FDynamicMaterialEditorModule::CreateEmptyTabContent()
-{
-	return SDMEditor::GetEmptyContent();
+	return NewDesigner;
 }
 
 void FDynamicMaterialEditorModule::Tick(float DeltaTime)
@@ -564,9 +514,9 @@ void FDynamicMaterialEditorModule::OpenEditor(UWorld* InWorld) const
 
 UDynamicMaterialModelBase* FDynamicMaterialEditorModule::GetOpenedMaterialModel(UWorld* InWorld) const
 {
-	if (TSharedPtr<SDMEditor> Editor = FDMLevelEditorIntegration::GetEditorForWorld(InWorld))
+	if (TSharedPtr<SDMMaterialDesigner> Designer = FDMLevelEditorIntegration::GetMaterialDesignerForWorld(InWorld))
 	{
-		return Editor->GetMaterialModelBase();
+		return Designer->GetMaterialModelBase();
 	}
 
 	if (IsValid(InWorld))
diff --git a/Engine/Plugins/Experimental/DynamicMaterial/Source/DynamicMaterialEditor/Private/DynamicMaterialEditorModule.h b/Engine/Plugins/Experimental/DynamicMaterial/Source/DynamicMaterialEditor/Private/DynamicMaterialEditorModule.h
--- a/Engine/Plugins/Experimental/DynamicMaterial/Source/DynamicMaterialEditor/Private/DynamicMaterialEditorModule.h
+++ b/Engine/Plugins/Experimental/DynamicMaterial/Source/DynamicMaterialEditor/Private/DynamicMaterialEditorModule.h
@@ -6,7 +6,6 @@
 #include "TickableEditorObject.h"
 
 #include "DMEDefs.h"
-#include "DMObjectMaterialProperty.h"
 #include "Templates/SharedPointer.h"
 
 class AActor;
@@ -14,8 +13,8 @@ class FDMMaterialFunctionLibrary;
 class FUICommandList;
 class IAssetTypeActions;
 class ILevelEditor;
-class SDMComponentEdit;
-class SDMEditor;
+class SDMMaterialComponentEditor;
+class SDMMaterialEditor;
 class SDockTab;
 class SWidget;
 class UDMMaterialComponent;
@@ -36,10 +35,10 @@ namespace UE::DynamicMaterialEditor
 DECLARE_MULTICAST_DELEGATE(FDMOnUIValueUpdate);
 
 /** Takes a UMaterialValue and returns the widget used to edit it. */
-DECLARE_DELEGATE_RetVal_TwoParams(TSharedPtr<SWidget>, FDMCreateValueEditWidgetDelegate, const TSharedPtr<SDMComponentEdit>&, UDMMaterialValue*);
+DECLARE_DELEGATE_RetVal_TwoParams(TSharedPtr<SWidget>, FDMCreateValueEditWidgetDelegate, const TSharedPtr<SDMMaterialComponentEditor>&, UDMMaterialValue*);
 
 /** Creates property rows in the edit widget. */
-DECLARE_DELEGATE_FourParams(FDMComponentPropertyRowGeneratorDelegate, const TSharedRef<SDMComponentEdit>&, UDMMaterialComponent*,
+DECLARE_DELEGATE_FourParams(FDMComponentPropertyRowGeneratorDelegate, const TSharedRef<SDMMaterialComponentEditor>&, UDMMaterialComponent*,
 	TArray<FDMPropertyHandle>&, TSet<UDMMaterialComponent*>&)
 
 struct FDMBuildRequestEntry
@@ -69,22 +68,16 @@ public:
 
 	static FDynamicMaterialEditorModule& Get();
 
-	static void RegisterValueEditWidgetDelegate(UClass* InClass, FDMCreateValueEditWidgetDelegate InValueEditBodyDelegate);
-	template<class InValClass, class InEditClass> static void RegisterValueEditWidgetDelegate();
-	static FDMCreateValueEditWidgetDelegate GetValueEditWidgetDelegate(UClass* InClass);
-	static TSharedPtr<SWidget> CreateEditWidgetForValue(const TSharedPtr<SDMComponentEdit>& InComponentEditWidget, UDMMaterialValue* InValue);
-
 	static void RegisterComponentPropertyRowGeneratorDelegate(UClass* InClass, FDMComponentPropertyRowGeneratorDelegate InComponentPropertyRowGeneratorDelegate);
 	template<class InObjClass, class InGenClass> static void RegisterComponentPropertyRowGeneratorDelegate();
 	static FDMComponentPropertyRowGeneratorDelegate GetComponentPropertyRowGeneratorDelegate(UClass* InClass);
-	static void GeneratorComponentPropertyRows(const TSharedRef<SDMComponentEdit>& InComponentEditWidget, UDMMaterialComponent* InComponent, 
+	static void GeneratorComponentPropertyRows(const TSharedRef<SDMMaterialComponentEditor>& InComponentEditorWidget, UDMMaterialComponent* InComponent, 
 		TArray<FDMPropertyHandle>& InOutPropertyRows, TSet<UDMMaterialComponent*>& InOutProcessedObjects);
 
 	static FDMGetObjectMaterialPropertiesDelegate GetCustomMaterialPropertyGenerator(UClass* InClass);
 
 	/** With a provided world, the editor will bind to the MD world subsystem to receive model changes. */
 	static TSharedRef<SWidget> CreateEditor(UDynamicMaterialModelBase* InMaterialModelBase, UWorld* InAssetEditorWorld);
-	static TSharedRef<SWidget> CreateEmptyTabContent();
 
 	FDynamicMaterialEditorModule();
 
@@ -123,7 +116,6 @@ public:
 	void OnWizardComplete(UDynamicMaterialModel* InModel);
 
 protected:
-	static TMap<UClass*, FDMCreateValueEditWidgetDelegate> ValueEditWidgetDelegates;
 	static TMap<UClass*, FDMComponentPropertyRowGeneratorDelegate> ComponentPropertyRowGenerators;
 	static TMap<UClass*, FDMGetObjectMaterialPropertiesDelegate> CustomMaterialPropertyGenerators;
 	static FDMOnUIValueUpdate OnUIValueUpdate;
@@ -138,17 +130,6 @@ protected:
 	void UnmapCommands();
 };
 
-template <class InValClass, class InEditClass>
-void FDynamicMaterialEditorModule::RegisterValueEditWidgetDelegate()
-{
-	RegisterValueEditWidgetDelegate(
-		InValClass::StaticClass(),
-		FDMCreateValueEditWidgetDelegate::CreateStatic(
-			&InEditClass::CreateEditWidget
-		)
-	);
-}
-
 template <class InObjClass, class InGenClass>
 void FDynamicMaterialEditorModule::RegisterComponentPropertyRowGeneratorDelegate()
 {
diff --git a/Engine/Plugins/Experimental/DynamicMaterial/Source/DynamicMaterialEditor/Private/DynamicMaterialEditorSettings.cpp b/Engine/Plugins/Experimental/DynamicMaterial/Source/DynamicMaterialEditor/Private/DynamicMaterialEditorSettings.cpp
--- a/Engine/Plugins/Experimental/DynamicMaterial/Source/DynamicMaterialEditor/Private/DynamicMaterialEditorSettings.cpp
+++ b/Engine/Plugins/Experimental/DynamicMaterial/Source/DynamicMaterialEditor/Private/DynamicMaterialEditorSettings.cpp
@@ -251,18 +251,10 @@ void UDynamicMaterialEditorSettings::PostEditChangeProperty(FPropertyChangedEven
 
 	const FName PropertyName = InPropertyChangedEvent.GetMemberPropertyName();
 
-	if (PropertyName == GET_MEMBER_NAME_CHECKED(UDynamicMaterialEditorSettings, SpinBoxValueMultiplier_Default)
-		|| PropertyName == GET_MEMBER_NAME_CHECKED(UDynamicMaterialEditorSettings, SpinBoxValueMultiplier_Shift)
-		|| PropertyName == GET_MEMBER_NAME_CHECKED(UDynamicMaterialEditorSettings, SpinBoxValueMultiplier_AltShift)
-		|| PropertyName == GET_MEMBER_NAME_CHECKED(UDynamicMaterialEditorSettings, SpinBoxValueMultiplier_Cmd)
-		|| PropertyName == GET_MEMBER_NAME_CHECKED(UDynamicMaterialEditorSettings, SpinBoxValueMultiplier_AltCmd)
-		|| PropertyName == GET_MEMBER_NAME_CHECKED(UDynamicMaterialEditorSettings, SplitterLocation)
-		|| PropertyName == GET_MEMBER_NAME_CHECKED(UDynamicMaterialEditorSettings, LayerPreviewSize)
-		|| PropertyName == GET_MEMBER_NAME_CHECKED(UDynamicMaterialEditorSettings, SlotPreviewSize)
-		|| PropertyName == GET_MEMBER_NAME_CHECKED(UDynamicMaterialEditorSettings, DetailsPreviewSize)
-		|| PropertyName == GET_MEMBER_NAME_CHECKED(UDynamicMaterialEditorSettings, bShowTooltipPreview)
-		|| PropertyName == GET_MEMBER_NAME_CHECKED(UDynamicMaterialEditorSettings, TooltipTextureSize)
-		|| PropertyName == GET_MEMBER_NAME_CHECKED(UDynamicMaterialEditorSettings, bPreviewImagesUseTextureUVs))
+	if (PropertyName == GET_MEMBER_NAME_CHECKED(UDynamicMaterialEditorSettings, SplitterLocation)
+		|| PropertyName == GET_MEMBER_NAME_CHECKED(UDynamicMaterialEditorSettings, bPreviewImagesUseTextureUVs)
+		|| PropertyName == GET_MEMBER_NAME_CHECKED(UDynamicMaterialEditorSettings, PreviewMesh)
+		|| PropertyName == GET_MEMBER_NAME_CHECKED(UDynamicMaterialEditorSettings, bShowPreviewBackground))
 	{
 		OnSettingsChanged.Broadcast(InPropertyChangedEvent);
 	}
@@ -289,41 +281,13 @@ void UDynamicMaterialEditorSettings::OpenEditorSettingsWindow() const
 
 void UDynamicMaterialEditorSettings::ResetAllLayoutSettings()
 {
-	SpinBoxValueMultiplier_AltShift = 1.0f; // Normal Sensitivity
-	SpinBoxValueMultiplier_Shift = 0.25f;   //   5x Default
-	SpinBoxValueMultiplier_Default = 0.05f; //  20x Sensitivity
-	SpinBoxValueMultiplier_Cmd = 0.01f;		//  1/5 Default
-	SpinBoxValueMultiplier_AltCmd = 0.001f; // 1/50 Default, 1/10 Cmd
 	SplitterLocation = 240;
-	LayerPreviewSize = 40;
-	SlotPreviewSize = 96;
-	DetailsPreviewSize = 96;
-	bShowTooltipPreview = false;
-	TooltipTextureSize = 512;
-	MaxFloatSliderWidth = 200.f;
+	PreviewMesh = EDMMaterialPreviewMesh::Plane;
+	bShowPreviewBackground = true;
 	bPreviewImagesUseTextureUVs = true;
 	bUVVisualizerVisible = true;
 }
 
-float UDynamicMaterialEditorSettings::GetSpinboxValueChangeMultiplier(const FModifierKeysState& InModifierKeys) const
-{
-	const bool bIsShiftDown = InModifierKeys.IsLeftShiftDown() || InModifierKeys.IsRightShiftDown();
-	const bool bIsAltDown = InModifierKeys.IsLeftAltDown() || InModifierKeys.IsRightAltDown();
-	const bool bIsCmdDown = InModifierKeys.IsLeftControlDown() || InModifierKeys.IsLeftCommandDown() || InModifierKeys.IsRightControlDown() || InModifierKeys.IsRightCommandDown();
-	
-	if (InModifierKeys.IsShiftDown())
-	{
-		return InModifierKeys.IsAltDown() ? SpinBoxValueMultiplier_AltShift : SpinBoxValueMultiplier_Shift;
-	}
-
-	if (InModifierKeys.IsCommandDown() || InModifierKeys.IsControlDown())
-	{
-		return InModifierKeys.IsAltDown() ? SpinBoxValueMultiplier_AltCmd : SpinBoxValueMultiplier_Cmd;
-	}
-
-	return SpinBoxValueMultiplier_Default;
-}
-
 TArray<FDMMaterialEffectList> UDynamicMaterialEditorSettings::GetEffectList() const
 {
 	TArray<FDMMaterialEffectList> Effects = {};
@@ -486,6 +450,11 @@ const FDMMaterialChannelListPreset* UDynamicMaterialEditorSettings::GetPresetByN
 	return nullptr;
 }
 
+FOnFinishedChangingProperties::RegistrationType& UDynamicMaterialEditorSettings::GetOnSettingsChanged()
+{
+	return OnSettingsChanged;
+}
+
 void UDynamicMaterialEditorSettings::EnsureUniqueChannelPresetNames()
 {
 	// This is splitting the *text* portion of the FName, not the number part.
diff --git a/Engine/Plugins/Experimental/DynamicMaterial/Source/DynamicMaterialEditor/Private/DynamicMaterialEditorSettings.h b/Engine/Plugins/Experimental/DynamicMaterial/Source/DynamicMaterialEditor/Private/DynamicMaterialEditorSettings.h
--- a/Engine/Plugins/Experimental/DynamicMaterial/Source/DynamicMaterialEditor/Private/DynamicMaterialEditorSettings.h
+++ b/Engine/Plugins/Experimental/DynamicMaterial/Source/DynamicMaterialEditor/Private/DynamicMaterialEditorSettings.h
@@ -119,6 +119,16 @@ struct FDMMaterialChannelListPreset
 	bool IsPropertyEnabled(EDMMaterialPropertyType InProperty) const;
 };
 
+UENUM(BlueprintType)
+enum class EDMMaterialPreviewMesh : uint8
+{
+	Plane,
+	Cube,
+	Sphere,
+	Cylinder,
+	ShaderBall
+};
+
 /**
  * Material Designer Settings
  */
@@ -134,74 +144,26 @@ public:
 	static UDynamicMaterialEditorSettings* Get();
 
 	/** Changes the currently active material in the designer following actor/object selection. */
-	UPROPERTY(Config, EditAnywhere, BlueprintReadWrite, Category = "Preview")
+	UPROPERTY(Config, EditAnywhere, BlueprintReadWrite, Category = "Editor")
 	bool bFollowSelection;
 
-	/** Adjusts the the spin box value sensitivity. */
-	UPROPERTY(Config, EditAnywhere, BlueprintReadWrite, AdvancedDisplay, Category = "Layout", meta = (
-		ClampMin = "0.01", UIMin = "0.01", ClampMax = "1.0", UIMax = "1.0"))
-	float SpinBoxValueMultiplier_Default;
-
-	/** Adjusts the the spin box value sensitivity. */
-	UPROPERTY(Config, EditAnywhere, BlueprintReadWrite, AdvancedDisplay, Category = "Layout", meta = (
-		ClampMin = "0.01", UIMin = "0.01", ClampMax = "1.0", UIMax = "1.0"))
-	float SpinBoxValueMultiplier_Shift;
-
-	/** Adjusts the the spin box value sensitivity. */
-	UPROPERTY(Config, EditAnywhere, BlueprintReadWrite, AdvancedDisplay, Category = "Layout", meta = (
-		ClampMin = "0.01", UIMin = "0.01", ClampMax = "1.0", UIMax = "1.0"))
-	float SpinBoxValueMultiplier_AltShift;
-
-	/** Adjusts the the spin box value sensitivity. */
-	UPROPERTY(Config, EditAnywhere, BlueprintReadWrite, AdvancedDisplay, Category = "Layout", meta = (
-		ClampMin = "0.01", UIMin = "0.01", ClampMax = "1.0", UIMax = "1.0"))
-	float SpinBoxValueMultiplier_Cmd;
-
-	/** Adjusts the the spin box value sensitivity. */
-	UPROPERTY(Config, EditAnywhere, BlueprintReadWrite, AdvancedDisplay, Category = "Layout", meta = (
-		ClampMin = "0.01", UIMin = "0.01", ClampMax = "1.0", UIMax = "1.0"))
-	float SpinBoxValueMultiplier_AltCmd;
-
 	/** Adjusts the vertical size of the material layer view. */
 	UPROPERTY(Config, EditAnywhere, BlueprintReadWrite, Category = "Layout", meta = (
 		ClampMin = "0.05", UIMin = "0.05", ClampMax = "0.95", UIMax = "0.95"))
 	float SplitterLocation;
 
-	/** Sets the maximum width for float-based sliders in the editor */
-	UPROPERTY(Config, EditAnywhere, BlueprintReadWRite, Category = "Layout", meta = (
-		ClampMin = "100", UIMin = "100", ClampMax = "1000", UIMax = "1000"))
-	float MaxFloatSliderWidth;
-
 	UPROPERTY(Config, EditAnywhere, BlueprintReadWrite, Category = "Layout")
 	bool bUVVisualizerVisible;
 
-	/** The size of the material layer preview images. */
-	UPROPERTY(Config, EditAnywhere, BlueprintReadWrite, Category = "Preview", meta = (ClampMin = "32", UIMin = "32", ClampMax = "128", UIMax = "128"))
-	int32 LayerPreviewSize;
-
-	/** The size of the material slot preview images. */
-	UPROPERTY(Config, EditAnywhere, BlueprintReadWrite, Category = "Preview", meta = (ClampMin = "40", UIMin = "40", ClampMax = "128", UIMax = "128"))
-	int32 SlotPreviewSize;
+	UPROPERTY(Config, EditAnywhere, BlueprintReadWrite, Category = "Preview")
+	EDMMaterialPreviewMesh PreviewMesh;
 
-	/** The size of the material layer details preview images. */
-	UPROPERTY(Config, EditAnywhere, BlueprintReadWrite, Category = "Preview", meta = (ClampMin = "32", UIMin = "32", ClampMax = "128", UIMax = "128"))
-	int32 DetailsPreviewSize;
+	UPROPERTY(Config, EditAnywhere, BlueprintReadWrite, Category = "Preview")
+	bool bShowPreviewBackground;
 
 	UPROPERTY(Config, EditAnywhere, BlueprintReadWrite, Category = "Preview")
 	bool bPreviewImagesUseTextureUVs;
 
-	/**
-	 * If true, will display the hovered preview image in an enlarged format inside the tooltip.
-	 * 
-	 * NOTE: The Material Designer must be re-opened for changes to take effect.
-	 */
-	UPROPERTY(Config, EditAnywhere, BlueprintReadWrite, Category = "ToolTips")
-	bool bShowTooltipPreview;
-
-	/** The size of the preview image when displayed in the tooltip. */
-	UPROPERTY(Config, EditAnywhere, BlueprintReadWrite, Category = "ToolTips", meta = (EditCondition = "bShowTooltipPreview", ClampMin = "128", UIMin = "128", ClampMax = "2048", UIMax = "2048"))
-	int32 TooltipTextureSize;
-
 	UPROPERTY(Config, EditAnywhere, BlueprintReadWrite, Category = "Preview")
 	TSoftObjectPtr<UTexture> DefaultMask;
 
@@ -240,28 +202,27 @@ public:
 	/** This variable is accessed in multiple places, so this is a quick accessor. */
 	static bool IsUseLinearColorForVectorsEnabled();
 
-	FOnFinishedChangingProperties OnSettingsChanged;
-
-	//~ Begin UObject
-	virtual void PostInitProperties() override;
-	virtual void PreEditChange(FEditPropertyChain& InPropertyAboutToChange) override;
-	virtual void PostEditChangeProperty(FPropertyChangedEvent& InPropertyChangedEvent) override;
-	//~ End UObject
-
 	void OpenEditorSettingsWindow() const;
 
 	void ResetAllLayoutSettings();
 
-	float GetSpinboxValueChangeMultiplier(const FModifierKeysState& InModifierKeys) const;
-
 	TArray<FDMMaterialEffectList> GetEffectList() const;
 
 	const FDMDefaultMaterialPropertySlotValue& GetDefaultSlotValue(EDMMaterialPropertyType InProperty) const;
 
 	const FDMMaterialChannelListPreset* GetPresetByName(FName InName) const;
 
+	FOnFinishedChangingProperties::RegistrationType& GetOnSettingsChanged();
+
+	//~ Begin UObject
+	virtual void PostInitProperties() override;
+	virtual void PreEditChange(FEditPropertyChain& InPropertyAboutToChange) override;
+	virtual void PostEditChangeProperty(FPropertyChangedEvent& InPropertyChangedEvent) override;
+	//~ End UObject
+
 private:
 	TArray<FName> PreEditPresetNames;
+	FOnFinishedChangingProperties OnSettingsChanged;
 
 	void EnsureUniqueChannelPresetNames();
 };
diff --git a/Engine/Plugins/Experimental/DynamicMaterial/Source/DynamicMaterialEditor/Private/DynamicMaterialEditorStyle.cpp b/Engine/Plugins/Experimental/DynamicMaterial/Source/DynamicMaterialEditor/Private/DynamicMaterialEditorStyle.cpp
--- a/Engine/Plugins/Experimental/DynamicMaterial/Source/DynamicMaterialEditor/Private/DynamicMaterialEditorStyle.cpp
+++ b/Engine/Plugins/Experimental/DynamicMaterial/Source/DynamicMaterialEditor/Private/DynamicMaterialEditorStyle.cpp
@@ -1,6 +1,7 @@
 // Copyright Epic Games, Inc. All Rights Reserved.
 
 #include "DynamicMaterialEditorStyle.h"
+
 #include "Brushes/SlateImageBrush.h"
 #include "DetailLayoutBuilder.h"
 #include "DynamicMaterialEditorModule.h"
@@ -230,40 +231,45 @@ TSharedRef<FSlateStyleSet> FDynamicMaterialEditorStyle::Create()
 
 void FDynamicMaterialEditorStyle::SetupStageStyles(const TSharedRef<FSlateStyleSet>& Style)
 {
-	constexpr FLinearColor StageEnabledColor = FLinearColor(0.0f, 1.0f, 0.0f, 0.7f);
-	constexpr FLinearColor StageDisabledColor = FLinearColor(1.0f, 0.0f, 0.0f, 0.7f);
-
-	Style->Set("Color.Stage.Enabled", StageEnabledColor);
-	Style->Set("Color.Stage.Disabled", StageDisabledColor);
+	constexpr float StageCornerRadius = 6.0f;
+	constexpr float StageBorderWidth = 2.0f;
+	constexpr float NonHoverAlpha = 1.f;
+	constexpr float HoverAlpha = 0.75f;
 
-	const float StageCornerRadius = 6.0f;
-	const float StageBorderWidth = 2.0f;
+	const FLinearColor EnabledColor = FStyleColors::Foreground.GetSpecifiedColor();
+	const FLinearColor EnabledSelectedColor = FStyleColors::Primary.GetSpecifiedColor();
+	const FLinearColor DisabledColor = FStyleColors::AccentPurple.GetSpecifiedColor();
+	const FLinearColor DisabledSelectedColor = FStyleColors::AccentRed.GetSpecifiedColor();
 
 	Style->Set("Stage.Inactive", new FSlateRoundedBoxBrush(
 		FLinearColor::Transparent, StageCornerRadius,
-		FStyleColors::Panel.GetSpecifiedColor(), 1.0f));
-	Style->Set("Stage.Inactive.Hover", new FSlateRoundedBoxBrush(
+		FStyleColors::Panel.GetSpecifiedColor(), StageBorderWidth));
+
+	Style->Set("Stage.Enabled", new FSlateRoundedBoxBrush(
+		FLinearColor::Transparent, StageCornerRadius,
+		ReplaceColorAlpha(EnabledColor, NonHoverAlpha), StageBorderWidth));
+	Style->Set("Stage.Enabled.Hover", new FSlateRoundedBoxBrush(
 		FLinearColor::Transparent, StageCornerRadius,
-		ReplaceColorAlpha(FStyleColors::Foreground.GetSpecifiedColor(), 0.2f), 1.0f));
-	Style->Set("Stage.Inactive.Select", new FSlateRoundedBoxBrush(
+		ReplaceColorAlpha(EnabledColor, HoverAlpha), StageBorderWidth));
+	Style->Set("Stage.Enabled.Select", new FSlateRoundedBoxBrush(
 		FLinearColor::Transparent, StageCornerRadius,
-		Style->GetColor("Color.Select"), 2.0f));
-	Style->Set("Stage.Inactive.Select.Hover", new FSlateRoundedBoxBrush(
+		ReplaceColorAlpha(EnabledSelectedColor, NonHoverAlpha), StageBorderWidth));
+	Style->Set("Stage.Enabled.Select.Hover", new FSlateRoundedBoxBrush(
 		FLinearColor::Transparent, StageCornerRadius,
-		Style->GetColor("Color.Select.Hover"), 2.0f));
+		ReplaceColorAlpha(EnabledSelectedColor, HoverAlpha), StageBorderWidth));
 
 	Style->Set("Stage.Disabled", new FSlateRoundedBoxBrush(
 		FLinearColor::Transparent, StageCornerRadius,
-		ReplaceColorAlpha(StageDisabledColor, 0.8f), 1.0f));
+		ReplaceColorAlpha(DisabledColor, NonHoverAlpha), StageBorderWidth));
 	Style->Set("Stage.Disabled.Hover", new FSlateRoundedBoxBrush(
 		FLinearColor::Transparent, StageCornerRadius,
-		StageDisabledColor, 1.0f));
+		ReplaceColorAlpha(DisabledColor, HoverAlpha), StageBorderWidth));
 	Style->Set("Stage.Disabled.Select", new FSlateRoundedBoxBrush(
 		FLinearColor::Transparent, StageCornerRadius,
-		ReplaceColorAlpha(FStyleColors::Select.GetSpecifiedColor(), 0.9f), 2.0f));
+		ReplaceColorAlpha(DisabledSelectedColor, NonHoverAlpha), StageBorderWidth));
 	Style->Set("Stage.Disabled.Select.Hover", new FSlateRoundedBoxBrush(
 		FLinearColor::Transparent, StageCornerRadius,
-		FStyleColors::Select.GetSpecifiedColor(), 2.0f));
+		ReplaceColorAlpha(DisabledSelectedColor, HoverAlpha), StageBorderWidth));
 
 	Style->Set("Stage.Outline", new FSlateRoundedBoxBrush(
 		FStyleColors::InputOutline.GetSpecifiedColor(), 1.0f));
@@ -288,7 +294,7 @@ void FDynamicMaterialEditorStyle::SetupLayerViewStyles(const TSharedRef<FSlateSt
 
 	/**
 	 * SListView and and FTableViewStyle have no support for adding padding between the background brush
-	 * and the SListView widget, so we are not using this style for the SDMSlotLayerView. Instead, we add a
+	 * and the SListView widget, so we are not using this style for the SDMMaterialSlotLayerView. Instead, we add a
 	 * SBorder around the SDMBLayerView and style that.
 	 */
 	Style->Set("LayerView", FTableViewStyle()
diff --git a/Engine/Plugins/Experimental/DynamicMaterial/Source/DynamicMaterialEditor/Private/LevelEditor/DMLevelEditorIntegration.cpp b/Engine/Plugins/Experimental/DynamicMaterial/Source/DynamicMaterialEditor/Private/LevelEditor/DMLevelEditorIntegration.cpp
--- a/Engine/Plugins/Experimental/DynamicMaterial/Source/DynamicMaterialEditor/Private/LevelEditor/DMLevelEditorIntegration.cpp
+++ b/Engine/Plugins/Experimental/DynamicMaterial/Source/DynamicMaterialEditor/Private/LevelEditor/DMLevelEditorIntegration.cpp
@@ -1,13 +1,15 @@
 // Copyright Epic Games, Inc. All Rights Reserved.
 
-#include "DMLevelEditorIntegration.h"
-#include "DMLevelEditorIntegrationInstance.h"
+#include "LevelEditor/DMLevelEditorIntegration.h"
+
 #include "DynamicMaterialEditorModule.h"
 #include "LevelEditor.h"
-#include "Model/DynamicMaterialModel.h"
+#include "LevelEditor/DMLevelEditorIntegrationInstance.h"
+#include "Model/DynamicMaterialModelBase.h"
 #include "Model/DynamicMaterialModelEditorOnlyData.h"
 #include "Modules/ModuleManager.h"
-#include "Slate/SDMEditor.h"
+#include "Templates/SharedPointer.h"
+#include "UI/Widgets/SDMMaterialDesigner.h"
 
 namespace UE::DynamicMaterialEditor::Private
 {
@@ -27,93 +29,7 @@ namespace UE::DynamicMaterialEditor::Private
 	FLevelEditorModule& LoadLevelEditorModuleChecked()
 	{
 		return FModuleManager::LoadModuleChecked<FLevelEditorModule>("LevelEditor");
-	}
-
-	void OnMapTearDown(UWorld* InWorld)
-	{
-		FDynamicMaterialEditorModule::Get().RemoveBuildRequestForOuter(InWorld);
-
-		FDMLevelEditorIntegrationInstance* Instance = FDMLevelEditorIntegrationInstance::GetMutableIntegrationForWorld(InWorld);
-
-		if (!Instance)
-		{
-			return;
-		}
-
-		Instance->SetLastAssetOpenPartialPath("");
-
-		TSharedPtr<SDMEditor> Editor = Instance->GetEditor();
-
-		if (!Editor.IsValid())
-		{
-			return;
-		}
-
-		UDynamicMaterialModelBase* MaterialModelBase = Editor->GetMaterialModelBase();
-
-		if (!MaterialModelBase)
-		{
-			return;
-		}
-
-		Editor->ResetEditor();
-
-		const FString WorldPath = InWorld->GetPathName();
-		const int WorldPathLength = WorldPath.Len();
-		const FString ModelPath = MaterialModelBase->GetPathName();
-
-		if (ModelPath.Len() > WorldPathLength && ModelPath.StartsWith(WorldPath))
-		{
-			switch (ModelPath[WorldPathLength])
-			{
-				case '.':
-				case '/':
-				case ':':
-					Instance->SetLastAssetOpenPartialPath(ModelPath.RightChop(WorldPath.Len()));
-					break;
-			}
-		}
-	}
-
-	void OnMapLoad(UWorld* InWorld)
-	{
-		FDMLevelEditorIntegrationInstance* Instance = FDMLevelEditorIntegrationInstance::GetMutableIntegrationForWorld(InWorld);
-
-		if (!Instance)
-		{
-			return;
-		}
-
-		const FString PartialAssetPath = Instance->GetLastOpenAssetPartialPath();
-		Instance->SetLastAssetOpenPartialPath("");
-
-		if (PartialAssetPath.IsEmpty())
-		{
-			return;
-		}
-
-		TSharedPtr<SDMEditor> Editor = Instance->GetEditor();
-
-		if (!Editor.IsValid())
-		{
-			return;
-		}
-
-		const FString ModelPath = InWorld->GetPathName() + PartialAssetPath;
-
-		if (UObject* Object = FindObject<UObject>(nullptr, *ModelPath, false))
-		{
-			if (UDynamicMaterialModel* MaterialModel = Cast<UDynamicMaterialModel>(Object))
-			{
-				Editor->SetMaterialModelBase(MaterialModel);
-
-				if (UDynamicMaterialModelEditorOnlyData* EditorOnlyData = UDynamicMaterialModelEditorOnlyData::Get(MaterialModel))
-				{
-					FDynamicMaterialEditorModule::Get().AddBuildRequest(EditorOnlyData, false);
-				}
-			}
-		}
-	}
+	}	
 }
 
 void FDMLevelEditorIntegration::Initialize()
@@ -168,7 +84,7 @@ void FDMLevelEditorIntegration::Shutdown()
 	FDMLevelEditorIntegrationInstance::RemoveIntegrations();
 }
 
-TSharedPtr<SDMEditor> FDMLevelEditorIntegration::GetEditorForWorld(UWorld* InWorld)
+TSharedPtr<SDMMaterialDesigner> FDMLevelEditorIntegration::GetMaterialDesignerForWorld(UWorld* InWorld)
 {
 	// If we have an invalid world, return the first level editor integration (for assets)... if possible
 	if (!IsValid(InWorld))
@@ -188,7 +104,7 @@ TSharedPtr<SDMEditor> FDMLevelEditorIntegration::GetEditorForWorld(UWorld* InWor
 
 	if (const FDMLevelEditorIntegrationInstance* Integration = FDMLevelEditorIntegrationInstance::GetIntegrationForWorld(InWorld))
 	{
-		return Integration->GetEditor();
+		return Integration->GetMaterialDesigner();
 	}
 
 	return nullptr;
@@ -219,3 +135,86 @@ TSharedPtr<SDockTab> FDMLevelEditorIntegration::InvokeTabForWorld(UWorld* InWorl
 
 	return nullptr;
 }
+
+void FDMLevelEditorIntegration::OnMapTearDown(UWorld* InWorld)
+{
+	FDynamicMaterialEditorModule::Get().RemoveBuildRequestForOuter(InWorld);
+
+	FDMLevelEditorIntegrationInstance* Instance = FDMLevelEditorIntegrationInstance::GetMutableIntegrationForWorld(InWorld);
+
+	if (!Instance)
+	{
+		return;
+	}
+
+	Instance->SetLastAssetOpenPartialPath("");
+
+	TSharedPtr<SDMMaterialDesigner> Designer = Instance->GetMaterialDesigner();
+
+	if (!Designer.IsValid())
+	{
+		return;
+	}
+
+	UDynamicMaterialModelBase* MaterialModelBase = Designer->GetMaterialModelBase();
+	Designer->Empty();
+
+	if (MaterialModelBase)
+	{
+		const FString WorldPath = InWorld->GetPathName();
+		const int32 WorldPathLength = WorldPath.Len();
+		const FString ModelPath = MaterialModelBase->GetPathName();
+
+		if (ModelPath.Len() > WorldPathLength && ModelPath.StartsWith(WorldPath))
+		{
+			switch (ModelPath[WorldPathLength])
+			{
+				case '.':
+				case '/':
+				case ':':
+					Instance->SetLastAssetOpenPartialPath(ModelPath.RightChop(WorldPath.Len()));
+					break;
+			}
+		}
+	}
+}
+
+void FDMLevelEditorIntegration::OnMapLoad(UWorld* InWorld)
+{
+	FDMLevelEditorIntegrationInstance* Instance = FDMLevelEditorIntegrationInstance::GetMutableIntegrationForWorld(InWorld);
+
+	if (!Instance)
+	{
+		return;
+	}
+
+	const FString PartialAssetPath = Instance->GetLastOpenAssetPartialPath();
+	Instance->SetLastAssetOpenPartialPath("");
+
+	if (PartialAssetPath.IsEmpty())
+	{
+		return;
+	}
+
+	TSharedPtr<SDMMaterialDesigner> Designer = Instance->GetMaterialDesigner();
+
+	if (!Designer.IsValid())
+	{
+		return;
+	}
+
+	const FString ModelPath = InWorld->GetPathName() + PartialAssetPath;
+
+	if (UObject* Object = FindObject<UObject>(nullptr, *ModelPath, /* Exact Class */ false))
+	{
+		if (UDynamicMaterialModelBase* MaterialModel = Cast<UDynamicMaterialModelBase>(Object))
+		{
+			Designer->OpenMaterialModelBase(MaterialModel);
+
+			if (UDynamicMaterialModelEditorOnlyData* EditorOnlyData = UDynamicMaterialModelEditorOnlyData::Get(MaterialModel))
+			{
+				FDynamicMaterialEditorModule::Get().AddBuildRequest(EditorOnlyData, /* Dirty Assets */ false);
+			}
+		}
+	}
+}
diff --git a/Engine/Plugins/Experimental/DynamicMaterial/Source/DynamicMaterialEditor/Private/LevelEditor/DMLevelEditorIntegration.h b/Engine/Plugins/Experimental/DynamicMaterial/Source/DynamicMaterialEditor/Private/LevelEditor/DMLevelEditorIntegration.h
--- a/Engine/Plugins/Experimental/DynamicMaterial/Source/DynamicMaterialEditor/Private/LevelEditor/DMLevelEditorIntegration.h
+++ b/Engine/Plugins/Experimental/DynamicMaterial/Source/DynamicMaterialEditor/Private/LevelEditor/DMLevelEditorIntegration.h
@@ -4,8 +4,7 @@
 
 #include "Templates/SharedPointerFwd.h"
 
-class ILevelEditor;
-class SDMEditor;
+class SDMMaterialDesigner;
 class SDockTab;
 class UWorld;
 
@@ -16,7 +15,12 @@ public:
 
 	static void Shutdown();
 
-	static TSharedPtr<SDMEditor> GetEditorForWorld(UWorld* InWorld);
+	static TSharedPtr<SDMMaterialDesigner> GetMaterialDesignerForWorld(UWorld* InWorld);
 
 	static TSharedPtr<SDockTab> InvokeTabForWorld(UWorld* InWorld);
+
+private:
+	static void OnMapTearDown(UWorld* InWorld);
+
+	static void OnMapLoad(UWorld* InWorld);
 };
diff --git a/Engine/Plugins/Experimental/DynamicMaterial/Source/DynamicMaterialEditor/Private/LevelEditor/DMLevelEditorIntegrationInstance.cpp b/Engine/Plugins/Experimental/DynamicMaterial/Source/DynamicMaterialEditor/Private/LevelEditor/DMLevelEditorIntegrationInstance.cpp
--- a/Engine/Plugins/Experimental/DynamicMaterial/Source/DynamicMaterialEditor/Private/LevelEditor/DMLevelEditorIntegrationInstance.cpp
+++ b/Engine/Plugins/Experimental/DynamicMaterial/Source/DynamicMaterialEditor/Private/LevelEditor/DMLevelEditorIntegrationInstance.cpp
@@ -1,8 +1,8 @@
 // Copyright Epic Games, Inc. All Rights Reserved.
 
-#include "DMLevelEditorIntegrationInstance.h"
+#include "LevelEditor/DMLevelEditorIntegrationInstance.h"
+
 #include "Components/PrimitiveComponent.h"
-#include "DMWorldSubsystem.h"
 #include "DynamicMaterialEditorModule.h"
 #include "EditorModeManager.h"
 #include "Elements/Framework/TypedElementSelectionSet.h"
@@ -11,10 +11,10 @@
 #include "ILevelEditor.h"
 #include "Material/DynamicMaterialInstance.h"
 #include "Materials/Material.h"
-#include "Model/DynamicMaterialModel.h"
+#include "Model/DynamicMaterialModelBase.h"
 #include "Selection.h"
-#include "Slate/SDMEditor.h"
 #include "Styling/SlateIconFinder.h"
+#include "UI/Widgets/SDMMaterialDesigner.h"
 #include "Widgets/Docking/SDockTab.h"
 #include "WorkspaceMenuStructure.h"
 #include "WorkspaceMenuStructureModule.h"
@@ -57,9 +57,9 @@ FDMLevelEditorIntegrationInstance::~FDMLevelEditorIntegrationInstance()
 	UnregisterWithTabManager();
 }
 
-const TSharedPtr<SDMEditor>& FDMLevelEditorIntegrationInstance::GetEditor() const
+const TSharedPtr<SDMMaterialDesigner>& FDMLevelEditorIntegrationInstance::GetMaterialDesigner() const
 {
-	return Editor;
+	return MaterialDesigner;
 }
 
 TSharedPtr<SDockTab> FDMLevelEditorIntegrationInstance::InvokeTab() const
@@ -149,7 +149,7 @@ FDMLevelEditorIntegrationInstance::FDMLevelEditorIntegrationInstance(const TShar
 {
 	LevelEditorWeak = InLevelEditor;
 
-	Editor = StaticCastSharedRef<SDMEditor>(
+	MaterialDesigner = StaticCastSharedRef<SDMMaterialDesigner>(
 		FDynamicMaterialEditorModule::CreateEditor(nullptr, nullptr)
 	);
 
@@ -229,7 +229,7 @@ void FDMLevelEditorIntegrationInstance::RegisterWithTabManager()
 					.TabRole(ETabRole::PanelTab)
 					.Content()
 					[
-						Editor.ToSharedRef()
+						MaterialDesigner.ToSharedRef()
 					];
 			}
 		)
@@ -287,27 +287,27 @@ void FDMLevelEditorIntegrationInstance::OnActorSelected(AActor* InActor)
 		return;
 	}
 
-	Editor->OnActorSelected(InActor);
+	MaterialDesigner->OnActorSelected(InActor);
 }
 
 void FDMLevelEditorIntegrationInstance::OnObjectSelectionChanged(const UTypedElementSelectionSet* InSelectionSet)
 {
-	UDynamicMaterialModel* NewSelectedMaterialModel = nullptr;
+	UDynamicMaterialModelBase* NewSelectedMaterialModelBase = nullptr;
 
-	for (UDynamicMaterialModel* MaterialModel : InSelectionSet->GetSelectedObjects<UDynamicMaterialModel>())
+	for (UDynamicMaterialModelBase* MaterialModelBase : InSelectionSet->GetSelectedObjects<UDynamicMaterialModelBase>())
 	{
 		// Only do this if we have a single selected instance.
-		if (NewSelectedMaterialModel != nullptr)
+		if (NewSelectedMaterialModelBase != nullptr)
 		{
 			return;
 		}
 
-		NewSelectedMaterialModel = MaterialModel;
+		NewSelectedMaterialModelBase = MaterialModelBase;
 	}
 
-	if (NewSelectedMaterialModel)
+	if (NewSelectedMaterialModelBase)
 	{
-		Editor->OnMaterialModelSelected(NewSelectedMaterialModel);
+		MaterialDesigner->OnMaterialModelBaseSelected(NewSelectedMaterialModelBase);
 		return;
 	}
 
@@ -323,10 +323,10 @@ void FDMLevelEditorIntegrationInstance::OnObjectSelectionChanged(const UTypedEle
 
 		NewSelectedMaterialInstance = MaterialInstance;
 	}
-
+	
 	if (NewSelectedMaterialInstance)
 	{
-		Editor->OnMaterialModelSelected(NewSelectedMaterialInstance->GetMaterialModel());
+		MaterialDesigner->OnMaterialModelBaseSelected(NewSelectedMaterialInstance->GetMaterialModelBase());
 	}
 }
 
diff --git a/Engine/Plugins/Experimental/DynamicMaterial/Source/DynamicMaterialEditor/Private/LevelEditor/DMLevelEditorIntegrationInstance.h b/Engine/Plugins/Experimental/DynamicMaterial/Source/DynamicMaterialEditor/Private/LevelEditor/DMLevelEditorIntegrationInstance.h
--- a/Engine/Plugins/Experimental/DynamicMaterial/Source/DynamicMaterialEditor/Private/LevelEditor/DMLevelEditorIntegrationInstance.h
+++ b/Engine/Plugins/Experimental/DynamicMaterial/Source/DynamicMaterialEditor/Private/LevelEditor/DMLevelEditorIntegrationInstance.h
@@ -3,14 +3,14 @@
 #pragma once
 
 #include "Containers/Array.h"
-#include "Containers/ArrayView.h"
 #include "Containers/UnrealString.h"
 #include "Templates/SharedPointer.h"
 #include "UObject/WeakObjectPtr.h"
 
 class AActor;
+class FString;
 class ILevelEditor;
-class SDMEditor;
+class SDMMaterialDesigner;
 class SDockTab;
 class UTypedElementSelectionSet;
 class UWorld;
@@ -28,7 +28,7 @@ public:
 
 	~FDMLevelEditorIntegrationInstance();
 
-	const TSharedPtr<SDMEditor>& GetEditor() const;
+	const TSharedPtr<SDMMaterialDesigner>& GetMaterialDesigner() const;
 
 	TSharedPtr<SDockTab> InvokeTab() const;
 
@@ -44,7 +44,7 @@ private:
 	TWeakPtr<ILevelEditor> LevelEditorWeak;
 	TWeakObjectPtr<UTypedElementSelectionSet> ActorSelectionSetWeak;
 	TWeakObjectPtr<UTypedElementSelectionSet> ObjectSelectionSetWeak;
-	TSharedPtr<SDMEditor> Editor;
+	TSharedPtr<SDMMaterialDesigner> MaterialDesigner;
 	FString LastOpenAssetPartialPath;
 
 	FDMLevelEditorIntegrationInstance(const TSharedRef<ILevelEditor>& InLevelEditor);
diff --git a/Engine/Plugins/Experimental/DynamicMaterial/Source/DynamicMaterialEditor/Private/Material/DynamicMaterialInstanceFactory.cpp b/Engine/Plugins/Experimental/DynamicMaterial/Source/DynamicMaterialEditor/Private/Material/DynamicMaterialInstanceFactory.cpp
--- a/Engine/Plugins/Experimental/DynamicMaterial/Source/DynamicMaterialEditor/Private/Material/DynamicMaterialInstanceFactory.cpp
+++ b/Engine/Plugins/Experimental/DynamicMaterial/Source/DynamicMaterialEditor/Private/Material/DynamicMaterialInstanceFactory.cpp
@@ -1,7 +1,9 @@
 // Copyright Epic Games, Inc. All Rights Reserved.
 
 #include "Material/DynamicMaterialInstanceFactory.h"
+
 #include "AssetToolsModule.h"
+#include "DMDefs.h"
 #include "EngineAnalytics.h"
 #include "GameFramework/Actor.h"
 #include "IAssetTools.h"
@@ -10,7 +12,6 @@
 #include "Model/DynamicMaterialModelFactory.h"
 #include "Model/IDynamicMaterialModelEditorOnlyDataInterface.h"
 #include "Modules/ModuleManager.h"
-#include "Utils/DMPrivate.h"
 
 #define LOCTEXT_NAMESPACE "MaterialDesignerInstanceFactory"
 
diff --git a/Engine/Plugins/Experimental/DynamicMaterial/Source/DynamicMaterialEditor/Private/Model/DynamicMaterialModelEditorOnlyData.cpp b/Engine/Plugins/Experimental/DynamicMaterial/Source/DynamicMaterialEditor/Private/Model/DynamicMaterialModelEditorOnlyData.cpp
--- a/Engine/Plugins/Experimental/DynamicMaterial/Source/DynamicMaterialEditor/Private/Model/DynamicMaterialModelEditorOnlyData.cpp
+++ b/Engine/Plugins/Experimental/DynamicMaterial/Source/DynamicMaterialEditor/Private/Model/DynamicMaterialModelEditorOnlyData.cpp
@@ -587,6 +587,8 @@ void UDynamicMaterialModelEditorOnlyData::BuildMaterial(bool bInDirtyAssets)
 		Pair.Value->AddOutputProcessor(BuildState);
 	}
 
+	MaterialStats = UMaterialEditingLibrary::GetStatistics(MaterialModel->DynamicMaterial);
+
 	State = EDMState::Idle;
 
 	if (IsValid(MaterialModel) && IsValid(MaterialModel->DynamicMaterialInstance))
@@ -937,6 +939,11 @@ void UDynamicMaterialModelEditorOnlyData::SetChannelListPreset(FName InPresetNam
 	OnChannelListPresetChanged();
 }
 
+const FMaterialStatistics& UDynamicMaterialModelEditorOnlyData::GetMaterialStats() const
+{
+	return MaterialStats;
+}
+
 UDMMaterialComponent* UDynamicMaterialModelEditorOnlyData::GetSubComponentByPath(FDMComponentPath& InPath,
 	const FDMComponentPathSegment& InPathSegment) const
 {
@@ -1506,13 +1513,16 @@ void UDynamicMaterialModelEditorOnlyData::PostEditImport()
 	RequestMaterialBuild();
 }
 
-void UDynamicMaterialModelEditorOnlyData::PostDuplicate(bool bDuplicateForPIE)
+void UDynamicMaterialModelEditorOnlyData::PostDuplicate(bool bInDuplicateForPIE)
 {
-	Super::PostDuplicate(bDuplicateForPIE);
+	Super::PostDuplicate(bInDuplicateForPIE);
 
-	PostEditorDuplicate();
-	ReinitComponents();
-	RequestMaterialBuild();
+	if (!bInDuplicateForPIE)
+	{
+		PostEditorDuplicate();
+		ReinitComponents();
+		RequestMaterialBuild();
+	}
 }
 
 void UDynamicMaterialModelEditorOnlyData::PostEditChangeChainProperty(FPropertyChangedChainEvent& PropertyChangedEvent)
