Material Designer

diff --git a/Engine/Plugins/Experimental/DynamicMaterial/Source/DynamicMaterialEditor/Private/UI/Menus/DMMaterialStageMenus.cpp b/Engine/Plugins/Experimental/DynamicMaterial/Source/DynamicMaterialEditor/Private/UI/Menus/DMMaterialStageMenus.cpp
new file mode 100644
--- /dev/null
+++ b/Engine/Plugins/Experimental/DynamicMaterial/Source/DynamicMaterialEditor/Private/UI/Menus/DMMaterialStageMenus.cpp
@@ -0,0 +1,177 @@
+// Copyright Epic Games, Inc. All Rights Reserved.
+
+#include "UI/Menus/DMMaterialStageMenus.h"
+
+#include "Components/DMMaterialLayer.h"
+#include "Components/DMMaterialSlot.h"
+#include "Components/DMMaterialStage.h"
+#include "ScopedTransaction.h"
+#include "Templates/SharedPointer.h"
+#include "ToolMenu.h"
+#include "UI/Menus/DMMenuContext.h"
+#include "UI/Widgets/Editor/SDMMaterialSlotEditor.h"
+#include "UI/Widgets/Editor/SlotEditor/SDMMaterialStage.h"
+#include "UI/Widgets/SDMMaterialEditor.h"
+
+#define LOCTEXT_NAMESPACE "FDMMaterialStageMenus"
+
+FName FDMMaterialStageMenus::GetStageMenuName()
+{
+	return "MaterialDesigner.MaterialStage";
+}
+
+FName FDMMaterialStageMenus::GetStageToggleSectionName()
+{
+	return "StageToggle";
+}
+
+UToolMenu* FDMMaterialStageMenus::GenerateStageMenu(const TSharedPtr<SDMMaterialSlotEditor>& InSlotWidget, 
+	const TSharedPtr<SDMMaterialStage>& InStageWidget)
+{
+	UToolMenu* NewToolMenu = UDMMenuContext::GenerateContextMenuStage(GetStageMenuName(), InSlotWidget->GetEditorWidget(), InStageWidget);
+	if (!NewToolMenu)
+	{
+		return nullptr;
+	}
+
+	AddStageSection(NewToolMenu);
+
+	return NewToolMenu;
+}
+
+void FDMMaterialStageMenus::AddStageSection(UToolMenu* InMenu)
+{
+	if (!IsValid(InMenu) || InMenu->ContainsSection(GetStageToggleSectionName()))
+	{
+		return;
+	}
+
+	const UDMMenuContext* const MenuContext = InMenu->FindContext<UDMMenuContext>();
+	if (!MenuContext)
+	{
+		return;
+	}
+
+	UDMMaterialStage* Stage = MenuContext->GetStage();
+	if (!Stage)
+	{
+		return;
+	}
+
+	UDMMaterialLayerObject* Layer = MenuContext->GetLayer();
+	if (!Layer)
+	{
+		return;
+	}
+
+	UDMMaterialSlot* Slot = Layer->GetSlot();
+	if (!Slot)
+	{
+		return;
+	}
+
+	const int32 LayerIndex = Layer->FindIndex();
+
+	const bool bAllowRemoveLayer = Slot->CanRemoveLayer(Layer);
+	const EDMMaterialLayerStage StageType = Layer->GetStageType(Stage);
+
+	if (bAllowRemoveLayer || StageType == EDMMaterialLayerStage::Mask)
+	{
+		FToolMenuSection& NewSection = InMenu->AddSection(GetStageToggleSectionName(), LOCTEXT("MaterialStageMenu", "Material Stage"));
+
+		if (bAllowRemoveLayer)
+		{
+			NewSection.AddMenuEntry(NAME_None,
+				LOCTEXT("ToggleLayer", "Toggle"),
+				LOCTEXT("ToggleLayerTooltip", "Toggle the entire layer on and off.\n\n"
+					"Warning: Toggling a layer off may result in inputs being reset where incompatibilities are found.\n\nAlt+Left Click"),
+				FSlateIcon(),
+				FUIAction(FExecuteAction::CreateWeakLambda(
+					Layer,
+					[Layer]()
+					{
+						FScopedTransaction Transaction(LOCTEXT("ToggleAllStageEnabled", "Toggle All Stage Enabled"));
+
+						for (UDMMaterialStage* Stage : Layer->GetStages(EDMMaterialLayerStage::All))
+						{
+							Stage->Modify();
+							Stage->SetEnabled(!Stage->IsEnabled());
+						}
+					}
+				))
+			);
+
+			if (StageType == EDMMaterialLayerStage::Base)
+			{
+				NewSection.AddMenuEntry(NAME_None,
+					LOCTEXT("ToggleLayerBase", "Toggle Base"),
+					LOCTEXT("ToggleLayerBaseTooltip", "Toggle the layer base on and off.\n\n"
+						"Warning: Toggling a layer base off may result in inputs being reset where incompatibilities are found.\n\nAlt+Shift+Left Click"),
+					FSlateIcon(),
+					FUIAction(FExecuteAction::CreateWeakLambda(
+						Layer,
+						[Layer]()
+						{
+							FScopedTransaction Transaction(LOCTEXT("ToggleBaseStageEnabled", "Toggle Base Stage Enabled"));
+
+							if (UDMMaterialStage* Stage = Layer->GetStage(EDMMaterialLayerStage::Base))
+							{
+								Stage->Modify();
+								Stage->SetEnabled(!Stage->IsEnabled());
+							}
+						}
+					))
+				);
+			}
+		}
+
+		if (StageType == EDMMaterialLayerStage::Mask)
+		{
+			NewSection.AddMenuEntry(NAME_None,
+				LOCTEXT("ToggleLayerMask", "Toggle Mask"),
+				LOCTEXT("ToggleLayerMaskTooltip", "Toggle the layer mask on and off.\n\nAlt+Shift+Left Click"),
+				FSlateIcon(),
+				FUIAction(FExecuteAction::CreateWeakLambda(
+					Layer,
+					[Layer]()
+					{
+						FScopedTransaction Transaction(LOCTEXT("ToggleBaseStageEnabled", "Toggle Base Stage Enabled"));
+
+						if (UDMMaterialStage* Stage = Layer->GetStage(EDMMaterialLayerStage::Mask))
+						{
+							Stage->Modify();
+							Stage->SetEnabled(!Stage->IsEnabled());
+						}
+					}
+				))
+			);
+		}
+
+		if (bAllowRemoveLayer)
+		{
+			if (TSharedPtr<SDMMaterialEditor> EditorWidget = MenuContext->GetEditorWidget())
+			{
+				NewSection.AddMenuEntry(NAME_None,
+					LOCTEXT("RemoveLayer", "Remove"),
+					LOCTEXT("RemoveLayerTooltip", "Remove this layer from its Material Slot.\n\n"
+						"Warning: Removing a stage off may result in inputs being reset where incompatibilities are found."),
+					FSlateIcon(),
+					FUIAction(FExecuteAction::CreateWeakLambda(
+						Layer,
+						[Layer]()
+						{
+							FScopedTransaction Transaction(LOCTEXT("RemoverLayer", "Remover Layer"));
+
+							if (UDMMaterialSlot* Slot = Layer->GetSlot())
+							{
+								Slot->RemoveLayer(Layer);
+							}
+						}
+					))
+				);
+			}
+		}
+	}
+}
+
+#undef LOCTEXT_NAMESPACE
diff --git a/Engine/Plugins/Experimental/DynamicMaterial/Source/DynamicMaterialEditor/Private/UI/Menus/DMMaterialStageMenus.h b/Engine/Plugins/Experimental/DynamicMaterial/Source/DynamicMaterialEditor/Private/UI/Menus/DMMaterialStageMenus.h
new file mode 100644
--- /dev/null
+++ b/Engine/Plugins/Experimental/DynamicMaterial/Source/DynamicMaterialEditor/Private/UI/Menus/DMMaterialStageMenus.h
@@ -0,0 +1,23 @@
+// Copyright Epic Games, Inc. All Rights Reserved.
+
+#pragma once
+
+#include "Templates/SharedPointerFwd.h"
+
+class FName;
+class SDMMaterialSlotEditor;
+class SDMMaterialStage;
+class UToolMenu;
+
+class FDMMaterialStageMenus final
+{
+public:
+	static UToolMenu* GenerateStageMenu(const TSharedPtr<SDMMaterialSlotEditor>& InSlotWidget, const TSharedPtr<SDMMaterialStage>& InStageWidget);
+
+private:
+	static FName GetStageMenuName();
+
+	static FName GetStageToggleSectionName();
+
+	static void AddStageSection(UToolMenu* InMenu);
+};
diff --git a/Engine/Plugins/Experimental/DynamicMaterial/Source/DynamicMaterialEditor/Private/UI/Menus/DMToolBarMenus.cpp b/Engine/Plugins/Experimental/DynamicMaterial/Source/DynamicMaterialEditor/Private/UI/Menus/DMToolBarMenus.cpp
new file mode 100644
--- /dev/null
+++ b/Engine/Plugins/Experimental/DynamicMaterial/Source/DynamicMaterialEditor/Private/UI/Menus/DMToolBarMenus.cpp
@@ -0,0 +1,489 @@
+// Copyright Epic Games, Inc. All Rights Reserved.
+
+#include "UI/Menus/DMToolBarMenus.h"
+
+#include "AssetToolsModule.h"
+#include "ContentBrowserModule.h"
+#include "DesktopPlatformModule.h"
+#include "DynamicMaterialEditorModule.h"
+#include "DynamicMaterialEditorSettings.h"
+#include "Engine/Engine.h"
+#include "Engine/Texture2D.h"
+#include "EngineAnalytics.h"
+#include "Framework/Application/SlateApplication.h"
+#include "IContentBrowserSingleton.h"
+#include "IDesktopPlatform.h"
+#include "ISinglePropertyView.h"
+#include "Material/DynamicMaterialInstance.h"
+#include "Materials/Material.h"
+#include "Model/DynamicMaterialModel.h"
+#include "Model/DynamicMaterialModelEditorOnlyData.h"
+#include "PropertyEditorModule.h"
+#include "Templates/SharedPointer.h"
+#include "ToolMenu.h"
+#include "ToolMenus.h"
+#include "UI/Menus/DMMenuContext.h"
+#include "UI/Widgets/SDMMaterialEditor.h"
+#include "Utils/DMMaterialModelFunctionLibrary.h"
+#include "Utils/DMMaterialSnapshotLibrary.h"
+#include "Utils/DMPrivate.h"
+#include "Widgets/SWidget.h"
+
+#define LOCTEXT_NAMESPACE "FDMToolBarMenus"
+
+namespace UE::DynamicMaterialEditor::Private
+{
+	static FName ToolBarEditorLayoutMenuName = TEXT("MaterialDesigner.EditorLayout");
+	static FName ToolBarMaterialExportSectionName = TEXT("MaterialExport");
+	static FName ToolBarMaterialDesignerSettingsSectionName = TEXT("MaterialDesignerSettings");
+}
+
+TSharedRef<SWidget> FDMToolBarMenus::MakeEditorLayoutMenu(const TSharedPtr<SDMMaterialEditor>& InEditorWidget)
+{
+	using namespace UE::DynamicMaterialEditor::Private;
+
+	UToolMenus* ToolMenus = UToolMenus::Get();
+
+	if (!ToolMenus->IsMenuRegistered(ToolBarEditorLayoutMenuName))
+	{
+		UToolMenu* const NewToolMenu = UDMMenuContext::GenerateContextMenuDefault(ToolBarEditorLayoutMenuName);
+
+		if (!NewToolMenu)
+		{
+			return SNullWidget::NullWidget;
+		}
+
+		NewToolMenu->AddDynamicSection(
+			TEXT("MaterialDesignerSettings"),
+			FNewToolMenuDelegate::CreateStatic(&AddToolBarEditorLayoutMenu)
+		);
+	}
+
+	FToolMenuContext MenuContext(
+		FDynamicMaterialEditorModule::Get().GetCommandList(),
+		TSharedPtr<FExtender>(),
+		UDMMenuContext::CreateEditor(InEditorWidget)
+	);
+
+	return ToolMenus->GenerateWidget(ToolBarEditorLayoutMenuName, MenuContext);
+}
+
+void FDMToolBarMenus::AddToolBarExportMenu(UToolMenu* InMenu)
+{
+	using namespace UE::DynamicMaterialEditor::Private;
+
+	if (!IsValid(InMenu) || InMenu->ContainsSection(ToolBarMaterialExportSectionName))
+	{
+		return;
+	}
+
+	UDMMenuContext* const MenuContext = InMenu->FindContext<UDMMenuContext>();
+
+	if (!MenuContext)
+	{
+		return;
+	}
+
+	UDynamicMaterialModelBase* const MaterialModelBase = MenuContext->GetModelBase();
+
+	if (!MaterialModelBase)
+	{
+		return;
+	}
+
+	UDynamicMaterialInstance* Instance = MaterialModelBase->GetDynamicMaterialInstance();
+
+	if (!Instance)
+	{
+		return;
+	}
+
+	const bool bAllowInstanceExport = IsValid(Instance->GetOuter()) && !Instance->GetOuter()->IsA<UPackage>();
+	const bool bAllowMaterialExport = IsValid(MaterialModelBase->GetGeneratedMaterial());
+
+	if (!bAllowInstanceExport && !bAllowMaterialExport)
+	{
+		return;
+	}
+
+	FToolMenuSection& NewSection = InMenu->AddSection(ToolBarMaterialExportSectionName, LOCTEXT("ExportSection", "Export"));
+
+	NewSection.AddMenuEntry(
+		NAME_None,
+		LOCTEXT("OpenInUEMaterialEditor", "Open in Standard Material Editor"),
+		LOCTEXT("OpenInUEMaterialEditorTooltip", "Opens the currently editing generated Material Designer Instance material in the standard material editor."),
+		FSlateIcon(),
+		FUIAction(FExecuteAction::CreateStatic(
+			&OpenMaterialEditorFromContext,
+			MenuContext
+		))
+	);
+
+	if (bAllowInstanceExport)
+	{
+		NewSection.AddMenuEntry(
+			NAME_None,
+			LOCTEXT("ExportMaterialInstance", "Export Material Designer Instance"),
+			LOCTEXT("ExportMaterialInstanceTooltip", "Export the material instance to an asset."),
+			FSlateIcon(),
+			FUIAction(FExecuteAction::CreateStatic(
+				&FDMToolBarMenus::ExportMaterialInstanceFromInstance,
+				TWeakObjectPtr<UDynamicMaterialInstance>(Instance)
+			))
+		);
+	}
+
+	if (bAllowMaterialExport)
+	{
+		NewSection.AddMenuEntry(
+			NAME_None,
+			LOCTEXT("ExportGeneratedMaterial", "Export Generated Material"),
+			LOCTEXT("ExportGeneratedMaterialTooltip", "Export the generated material to an asset."),
+			FSlateIcon(),
+			FUIAction(FExecuteAction::CreateStatic(
+				&FDMToolBarMenus::ExportMaterialModelFromModel,
+				TWeakObjectPtr<UDynamicMaterialModelBase>(MaterialModelBase)
+			))
+		);
+	}
+
+	NewSection.AddSubMenu(
+		NAME_None,
+		LOCTEXT("SnapshotMaterial", "Snapshop Material"),
+		LOCTEXT("SnapshotMaterialTooltip", "Take a snapshot of the material with the current settings and export it as a texture."),
+		FNewToolMenuChoice(FNewToolMenuDelegate::CreateStatic(
+			&FDMToolBarMenus::CreateSnapshotMaterialMenu
+		))
+	);
+}
+
+void FDMToolBarMenus::AddToolBarAdvancedSection(UToolMenu* InMenu)
+{
+	FToolMenuSection& NewSection = InMenu->AddSection(TEXT("AdvancedSettings"), LOCTEXT("AdvancedSettingsSection", "Advanced Settings"));
+
+	NewSection.AddMenuEntry(
+		NAME_None,
+		LOCTEXT("ResetAllSettingsToDefaults", "Reset All To Defaults"),
+		LOCTEXT("ResetAllSettingsToDefaultsTooltip", "Resets all the Material Designer settings to their default values."),
+		FSlateIcon(),
+		FUIAction(FExecuteAction::CreateUObject(
+			UDynamicMaterialEditorSettings::Get(), 
+			&UDynamicMaterialEditorSettings::ResetAllLayoutSettings)
+		)
+	);
+}
+
+void FDMToolBarMenus::AddToolBarSettingsMenu(UToolMenu* InMenu)
+{
+	using namespace UE::DynamicMaterialEditor::Private;
+
+	if (!IsValid(InMenu) || InMenu->ContainsSection(ToolBarMaterialDesignerSettingsSectionName))
+	{
+		return;
+	}
+
+	FToolMenuSection& NewSection = InMenu->AddSection(ToolBarMaterialDesignerSettingsSectionName, LOCTEXT("MaterialDesignerSection", "Material Designer"));
+
+	NewSection.AddSubMenu(
+		"AdvancedSettings",
+		LOCTEXT("AdvancedSettingsSubMenu", "Advanced Settings"),
+		LOCTEXT("AdvancedSettingsSubMenu_ToolTip", "Display advanced Material Designer settings"),
+		FNewToolMenuDelegate::CreateStatic(&AddToolBarAdvancedSection)
+	);
+
+	NewSection.AddMenuEntry(
+		NAME_None,
+		LOCTEXT("OpenSettings", "Material Designer Editor Settings"),
+		LOCTEXT("OpenSettingsTooltip", "Opens the Editor Settings and navigates to Material Designer section."),
+		FSlateIcon(FAppStyle::GetAppStyleSetName(), "FoliageEditMode.Settings"),
+		FUIAction(FExecuteAction::CreateUObject(
+			UDynamicMaterialEditorSettings::Get(),
+			&UDynamicMaterialEditorSettings::OpenEditorSettingsWindow)
+		)
+	);
+}
+
+void FDMToolBarMenus::AddToolBarEditorLayoutMenu(UToolMenu* InMenu)
+{
+	AddToolBarExportMenu(InMenu);
+	AddToolBarSettingsMenu(InMenu);
+}
+
+void FDMToolBarMenus::OpenMaterialEditorFromContext(UDMMenuContext* InMenuContext)
+{
+	if (IsValid(InMenuContext))
+	{
+		if (UDynamicMaterialModelEditorOnlyData* const ModelEditorOnlyData = UDynamicMaterialModelEditorOnlyData::Get(InMenuContext->GetModel()))
+		{
+			if (FEngineAnalytics::IsAvailable())
+			{
+				FEngineAnalytics::GetProvider().RecordEvent(TEXT("Editor.Usage.MaterialDesigner.OpenedGeneratedMaterial"));
+			}
+
+			ModelEditorOnlyData->OpenMaterialEditor();
+		}
+	}
+}
+
+void FDMToolBarMenus::ExportMaterialInstanceFromInstance(TWeakObjectPtr<UDynamicMaterialInstance> InMaterialInstanceWeak)
+{
+	if (UDynamicMaterialInstance* MaterialInstance = InMaterialInstanceWeak.Get())
+	{
+		IAssetTools& AssetTools = FModuleManager::LoadModuleChecked<FAssetToolsModule>("AssetTools").Get();
+		FString PackageName, AssetName;
+		AssetTools.CreateUniqueAssetName(MaterialInstance->GetName(), TEXT(""), PackageName, AssetName);
+
+		IContentBrowserSingleton& ContentBrowser = FModuleManager::LoadModuleChecked<FContentBrowserModule>("ContentBrowser").Get();
+		const FContentBrowserItemPath CurrentPath = ContentBrowser.GetCurrentPath();
+		const FString PathStr = CurrentPath.HasInternalPath() ? CurrentPath.GetInternalPathString() : TEXT("/Game");
+
+		FSaveAssetDialogConfig SaveAssetDialogConfig;
+		SaveAssetDialogConfig.DialogTitleOverride = LOCTEXT("SaveAssetDialogTitle", "Save Asset As");
+		SaveAssetDialogConfig.DefaultPath = PathStr;
+		SaveAssetDialogConfig.ExistingAssetPolicy = ESaveAssetDialogExistingAssetPolicy::Disallow;
+		SaveAssetDialogConfig.DefaultAssetName = AssetName;
+
+		FContentBrowserModule& ContentBrowserModule = FModuleManager::LoadModuleChecked<FContentBrowserModule>("ContentBrowser");
+		FString SaveObjectPath = ContentBrowserModule.Get().CreateModalSaveAssetDialog(SaveAssetDialogConfig);
+
+		if (!SaveObjectPath.IsEmpty())
+		{
+			UDMMaterialModelFunctionLibrary::ExportMaterialInstance(MaterialInstance->GetMaterialModelBase(), SaveObjectPath);
+
+			if (FEngineAnalytics::IsAvailable())
+			{
+				FEngineAnalytics::GetProvider().RecordEvent(TEXT("Editor.Usage.MaterialDesigner.ExportedMaterialInstance"));
+			}
+		}
+	}
+}
+
+void FDMToolBarMenus::ExportMaterialModelFromModel(TWeakObjectPtr<UDynamicMaterialModelBase> InMaterialModelBaseWeak)
+{
+	UDynamicMaterialModelBase* MaterialModelBase = InMaterialModelBaseWeak.Get();
+
+	if (!MaterialModelBase)
+	{
+		return;
+	}
+
+	UMaterial* GeneratedMaterial = MaterialModelBase->GetGeneratedMaterial();
+
+	if (!GeneratedMaterial)
+	{
+		UE::DynamicMaterialEditor::Private::LogError(TEXT("Failed to find a generated material to export."), true, MaterialModelBase);
+		return;
+	}
+
+	IContentBrowserSingleton& ContentBrowser = FModuleManager::LoadModuleChecked<FContentBrowserModule>("ContentBrowser").Get();
+	const FContentBrowserItemPath CurrentPath = ContentBrowser.GetCurrentPath();
+	const FString PathStr = CurrentPath.HasInternalPath() ? CurrentPath.GetInternalPathString() : TEXT("/Game");
+
+	FSaveAssetDialogConfig SaveAssetDialogConfig;
+	SaveAssetDialogConfig.DialogTitleOverride = LOCTEXT("SaveAssetDialogTitle", "Save Asset As");
+	SaveAssetDialogConfig.DefaultPath = PathStr;
+	SaveAssetDialogConfig.DefaultAssetName = GeneratedMaterial->GetName();
+	SaveAssetDialogConfig.ExistingAssetPolicy = ESaveAssetDialogExistingAssetPolicy::Disallow;
+
+	FContentBrowserModule& ContentBrowserModule = FModuleManager::LoadModuleChecked<FContentBrowserModule>("ContentBrowser");
+	const FString SaveObjectPath = ContentBrowserModule.Get().CreateModalSaveAssetDialog(SaveAssetDialogConfig);
+
+	if (SaveObjectPath.IsEmpty())
+	{
+		return;
+	}
+
+	UDMMaterialModelFunctionLibrary::ExportGeneratedMaterial(MaterialModelBase, SaveObjectPath);
+
+	if (FEngineAnalytics::IsAvailable())
+	{
+		FEngineAnalytics::GetProvider().RecordEvent(TEXT("Editor.Usage.MaterialDesigner.ExportedGeneratedMaterial"));
+	}
+}
+
+void FDMToolBarMenus::SnapshotMaterial(TWeakObjectPtr<UDynamicMaterialModelBase> InMaterialModelBaseWeak, FIntPoint InTextureSize)
+{
+	UDynamicMaterialModelBase* MaterialModelBase = InMaterialModelBaseWeak.Get();
+
+	if (!IsValid(MaterialModelBase))
+	{
+		return;
+	}
+
+	UMaterialInterface* Material = MaterialModelBase->GetGeneratedMaterial();
+
+	if (UDynamicMaterialInstance* MaterialInstance = MaterialModelBase->GetDynamicMaterialInstance())
+	{
+		if (!IsValid(MaterialInstance->Parent.Get()))
+		{
+			UE_LOG(LogDynamicMaterialEditor, Warning, TEXT("Unable to find world to find material instance parent."));
+			return;
+		}
+
+		Material = MaterialInstance;
+	}
+
+	if (!Material)
+	{
+		UE_LOG(LogDynamicMaterialEditor, Warning, TEXT("Unable to find material to snapshot."));
+		return;
+	}
+
+	TArray<FString> OutFilenames;
+
+	if (IDesktopPlatform* DesktopPlatform = FDesktopPlatformModule::Get())
+	{
+		DesktopPlatform->SaveFileDialog(
+			FSlateApplication::Get().FindBestParentWindowHandleForDialogs(nullptr),
+			LOCTEXT("SaveSnapshotAs", "Save Snapshot As").ToString(),
+			FPaths::ProjectSavedDir(),
+			Material->GetName() + "_Snapshot_" + FString::FromInt(InTextureSize.X) + "x" + FString::FromInt(InTextureSize.Y),
+			TEXT("HDR File (*.hdr)|*.hdr|EXR File (*.exr)|*.exr|PNG File (*.png)|*.png"),
+			EFileDialogFlags::None,
+			OutFilenames
+		);
+	}
+
+	if (OutFilenames.IsEmpty())
+	{
+		return;
+	}
+
+	FDMMaterialShapshotLibrary::SnapshotMaterial(Material, InTextureSize, OutFilenames[0]);
+
+	if (FEngineAnalytics::IsAvailable())
+	{
+		FEngineAnalytics::GetProvider().RecordEvent(TEXT("Editor.Usage.MaterialDesigner.SnapshotMaterial"));
+	}
+}
+
+void FDMToolBarMenus::AddToolBarBoolOptionMenuEntry(FToolMenuSection& InSection, const FName& InPropertyName, const FUIAction InAction)
+{
+	const FProperty* const OptionProperty = UDynamicMaterialEditorSettings::StaticClass()->FindPropertyByName(InPropertyName);
+
+	if (ensure(OptionProperty))
+	{
+		InSection.AddMenuEntry(
+			NAME_None,
+			OptionProperty->GetDisplayNameText(),
+			OptionProperty->GetToolTipText(),
+			FSlateIcon(),
+			InAction, EUserInterfaceActionType::ToggleButton
+		);
+	}
+}
+
+void FDMToolBarMenus::AddToolBarIntOptionMenuEntry(FToolMenuSection& InSection, const FName& InPropertyName,
+	TAttribute<bool> InIsEnabledAttribute, TAttribute<EVisibility> InVisibilityAttribute)
+{
+	InSection.AddDynamicEntry(NAME_None,
+		FNewToolMenuSectionDelegate::CreateLambda(
+			[InPropertyName, InIsEnabledAttribute, InVisibilityAttribute](FToolMenuSection& InSection)
+			{
+				const FProperty* const OptionProperty = UDynamicMaterialEditorSettings::StaticClass()->FindPropertyByName(InPropertyName);
+				FText DisplayName = FText::GetEmpty();
+				FText Tooltip = FText::GetEmpty();
+
+				if (ensure(OptionProperty))
+				{
+					DisplayName = OptionProperty->GetDisplayNameText();
+					Tooltip = OptionProperty->GetToolTipText();
+				}
+
+				FPropertyEditorModule& PropertyEditor = FModuleManager::Get().LoadModuleChecked<FPropertyEditorModule>(TEXT("PropertyEditor"));
+
+				FSinglePropertyParams SinglePropertyParams;
+				SinglePropertyParams.NamePlacement = EPropertyNamePlacement::Hidden;
+
+				TSharedRef<ISinglePropertyView> SinglePropertyView = PropertyEditor.CreateSingleProperty(UDynamicMaterialEditorSettings::Get(), InPropertyName, SinglePropertyParams).ToSharedRef();
+				SinglePropertyView->SetToolTipText(Tooltip);
+				SinglePropertyView->SetEnabled(InIsEnabledAttribute);
+				SinglePropertyView->SetVisibility(InVisibilityAttribute);
+
+				InSection.AddEntry(FToolMenuEntry::InitWidget(NAME_None,
+					SNew(SBox)
+					.HAlign(HAlign_Fill)
+					[
+						SNew(SBox)
+							.WidthOverride(80.0f)
+							.HAlign(HAlign_Right)
+							[
+								SinglePropertyView
+							]
+					],
+					DisplayName));
+			})
+	);
+}
+
+void FDMToolBarMenus::CreateSnapshotMaterialMenu(UToolMenu* InMenu)
+{
+	using namespace UE::DynamicMaterialEditor::Private;
+
+	const UDMMenuContext* const MenuContext = InMenu->FindContext<UDMMenuContext>();
+
+	if (!MenuContext)
+	{
+		return;
+	}
+
+	UDynamicMaterialModelBase* const MaterialModelBase = MenuContext->GetModelBase();
+
+	if (!MaterialModelBase)
+	{
+		return;
+	}
+
+	TWeakObjectPtr<UDynamicMaterialModelBase> MaterialModelWeak = MaterialModelBase;
+	FToolMenuSection& NewSection = InMenu->AddSection("SnapshotMaterial", LOCTEXT("SnapshotMaterial", "Snapshop Material"));
+
+	NewSection.AddMenuEntry(
+		NAME_None,
+		LOCTEXT("SnapshotMaterial512", "512x512"),
+		LOCTEXT("SnapshotMaterial512Tooltip", "Take a snapshot of the material with the current settings and export it as a texture with a resolution of 512x512 pixels."),
+		FSlateIcon(),
+		FUIAction(FExecuteAction::CreateStatic(
+			&FDMToolBarMenus::SnapshotMaterial,
+			MaterialModelWeak,
+			FIntPoint(512, 512)
+		))
+	);
+
+	NewSection.AddMenuEntry(
+		NAME_None,
+		LOCTEXT("SnapshotMaterial1024", "1024x1024"),
+		LOCTEXT("SnapshotMaterial1024Tooltip", "Take a snapshot of the material with the current settings and export it as a texture with a resolution of 1024x1024 pixels."),
+		FSlateIcon(),
+		FUIAction(FExecuteAction::CreateStatic(
+			&FDMToolBarMenus::SnapshotMaterial,
+			MaterialModelWeak,
+			FIntPoint(1024, 1024)
+		))
+	);
+
+	NewSection.AddMenuEntry(
+		NAME_None,
+		LOCTEXT("SnapshotMaterial2048", "2048x2048"),
+		LOCTEXT("SnapshotMaterial2048Tooltip", "Take a snapshot of the material with the current settings and export it as a texture with a resolution of 2048x2048 pixels."),
+		FSlateIcon(),
+		FUIAction(FExecuteAction::CreateStatic(
+			&FDMToolBarMenus::SnapshotMaterial,
+			MaterialModelWeak,
+			FIntPoint(2048, 2048)
+		))
+	);
+
+	NewSection.AddMenuEntry(
+		NAME_None,
+		LOCTEXT("SnapshotMaterial4096", "4096x4096"),
+		LOCTEXT("SnapshotMaterial4096Tooltip", "Take a snapshot of the material with the current settings and export it as a texture with a resolution of 4096x4096 pixels."),
+		FSlateIcon(),
+		FUIAction(FExecuteAction::CreateStatic(
+			&FDMToolBarMenus::SnapshotMaterial,
+			MaterialModelWeak,
+			FIntPoint(4096, 4096)
+		))
+	);
+}
+
+#undef LOCTEXT_NAMESPACE
diff --git a/Engine/Plugins/Experimental/DynamicMaterial/Source/DynamicMaterialEditor/Private/UI/Menus/DMToolBarMenus.h b/Engine/Plugins/Experimental/DynamicMaterial/Source/DynamicMaterialEditor/Private/UI/Menus/DMToolBarMenus.h
new file mode 100644
--- /dev/null
+++ b/Engine/Plugins/Experimental/DynamicMaterial/Source/DynamicMaterialEditor/Private/UI/Menus/DMToolBarMenus.h
@@ -0,0 +1,49 @@
+// Copyright Epic Games, Inc. All Rights Reserved.
+
+#pragma once
+
+#include "Layout/Visibility.h"
+#include "Math/MathFwd.h"
+#include "Misc/Attribute.h"
+#include "Templates/SharedPointerFwd.h"
+
+class FName;
+class SDMMaterialEditor;
+class SWidget;
+class UDMMenuContext;
+class UDynamicMaterialInstance;
+class UDynamicMaterialModelBase;
+class UToolMenu;
+struct FToolMenuSection;
+struct FUIAction;
+
+class FDMToolBarMenus final
+{
+public:
+	static TSharedRef<SWidget> MakeEditorLayoutMenu(const TSharedPtr<SDMMaterialEditor>& InEditorWidget = nullptr);
+
+private:
+	static void AddToolBarExportMenu(UToolMenu* InMenu);
+
+	static void AddToolBarAdvancedSection(UToolMenu* InMenu);
+
+	static void AddToolBarSettingsMenu(UToolMenu* InMenu);
+
+	static void AddToolBarEditorLayoutMenu(UToolMenu* InMenu);
+
+	static void AddToolBarBoolOptionMenuEntry(FToolMenuSection& InSection, const FName& InPropertyName, const FUIAction InAction);
+
+	static void AddToolBarIntOptionMenuEntry(FToolMenuSection& InSection, const FName& InPropertyName,
+		TAttribute<bool> InIsEnabledAttribute = TAttribute<bool>(),
+		TAttribute<EVisibility> InVisibilityAttribute = TAttribute<EVisibility>());
+
+	static void OpenMaterialEditorFromContext(UDMMenuContext* InMenuContext);
+
+	static void ExportMaterialInstanceFromInstance(TWeakObjectPtr<UDynamicMaterialInstance> InMaterialInstanceWeak);
+
+	static void ExportMaterialModelFromModel(TWeakObjectPtr<UDynamicMaterialModelBase> InMaterialModelBaseWeak);
+
+	static void SnapshotMaterial(TWeakObjectPtr<UDynamicMaterialModelBase> InMaterialModelBaseWeak, FIntPoint InTextureSize);
+
+	static void CreateSnapshotMaterialMenu(UToolMenu* InMenu);
+};
diff --git a/Engine/Plugins/Experimental/DynamicMaterial/Source/DynamicMaterialEditor/Private/UI/PropertyGenerators/DMComponentPropertyRowGenerator.cpp b/Engine/Plugins/Experimental/DynamicMaterial/Source/DynamicMaterialEditor/Private/UI/PropertyGenerators/DMComponentPropertyRowGenerator.cpp
new file mode 100644
--- /dev/null
+++ b/Engine/Plugins/Experimental/DynamicMaterial/Source/DynamicMaterialEditor/Private/UI/PropertyGenerators/DMComponentPropertyRowGenerator.cpp
@@ -0,0 +1,122 @@
+// Copyright Epic Games, Inc. All Rights Reserved.
+
+#include "UI/PropertyGenerators/DMComponentPropertyRowGenerator.h"
+
+#include "Components/DMMaterialComponent.h"
+#include "DynamicMaterialEditorModule.h"
+#include "Model/DynamicMaterialModelDynamic.h"
+#include "Runtime/Launch/Resources/Version.h"
+#include "UI/Utils/DMWidgetStatics.h"
+#include "UI/Widgets/Editor/SDMMaterialComponentEditor.h"
+#include "UI/Widgets/SDMMaterialEditor.h"
+
+const TSharedRef<FDMComponentPropertyRowGenerator>& FDMComponentPropertyRowGenerator::Get()
+{
+	static TSharedRef<FDMComponentPropertyRowGenerator> Generator = MakeShared<FDMComponentPropertyRowGenerator>();
+	return Generator;
+}
+
+void FDMComponentPropertyRowGenerator::AddComponentProperties(const TSharedRef<SDMMaterialComponentEditor>& InComponentEditorWidget, UDMMaterialComponent* InComponent,
+	TArray<FDMPropertyHandle>& InOutPropertyRows, TSet<UDMMaterialComponent*>& InOutProcessedObjects)
+{
+	if (!IsValid(InComponent))
+	{
+		return;
+	}
+
+	if (InOutProcessedObjects.Contains(InComponent))
+	{
+		return;
+	}
+
+	InOutProcessedObjects.Add(InComponent);
+
+	const TArray<FName>& Properties = InComponent->GetEditableProperties();
+
+	for (const FName& Property : Properties)
+	{
+		if (InComponent->IsPropertyVisible(Property))
+		{
+			AddPropertyEditRows(InComponentEditorWidget, InComponent, Property, InOutPropertyRows, InOutProcessedObjects);
+		}
+	}
+}
+
+void FDMComponentPropertyRowGenerator::AddPropertyEditRows(const TSharedRef<SDMMaterialComponentEditor>& InComponentEditorWidget, UDMMaterialComponent* InComponent, 
+	const FName& InProperty, TArray<FDMPropertyHandle>& InOutPropertyRows, TSet<UDMMaterialComponent*>& InOutProcessedObjects)
+{
+	if (!IsValid(InComponent))
+	{
+		return;
+	}
+
+	FProperty* Property = InComponent->GetClass()->FindPropertyByName(InProperty);
+
+	if (!Property)
+	{
+		return;
+	}
+
+	void* MemoryPtr = Property->ContainerPtrToValuePtr<void>(InComponent);
+
+	if (FArrayProperty* ArrayProperty = CastField<FArrayProperty>(Property))
+	{
+		FScriptArrayHelper ArrayHelper(ArrayProperty, MemoryPtr);
+
+		for (int32 Idx = 0; Idx < ArrayHelper.Num(); ++Idx)
+		{
+#if ENGINE_MAJOR_VERSION < 5 || (ENGINE_MAJOR_VERSION == 5 && ENGINE_MINOR_VERSION <= 2)
+			void* ElemPtr = ArrayHelper.GetRawPtr(Idx);
+#else
+			void* ElemPtr = ArrayHelper.GetElementPtr(Idx);
+#endif
+			AddPropertyEditRows(InComponentEditorWidget, InComponent, ArrayProperty->Inner, ElemPtr, InOutPropertyRows, InOutProcessedObjects);
+		}
+	}
+	else
+	{
+		AddPropertyEditRows(InComponentEditorWidget, InComponent, Property, MemoryPtr, InOutPropertyRows, InOutProcessedObjects);
+	}
+}
+
+void FDMComponentPropertyRowGenerator::AddPropertyEditRows(const TSharedRef<SDMMaterialComponentEditor>& InComponentEditorWidget, UDMMaterialComponent* InComponent,
+	FProperty* InProperty, void* MemoryPtr, TArray<FDMPropertyHandle>& InOutPropertyRows, TSet<UDMMaterialComponent*>& InOutProcessedObjects)
+{
+	if (InProperty->IsA<FArrayProperty>())
+	{
+		return;
+	}
+
+	if (FObjectPropertyBase* ObjectProperty = CastField<FObjectPropertyBase>(InProperty))
+	{
+		if (ObjectProperty->PropertyClass->IsChildOf(UDMMaterialComponent::StaticClass()))
+		{
+			UObject** ValuePtr = static_cast<UObject**>(MemoryPtr);
+			UObject* Value = *ValuePtr;
+			UDMMaterialComponent* ComponentValue = Cast<UDMMaterialComponent>(Value);
+			FDynamicMaterialEditorModule::GeneratorComponentPropertyRows(InComponentEditorWidget, ComponentValue, InOutPropertyRows, InOutProcessedObjects);
+			return;
+		}
+	}
+
+	FDMPropertyHandle& Handle = InOutPropertyRows.Add_GetRef(FDMWidgetStatics::Get().GetPropertyHandle(&*InComponentEditorWidget, InComponent, InProperty->GetFName()));
+	Handle.bEnabled = !IsDynamic(InComponentEditorWidget);
+}
+
+bool FDMComponentPropertyRowGenerator::AllowKeyframeButton(UDMMaterialComponent* InComponent, FProperty* InProperty)
+{
+	return false;
+}
+
+bool FDMComponentPropertyRowGenerator::IsDynamic(const TSharedRef<SDMMaterialComponentEditor>& InComponentEditorWidget)
+{
+	if (TSharedPtr<SDMMaterialEditor> EditorWidget = InComponentEditorWidget->GetEditorWidget())
+	{
+		if (UDynamicMaterialModelBase* MaterialModelBase = EditorWidget->GetMaterialModelBase())
+		{
+			return !!Cast<UDynamicMaterialModelDynamic>(MaterialModelBase);
+		}
+	}
+
+	return false;
+}
diff --git a/Engine/Plugins/Experimental/DynamicMaterial/Source/DynamicMaterialEditor/Private/UI/PropertyGenerators/DMInputThroughputPropertyRowGenerator.cpp b/Engine/Plugins/Experimental/DynamicMaterial/Source/DynamicMaterialEditor/Private/UI/PropertyGenerators/DMInputThroughputPropertyRowGenerator.cpp
new file mode 100644
--- /dev/null
+++ b/Engine/Plugins/Experimental/DynamicMaterial/Source/DynamicMaterialEditor/Private/UI/PropertyGenerators/DMInputThroughputPropertyRowGenerator.cpp
@@ -0,0 +1,51 @@
+// Copyright Epic Games, Inc. All Rights Reserved.
+
+#include "UI/PropertyGenerators/DMInputThroughputPropertyRowGenerator.h"
+
+#include "Components/DMMaterialStageThroughput.h"
+#include "Components/DMMaterialSubStage.h"
+#include "Components/MaterialStageInputs/DMMSIThroughput.h"
+
+const TSharedRef<FDMInputThroughputPropertyRowGenerator>& FDMInputThroughputPropertyRowGenerator::Get()
+{
+	static TSharedRef<FDMInputThroughputPropertyRowGenerator> Generator = MakeShared<FDMInputThroughputPropertyRowGenerator>();
+	return Generator;
+}
+
+void FDMInputThroughputPropertyRowGenerator::AddComponentProperties(const TSharedRef<SDMMaterialComponentEditor>& InComponentEditorWidget,
+	UDMMaterialComponent* InComponent, TArray<FDMPropertyHandle>& InOutPropertyRows,
+	TSet<UDMMaterialComponent*>& InOutProcessedObjects)
+{
+	if (!IsValid(InComponent))
+	{
+		return;
+	}
+
+	if (InOutProcessedObjects.Contains(InComponent))
+	{
+		return;
+	}
+
+	UDMMaterialStageInputThroughput* InputThroughput = Cast<UDMMaterialStageInputThroughput>(InComponent);
+
+	if (!InputThroughput)
+	{
+		return;
+	}
+
+	InOutProcessedObjects.Add(InputThroughput);
+
+	if (UDMMaterialSubStage* SubStage = InputThroughput->GetSubStage())
+	{
+		InOutProcessedObjects.Add(SubStage);
+	}
+
+	UDMMaterialStageThroughput* Throughput = InputThroughput->GetMaterialStageThroughput();
+
+	if (!Throughput)
+	{
+		return;
+	}
+
+	FDMThroughputPropertyRowGenerator::AddComponentProperties(InComponentEditorWidget, Throughput, InOutPropertyRows, InOutProcessedObjects);
+}
diff --git a/Engine/Plugins/Experimental/DynamicMaterial/Source/DynamicMaterialEditor/Private/UI/PropertyGenerators/DMInputThroughputPropertyRowGenerator.h b/Engine/Plugins/Experimental/DynamicMaterial/Source/DynamicMaterialEditor/Private/UI/PropertyGenerators/DMInputThroughputPropertyRowGenerator.h
new file mode 100644
--- /dev/null
+++ b/Engine/Plugins/Experimental/DynamicMaterial/Source/DynamicMaterialEditor/Private/UI/PropertyGenerators/DMInputThroughputPropertyRowGenerator.h
@@ -0,0 +1,18 @@
+// Copyright Epic Games, Inc. All Rights Reserved.
+
+#pragma once
+
+#include "UI/PropertyGenerators/DMThroughputPropertyRowGenerator.h"
+
+class FDMInputThroughputPropertyRowGenerator : public FDMThroughputPropertyRowGenerator
+{
+public:
+	static const TSharedRef<FDMInputThroughputPropertyRowGenerator>& Get();
+
+	virtual ~FDMInputThroughputPropertyRowGenerator() override = default;
+
+	//~ Begin FDMComponentPropertyRowGenerator
+	virtual void AddComponentProperties(const TSharedRef<SDMMaterialComponentEditor>& InComponentEditorWidget, UDMMaterialComponent* InComponent,
+		TArray<FDMPropertyHandle>& InOutPropertyRows, TSet<UDMMaterialComponent*>& InOutProcessedObjects) override;
+	//~ End FDMComponentPropertyRowGenerator
+};
diff --git a/Engine/Plugins/Experimental/DynamicMaterial/Source/DynamicMaterialEditor/Private/UI/PropertyGenerators/DMMaterialEffectFunctionPropertyRowGenerator.h b/Engine/Plugins/Experimental/DynamicMaterial/Source/DynamicMaterialEditor/Private/UI/PropertyGenerators/DMMaterialEffectFunctionPropertyRowGenerator.h
new file mode 100644
--- /dev/null
+++ b/Engine/Plugins/Experimental/DynamicMaterial/Source/DynamicMaterialEditor/Private/UI/PropertyGenerators/DMMaterialEffectFunctionPropertyRowGenerator.h
@@ -0,0 +1,18 @@
+// Copyright Epic Games, Inc. All Rights Reserved.
+
+#pragma once
+
+#include "UI/PropertyGenerators/DMComponentPropertyRowGenerator.h"
+
+class FDMMaterialEffectFunctionPropertyRowGenerator : public FDMComponentPropertyRowGenerator
+{
+public:
+	static const TSharedRef<FDMMaterialEffectFunctionPropertyRowGenerator>& Get();
+
+	virtual ~FDMMaterialEffectFunctionPropertyRowGenerator() override = default;
+
+	//~ Begin FDMComponentPropertyRowGenerator
+	virtual void AddComponentProperties(const TSharedRef<SDMMaterialComponentEditor>& InComponentEditorWidget, UDMMaterialComponent* InComponent,
+		TArray<FDMPropertyHandle>& InOutPropertyRows, TSet<UDMMaterialComponent*>& InOutProcessedObjects) override;
+	//~ End FDMComponentPropertyRowGenerator
+};
diff --git a/Engine/Plugins/Experimental/DynamicMaterial/Source/DynamicMaterialEditor/Private/UI/PropertyGenerators/DMMaterialModelPropertyRowGenerator.cpp b/Engine/Plugins/Experimental/DynamicMaterial/Source/DynamicMaterialEditor/Private/UI/PropertyGenerators/DMMaterialModelPropertyRowGenerator.cpp
new file mode 100644
--- /dev/null
+++ b/Engine/Plugins/Experimental/DynamicMaterial/Source/DynamicMaterialEditor/Private/UI/PropertyGenerators/DMMaterialModelPropertyRowGenerator.cpp
@@ -0,0 +1,174 @@
+// Copyright Epic Games, Inc. All Rights Reserved.
+
+#include "UI/PropertyGenerators/DMMaterialModelPropertyRowGenerator.h"
+
+#include "Components/DMMaterialProperty.h"
+#include "Components/MaterialValues/DMMaterialValueFloat1.h"
+#include "Model/DynamicMaterialModel.h"
+#include "Model/DynamicMaterialModelDynamic.h"
+#include "Model/DynamicMaterialModelEditorOnlyData.h"
+#include "UI/Utils/DMWidgetStatics.h"
+#include "UI/Widgets/Editor/SDMMaterialGlobalSettingsEditor.h"
+
+#define LOCTEXT_NAMESPACE "DMMaterialModelPropertyRowGenerator"
+
+void FDMMaterialModelPropertyRowGenerator::AddMaterialModelProperties(const TSharedRef<SDMMaterialGlobalSettingsEditor>& InGlobalSettingEditorWidget, 
+	UDynamicMaterialModelBase* InMaterialModelBase, TArray<FDMPropertyHandle>& InOutPropertyRows)
+{
+	UDynamicMaterialModel* MaterialModel = InMaterialModelBase->ResolveMaterialModel();
+
+	if (!MaterialModel)
+	{
+		return;
+	}
+
+	AddGlobalValue(InGlobalSettingEditorWidget, InMaterialModelBase, InOutPropertyRows, 
+		MaterialModel->GetGlobalParameterValue(UDynamicMaterialModel::GlobalOffsetValueName), 
+		LOCTEXT("GlobalOffset", "Global Offset"));
+
+	AddGlobalValue(InGlobalSettingEditorWidget, InMaterialModelBase, InOutPropertyRows,
+		MaterialModel->GetGlobalParameterValue(UDynamicMaterialModel::GlobalTilingValueName),
+		LOCTEXT("GlobalTiling", "Global Tiling"));
+
+	AddGlobalValue(InGlobalSettingEditorWidget, InMaterialModelBase, InOutPropertyRows,
+		MaterialModel->GetGlobalParameterValue(UDynamicMaterialModel::GlobalRotationValueName), 
+		LOCTEXT("GlobalRotation", "Global Rotation"));
+
+	if (UDynamicMaterialModelEditorOnlyData* EditorOnlyData = UDynamicMaterialModelEditorOnlyData::Get(InMaterialModelBase))
+	{
+		if (EditorOnlyData->GetBlendMode() != BLEND_Opaque)
+		{
+			AddGlobalValue(InGlobalSettingEditorWidget, InMaterialModelBase, InOutPropertyRows,  
+				MaterialModel->GetGlobalParameterValue(UDynamicMaterialModel::GlobalOpacityValueName), 
+				LOCTEXT("GlobalOpacity", "Global Opacity"));
+		}
+
+		EditorOnlyData->ForEachMaterialPropertyType(
+			[&InGlobalSettingEditorWidget, InMaterialModelBase, &InOutPropertyRows, EditorOnlyData]
+			(EDMMaterialPropertyType InProperty)
+			{
+				return ForEachProperty(InProperty, InGlobalSettingEditorWidget, InMaterialModelBase, InOutPropertyRows, EditorOnlyData);
+			},
+			/* Start from */ EDMMaterialPropertyType::Roughness
+		);
+	}
+
+	if (UDynamicMaterialModelEditorOnlyData* EditorOnlyData = UDynamicMaterialModelEditorOnlyData::Get(InMaterialModelBase))
+	{
+		AddVariable(InGlobalSettingEditorWidget, InMaterialModelBase, InOutPropertyRows, EditorOnlyData, 
+			GET_MEMBER_NAME_CHECKED(UDynamicMaterialModelEditorOnlyData, ChannelListPreset));
+
+		AddVariable(InGlobalSettingEditorWidget, InMaterialModelBase, InOutPropertyRows, EditorOnlyData,
+			GET_MEMBER_NAME_CHECKED(UDynamicMaterialModelEditorOnlyData, Domain));
+
+		AddVariable(InGlobalSettingEditorWidget, InMaterialModelBase, InOutPropertyRows, EditorOnlyData,
+			GET_MEMBER_NAME_CHECKED(UDynamicMaterialModelEditorOnlyData, BlendMode));
+
+		AddVariable(InGlobalSettingEditorWidget, InMaterialModelBase, InOutPropertyRows, EditorOnlyData,
+			GET_MEMBER_NAME_CHECKED(UDynamicMaterialModelEditorOnlyData, ShadingModel));
+
+		AddVariable(InGlobalSettingEditorWidget, InMaterialModelBase, InOutPropertyRows, EditorOnlyData,
+			GET_MEMBER_NAME_CHECKED(UDynamicMaterialModelEditorOnlyData, bHasPixelAnimation));
+
+		AddVariable(InGlobalSettingEditorWidget, InMaterialModelBase, InOutPropertyRows, EditorOnlyData,
+			GET_MEMBER_NAME_CHECKED(UDynamicMaterialModelEditorOnlyData, bTwoSided));
+
+		AddVariable(InGlobalSettingEditorWidget, InMaterialModelBase, InOutPropertyRows, EditorOnlyData,
+			GET_MEMBER_NAME_CHECKED(UDynamicMaterialModelEditorOnlyData, bResponsiveAAEnabled));
+
+		AddVariable(InGlobalSettingEditorWidget, InMaterialModelBase, InOutPropertyRows, EditorOnlyData,
+			GET_MEMBER_NAME_CHECKED(UDynamicMaterialModelEditorOnlyData, bOutputTranslucentVelocityEnabled));
+
+		AddVariable(InGlobalSettingEditorWidget, InMaterialModelBase, InOutPropertyRows, EditorOnlyData,
+			GET_MEMBER_NAME_CHECKED(UDynamicMaterialModelEditorOnlyData, bNaniteTessellationEnabled));
+	}
+}
+
+EDMIterationResult FDMMaterialModelPropertyRowGenerator::ForEachProperty(EDMMaterialPropertyType InProperty, 
+	const TSharedRef<SDMMaterialGlobalSettingsEditor>& InGlobalSettingEditorWidget, UDynamicMaterialModelBase* InMaterialModelBase, 
+	TArray<FDMPropertyHandle>& InOutPropertyRows, UDynamicMaterialModelEditorOnlyData* InEditorOnlyData)
+{
+	if (!InEditorOnlyData->GetSlotForMaterialProperty(InProperty))
+	{
+		if (UDMMaterialProperty* MaterialProperty = InEditorOnlyData->GetMaterialProperty(InProperty))
+		{
+			if (MaterialProperty->IsValidForModel(*InEditorOnlyData))
+			{
+				if (UDMMaterialValueFloat1* AlphaValue = Cast<UDMMaterialValueFloat1>(MaterialProperty->GetComponent(UDynamicMaterialModelEditorOnlyData::AlphaValueName)))
+				{
+					UEnum* MaterialPropertyEnum = StaticEnum<EDMMaterialPropertyType>();
+
+					AddGlobalValue(
+						InGlobalSettingEditorWidget,
+						InMaterialModelBase,
+						InOutPropertyRows,
+						AlphaValue,
+						FText::Format(
+							LOCTEXT("PropertyFormat", "Global {0}"),
+							MaterialPropertyEnum->GetDisplayNameTextByValue(static_cast<int64>(InProperty))
+						)
+					);
+				}
+			}
+		}
+	}
+
+	return EDMIterationResult::Continue;
+}
+
+void FDMMaterialModelPropertyRowGenerator::AddGlobalValue(const TSharedRef<SDMMaterialGlobalSettingsEditor>& InGlobalSettingEditorWidget,
+	UDynamicMaterialModelBase* InMaterialModelBase, TArray<FDMPropertyHandle>& InOutPropertyRows, UDMMaterialComponent* InComponent, 
+	const FText& InNameOverride)
+{
+	if (UDynamicMaterialModelDynamic* MaterialModelDynamic = Cast<UDynamicMaterialModelDynamic>(InMaterialModelBase))
+	{
+		InComponent = MaterialModelDynamic->GetComponentDynamic(InComponent->GetFName());
+
+		if (!InComponent)
+		{
+			return;
+		}
+	}
+
+	FDMPropertyHandle& ComponentHandle = InOutPropertyRows.Add_GetRef(FDMWidgetStatics::Get().GetPropertyHandle(&*InGlobalSettingEditorWidget,
+	                                                                                                            InComponent, UDMMaterialValue::ValueName));
+
+	ComponentHandle.CategoryOverrideName = TEXT("Material Settings");
+	ComponentHandle.NameOverride = InNameOverride;
+
+	if (UDMMaterialValue* MaterialValue = Cast<UDMMaterialValue>(InComponent))
+	{
+		ComponentHandle.ResetToDefaultOverride = FResetToDefaultOverride::Create(
+			FIsResetToDefaultVisible::CreateUObject(MaterialValue, &UDMMaterialValue::CanResetToDefault),
+			FResetToDefaultHandler::CreateUObject(MaterialValue, &UDMMaterialValue::ResetToDefault),
+			/* Propagate to children */ false
+		);
+	}
+	else if (UDMMaterialValueDynamic* MaterialValueDynamic = Cast<UDMMaterialValueDynamic>(InComponent))
+	{
+		ComponentHandle.ResetToDefaultOverride = FResetToDefaultOverride::Create(
+			FIsResetToDefaultVisible::CreateUObject(MaterialValueDynamic, &UDMMaterialValueDynamic::CanResetToDefault),
+			FResetToDefaultHandler::CreateUObject(MaterialValueDynamic, &UDMMaterialValueDynamic::ResetToDefault),
+			/* Propagate to children */ false
+		);
+	}
+}
+
+void FDMMaterialModelPropertyRowGenerator::AddVariable(const TSharedRef<SDMMaterialGlobalSettingsEditor>& InGlobalSettingEditorWidget, 
+	UDynamicMaterialModelBase* InMaterialModelBase, TArray<FDMPropertyHandle>& InOutPropertyRows, UObject* InObject, FName InPropertyName)
+{
+	FDMPropertyHandle& ValueHandle = InOutPropertyRows.Add_GetRef(FDMWidgetStatics::Get().GetPropertyHandle(
+		&*InGlobalSettingEditorWidget, 
+		InObject, InPropertyName
+	));
+
+	ValueHandle.CategoryOverrideName = TEXT("Material Type");
+	ValueHandle.bEnabled = !IsDynamic(InMaterialModelBase);
+}
+
+bool FDMMaterialModelPropertyRowGenerator::IsDynamic(UDynamicMaterialModelBase* InMaterialModelBase)
+{
+	return !!Cast<UDynamicMaterialModelDynamic>(InMaterialModelBase);
+}
+
+#undef LOCTEXT_NAMESPACE
diff --git a/Engine/Plugins/Experimental/DynamicMaterial/Source/DynamicMaterialEditor/Private/UI/PropertyGenerators/DMMaterialModelPropertyRowGenerator.h b/Engine/Plugins/Experimental/DynamicMaterial/Source/DynamicMaterialEditor/Private/UI/PropertyGenerators/DMMaterialModelPropertyRowGenerator.h
new file mode 100644
--- /dev/null
+++ b/Engine/Plugins/Experimental/DynamicMaterial/Source/DynamicMaterialEditor/Private/UI/PropertyGenerators/DMMaterialModelPropertyRowGenerator.h
@@ -0,0 +1,45 @@
+// Copyright Epic Games, Inc. All Rights Reserved.
+
+#pragma once
+
+#include "Containers/ContainersFwd.h"
+#include "Templates/SharedPointer.h"
+
+class FName;
+class FProperty;
+class SDMMaterialGlobalSettingsEditor;
+class UDMMaterialComponent;
+class UDynamicMaterialModelBase;
+class UDynamicMaterialModelEditorOnlyData;
+enum class EDMIterationResult : uint8;
+enum class EDMMaterialPropertyType : uint8;
+struct FDMPropertyHandle;
+
+/** Used to generate editable properties for material models. */
+class FDMMaterialModelPropertyRowGenerator
+{
+public:
+	/**
+	 * Generate properties for the given component.
+	 * @param InGlobalSettingEditorWidget The edit widget generating the properties.
+	 * @param InMaterialModelBase The model being edited.
+	 * @param InOutPropertyRows The generated rows.
+	 */
+	static void AddMaterialModelProperties(const TSharedRef<SDMMaterialGlobalSettingsEditor>& InGlobalSettingEditorWidget, 
+		UDynamicMaterialModelBase* InMaterialModelBase, TArray<FDMPropertyHandle>& InOutPropertyRows);
+
+	static EDMIterationResult ForEachProperty(EDMMaterialPropertyType InProperty, const TSharedRef<SDMMaterialGlobalSettingsEditor>& InGlobalSettingEditorWidget,
+		UDynamicMaterialModelBase* InMaterialModelBase, TArray<FDMPropertyHandle>& InOutPropertyRows,
+		UDynamicMaterialModelEditorOnlyData* InEditorOnlyData);
+
+	static void AddGlobalValue(const TSharedRef<SDMMaterialGlobalSettingsEditor>& InGlobalSettingEditorWidget,
+		UDynamicMaterialModelBase* InMaterialModelBase, TArray<FDMPropertyHandle>& InOutPropertyRows, UDMMaterialComponent* InComponent,
+		const FText& InNameOverride);
+
+	static void AddVariable(const TSharedRef<SDMMaterialGlobalSettingsEditor>& InGlobalSettingEditorWidget,
+		UDynamicMaterialModelBase* InMaterialModelBase, TArray<FDMPropertyHandle>& InOutPropertyRows, UObject* InObject, FName InPropertyName);
+
+protected:
+	/** Returns true if the component edit widget is editing a Material Designer Dynamic. */
+	static bool IsDynamic(UDynamicMaterialModelBase* InMaterialModelBase);
+};
diff --git a/Engine/Plugins/Experimental/DynamicMaterial/Source/DynamicMaterialEditor/Private/UI/PropertyGenerators/DMMaterialStageFunctionPropertyRowGenerator.cpp b/Engine/Plugins/Experimental/DynamicMaterial/Source/DynamicMaterialEditor/Private/UI/PropertyGenerators/DMMaterialStageFunctionPropertyRowGenerator.cpp
new file mode 100644
--- /dev/null
+++ b/Engine/Plugins/Experimental/DynamicMaterial/Source/DynamicMaterialEditor/Private/UI/PropertyGenerators/DMMaterialStageFunctionPropertyRowGenerator.cpp
@@ -0,0 +1,200 @@
+// Copyright Epic Games, Inc. All Rights Reserved.
+
+#include "UI/PropertyGenerators/DMMaterialStageFunctionPropertyRowGenerator.h"
+
+#include "Components/DMMaterialComponent.h"
+#include "Components/DMMaterialStageFunction.h"
+#include "Components/DMMaterialValue.h"
+#include "Components/MaterialStageInputs/DMMSIFunction.h"
+#include "DynamicMaterialEditorModule.h"
+#include "Materials/MaterialExpressionFunctionInput.h"
+#include "Materials/MaterialExpressionMaterialFunctionCall.h"
+#include "Materials/MaterialFunctionInterface.h"
+#include "UI/Widgets/Editor/SDMMaterialComponentEditor.h"
+
+#define LOCTEXT_NAMESPACE "DMMaterialStageFunctionPropertyRowGenerator"
+
+const TSharedRef<FDMMaterialStageFunctionPropertyRowGenerator>& FDMMaterialStageFunctionPropertyRowGenerator::Get()
+{
+	static TSharedRef<FDMMaterialStageFunctionPropertyRowGenerator> Generator = MakeShared<FDMMaterialStageFunctionPropertyRowGenerator>();
+	return Generator;
+}
+
+void FDMMaterialStageFunctionPropertyRowGenerator::ApplyMetaData(const FFunctionExpressionInput& InFunctionInput, 
+	const TSharedRef<IPropertyHandle>& InPropertyHandle)
+{
+	const TArray<FString> MetaNames = {TEXT("UIMin"), TEXT("UIMax"), TEXT("ClampMin"), TEXT("ClampMax")};
+
+	TArray<FString> MetaDatas;
+	InFunctionInput.ExpressionInput->Desc.ParseIntoArray(MetaDatas, TEXT(","));
+
+	for (FString& MetaData : MetaDatas)
+	{
+		const int32 EqualsPosition = MetaData.Find(TEXT("="));
+
+		if (EqualsPosition == INDEX_NONE)
+		{
+			continue;
+		}
+
+		const FString MetaDataName = MetaData.Left(EqualsPosition).TrimStartAndEnd();
+
+		bool bValidName = false;
+
+		for (const FString& Name : MetaNames)
+		{
+			if (MetaDataName.Equals(Name, ESearchCase::IgnoreCase))
+			{
+				bValidName = true;
+				break;
+			}
+		}
+
+		if (!bValidName)
+		{
+			continue;
+		}
+
+		const FString MetaDataValue = MetaData.Mid(EqualsPosition + 1).TrimStartAndEnd();
+		bool bValidValue = true;
+
+		for (int32 Index = 0; Index < MetaDataValue.Len(); ++Index)
+		{
+			if (MetaDataValue[Index] != '.'	&& MetaDataValue[Index] != '-' && MetaDataValue[Index] < '0' && MetaDataValue[Index] > '9')
+			{
+				bValidValue = false;
+				break;
+			}			
+		}
+
+		if (!bValidValue)
+		{
+			continue;
+		}
+
+		InPropertyHandle->SetInstanceMetaData(*MetaDataName, MetaDataValue);
+	}
+}
+
+void FDMMaterialStageFunctionPropertyRowGenerator::AddComponentProperties(const TSharedRef<SDMMaterialComponentEditor>& InComponentEditorWidget, UDMMaterialComponent* InComponent,
+	TArray<FDMPropertyHandle>& InOutPropertyRows, TSet<UDMMaterialComponent*>& InOutProcessedObjects)
+{
+	if (!IsValid(InComponent))
+	{
+		return;
+	}
+
+	if (InOutProcessedObjects.Contains(InComponent))
+	{
+		return;
+	}
+
+	UDMMaterialStageFunction* MaterialStageFunction = nullptr;
+	TArray<FDMPropertyHandle> AllValuePropertyRows;
+
+	if (UDMMaterialStageInputFunction* StageInputFunction = Cast<UDMMaterialStageInputFunction>(InComponent))
+	{
+		MaterialStageFunction = StageInputFunction->GetMaterialStageFunction();
+
+		if (MaterialStageFunction)
+		{
+			InOutProcessedObjects.Add(InComponent);
+
+			if (UMaterialFunctionInterface* MaterialFunction = MaterialStageFunction->GetMaterialFunction())
+			{
+				TArray<FFunctionExpressionInput> Inputs;
+				TArray<FFunctionExpressionOutput> Outputs;
+				MaterialFunction->GetInputsAndOutputs(Inputs, Outputs);
+
+				TArray<UDMMaterialValue*> InputValues = MaterialStageFunction->GetInputValues();
+
+				// 1 input, the previous stage, does not have a value.
+				if (Inputs.Num() != (InputValues.Num() + 1))
+				{
+					return;
+				}
+
+				for (int32 InputIndex = 0; InputIndex < InputValues.Num(); ++InputIndex)
+				{
+					FFunctionExpressionInput& Input = Inputs[InputIndex + 1];
+					UDMMaterialValue* Value = InputValues[InputIndex];
+
+					if (!IsValid(Value))
+					{
+						continue;
+					}
+
+					if (!Input.ExpressionInput)
+					{
+						continue;
+					}
+
+					TArray<FDMPropertyHandle> ValuePropertyRows;
+
+					FDynamicMaterialEditorModule::GeneratorComponentPropertyRows(
+						InComponentEditorWidget,
+						Value,
+						ValuePropertyRows,
+						InOutProcessedObjects
+					);
+
+					if (ValuePropertyRows.Num() == 1)
+					{
+						ValuePropertyRows[0].NameOverride = FText::FromName(Input.ExpressionInput->InputName);
+					}
+					else
+					{
+						static const FText NameFormat = LOCTEXT("ValueFormat", "{0}[{1}]");
+
+						for (int32 ValuePropertyIndex = 0; ValuePropertyIndex < ValuePropertyRows.Num(); ++ValuePropertyIndex)
+						{
+							ValuePropertyRows[ValuePropertyIndex].NameOverride = FText::Format(
+								NameFormat,
+								FText::FromName(Input.ExpressionInput->InputName),
+								FText::AsNumber(ValuePropertyIndex + 1)
+							);
+						}
+					}
+
+					const FText Description = FText::FromString(Input.ExpressionInput->Description);
+
+					FText MaterialInputText = FText::FromString(MaterialFunction->GetUserExposedCaption());
+
+					if (MaterialInputText.IsEmpty())
+					{
+						MaterialInputText = LOCTEXT("Function", "Function");
+					}
+
+					const FText MaterialInputFormat = LOCTEXT("MaterialInputFormat", "{0} Inputs");
+					MaterialInputText = FText::Format(MaterialInputFormat, MaterialInputText);
+					const FName MaterialInputName = *MaterialInputText.ToString();
+
+					for (FDMPropertyHandle& ValuePropertyRow : ValuePropertyRows)
+					{
+						ValuePropertyRow.NameToolTipOverride = Description;
+						ValuePropertyRow.CategoryOverrideName = MaterialInputName;
+
+						if (ValuePropertyRow.PropertyHandle.IsValid())
+						{
+							// The first input is the previous layer, so it does not have a value.
+							ApplyMetaData(Inputs[InputIndex + 1], ValuePropertyRow.PropertyHandle.ToSharedRef());
+						}
+					}
+
+					AllValuePropertyRows.Append(ValuePropertyRows);
+				}
+			}
+		}
+	}
+
+	FDynamicMaterialEditorModule::GeneratorComponentPropertyRows(
+		InComponentEditorWidget, 
+		MaterialStageFunction, 
+		InOutPropertyRows, 
+		InOutProcessedObjects
+	);
+
+	InOutPropertyRows.Append(AllValuePropertyRows);
+}
+
+#undef LOCTEXT_NAMESPACE
diff --git a/Engine/Plugins/Experimental/DynamicMaterial/Source/DynamicMaterialEditor/Private/UI/PropertyGenerators/DMMaterialStageFunctionPropertyRowGenerator.h b/Engine/Plugins/Experimental/DynamicMaterial/Source/DynamicMaterialEditor/Private/UI/PropertyGenerators/DMMaterialStageFunctionPropertyRowGenerator.h
new file mode 100644
--- /dev/null
+++ b/Engine/Plugins/Experimental/DynamicMaterial/Source/DynamicMaterialEditor/Private/UI/PropertyGenerators/DMMaterialStageFunctionPropertyRowGenerator.h
@@ -0,0 +1,23 @@
+// Copyright Epic Games, Inc. All Rights Reserved.
+
+#pragma once
+
+#include "UI/PropertyGenerators/DMComponentPropertyRowGenerator.h"
+
+class IPropertyHandle;
+struct FFunctionExpressionInput;
+
+class FDMMaterialStageFunctionPropertyRowGenerator : public FDMComponentPropertyRowGenerator
+{
+public:
+	static const TSharedRef<FDMMaterialStageFunctionPropertyRowGenerator>& Get();
+
+	virtual ~FDMMaterialStageFunctionPropertyRowGenerator() override = default;
+
+	static void ApplyMetaData(const FFunctionExpressionInput& InFunctionInput, const TSharedRef<IPropertyHandle>& InPropertyHandle);
+
+	//~ Begin FDMComponentPropertyRowGenerator
+	virtual void AddComponentProperties(const TSharedRef<SDMMaterialComponentEditor>& InComponentEditorWidget, UDMMaterialComponent* InComponent,
+		TArray<FDMPropertyHandle>& InOutPropertyRows, TSet<UDMMaterialComponent*>& InOutProcessedObjects) override;
+	//~ End FDMComponentPropertyRowGenerator
+};
diff --git a/Engine/Plugins/Experimental/DynamicMaterial/Source/DynamicMaterialEditor/Private/UI/PropertyGenerators/DMMaterialValueDynamicPropertyRowGenerator.cpp b/Engine/Plugins/Experimental/DynamicMaterial/Source/DynamicMaterialEditor/Private/UI/PropertyGenerators/DMMaterialValueDynamicPropertyRowGenerator.cpp
new file mode 100644
--- /dev/null
+++ b/Engine/Plugins/Experimental/DynamicMaterial/Source/DynamicMaterialEditor/Private/UI/PropertyGenerators/DMMaterialValueDynamicPropertyRowGenerator.cpp
@@ -0,0 +1,90 @@
+// Copyright Epic Games, Inc. All Rights Reserved.
+
+#include "UI/PropertyGenerators/DMMaterialValueDynamicPropertyRowGenerator.h"
+
+#include "Components/DMMaterialComponent.h"
+#include "Components/DMMaterialValue.h"
+#include "Components/DMMaterialValueDynamic.h"
+#include "DMEDefs.h"
+#include "IDetailPropertyRow.h"
+#include "UI/Utils/DMWidgetStatics.h"
+#include "UI/Widgets/Editor/SDMMaterialComponentEditor.h"
+
+const TSharedRef<FDMMaterialValueDynamicPropertyRowGenerator>& FDMMaterialValueDynamicPropertyRowGenerator::Get()
+{
+	static TSharedRef<FDMMaterialValueDynamicPropertyRowGenerator> Generator = MakeShared<FDMMaterialValueDynamicPropertyRowGenerator>();
+	return Generator;
+}
+
+void FDMMaterialValueDynamicPropertyRowGenerator::AddComponentProperties(const TSharedRef<SDMMaterialComponentEditor>& InComponentEditorWidget, UDMMaterialComponent* InComponent,
+	TArray<FDMPropertyHandle>& InOutPropertyRows, TSet<UDMMaterialComponent*>& InOutProcessedObjects)
+{
+	if (!IsValid(InComponent))
+	{
+		return;
+	}
+
+	if (InOutProcessedObjects.Contains(InComponent))
+	{
+		return;
+	}
+
+	UDMMaterialValueDynamic* ValueDynamic = Cast<UDMMaterialValueDynamic>(InComponent);
+
+	if (!ValueDynamic)
+	{
+		return;
+	}
+
+	// The base material value class is abstract and not allowed.
+	if (ValueDynamic->GetClass() == UDMMaterialValueDynamic::StaticClass())
+	{
+		return;
+	}
+
+	UDMMaterialValue* ParentValue = ValueDynamic->GetParentValue();
+
+	if (!ParentValue)
+	{
+		return;
+	}
+
+	InOutProcessedObjects.Add(InComponent);
+
+	if (ParentValue->AllowEditValue())
+	{
+		FDMPropertyHandle Handle = FDMWidgetStatics::Get().GetPropertyHandle(&*InComponentEditorWidget, ValueDynamic, UDMMaterialValue::ValueName);
+
+		Handle.ResetToDefaultOverride = FResetToDefaultOverride::Create(
+			FIsResetToDefaultVisible::CreateUObject(ValueDynamic, &UDMMaterialValueDynamic::CanResetToDefault),
+			FResetToDefaultHandler::CreateUObject(ValueDynamic, &UDMMaterialValueDynamic::ResetToDefault)
+		);
+
+		Handle.bEnabled = true;
+
+		InOutPropertyRows.Add(Handle);
+	}
+
+	const TArray<FName>& Properties = ParentValue->GetEditableProperties();
+	const int32 StartRow = InOutPropertyRows.Num();
+
+	for (const FName& Property : Properties)
+	{
+		if (Property == UDMMaterialValue::ValueName)
+		{
+			continue;
+		}
+			
+		if (InComponent->IsPropertyVisible(Property))
+		{
+			AddPropertyEditRows(InComponentEditorWidget, InComponent, Property, InOutPropertyRows, InOutProcessedObjects);
+		}
+	}
+
+	const int32 EndRow = InOutPropertyRows.Num();
+
+	for (int32 RowIndex = StartRow; RowIndex < EndRow; ++RowIndex)
+	{
+		InOutPropertyRows[RowIndex].bEnabled = false;
+	}
+}
diff --git a/Engine/Plugins/Experimental/DynamicMaterial/Source/DynamicMaterialEditor/Private/UI/PropertyGenerators/DMMaterialValueDynamicPropertyRowGenerator.h b/Engine/Plugins/Experimental/DynamicMaterial/Source/DynamicMaterialEditor/Private/UI/PropertyGenerators/DMMaterialValueDynamicPropertyRowGenerator.h
new file mode 100644
--- /dev/null
+++ b/Engine/Plugins/Experimental/DynamicMaterial/Source/DynamicMaterialEditor/Private/UI/PropertyGenerators/DMMaterialValueDynamicPropertyRowGenerator.h
@@ -0,0 +1,18 @@
+// Copyright Epic Games, Inc. All Rights Reserved.
+
+#pragma once
+
+#include "UI/PropertyGenerators/DMComponentPropertyRowGenerator.h"
+
+class FDMMaterialValueDynamicPropertyRowGenerator : public FDMComponentPropertyRowGenerator
+{
+public:
+	static const TSharedRef<FDMMaterialValueDynamicPropertyRowGenerator>& Get();
+
+	virtual ~FDMMaterialValueDynamicPropertyRowGenerator() override = default;
+
+	//~ Begin FDMComponentPropertyRowGenerator
+	virtual void AddComponentProperties(const TSharedRef<SDMMaterialComponentEditor>& InComponentEditorWidget, UDMMaterialComponent* InComponent,
+		TArray<FDMPropertyHandle>& InOutPropertyRows, TSet<UDMMaterialComponent*>& InOutProcessedObjects) override;
+	//~ End FDMComponentPropertyRowGenerator
+};
diff --git a/Engine/Plugins/Experimental/DynamicMaterial/Source/DynamicMaterialEditor/Private/UI/PropertyGenerators/DMMaterialValuePropertyRowGenerator.cpp b/Engine/Plugins/Experimental/DynamicMaterial/Source/DynamicMaterialEditor/Private/UI/PropertyGenerators/DMMaterialValuePropertyRowGenerator.cpp
new file mode 100644
--- /dev/null
+++ b/Engine/Plugins/Experimental/DynamicMaterial/Source/DynamicMaterialEditor/Private/UI/PropertyGenerators/DMMaterialValuePropertyRowGenerator.cpp
@@ -0,0 +1,94 @@
+// Copyright Epic Games, Inc. All Rights Reserved.
+
+#include "UI/PropertyGenerators/DMMaterialValuePropertyRowGenerator.h"
+
+#include "Components/DMMaterialComponent.h"
+#include "Components/DMMaterialValue.h"
+#include "Components/DMMaterialValueDynamic.h"
+#include "DynamicMaterialEditorModule.h"
+#include "IDetailPropertyRow.h"
+#include "Model/DynamicMaterialModelBase.h"
+#include "Model/DynamicMaterialModelDynamic.h"
+#include "UI/Utils/DMWidgetStatics.h"
+#include "UI/Widgets/Editor/SDMMaterialComponentEditor.h"
+#include "UI/Widgets/SDMMaterialEditor.h"
+
+const TSharedRef<FDMMaterialValuePropertyRowGenerator>& FDMMaterialValuePropertyRowGenerator::Get()
+{
+	static TSharedRef<FDMMaterialValuePropertyRowGenerator> Generator = MakeShared<FDMMaterialValuePropertyRowGenerator>();
+	return Generator;
+}
+
+void FDMMaterialValuePropertyRowGenerator::AddComponentProperties(const TSharedRef<SDMMaterialComponentEditor>& InComponentEditorWidget, UDMMaterialComponent* InComponent,
+	TArray<FDMPropertyHandle>& InOutPropertyRows, TSet<UDMMaterialComponent*>& InOutProcessedObjects)
+{
+	if (!IsValid(InComponent))
+	{
+		return;
+	}
+
+	if (InOutProcessedObjects.Contains(InComponent))
+	{
+		return;
+	}
+
+	UDMMaterialValue* Value = Cast<UDMMaterialValue>(InComponent);
+
+	if (!Value)
+	{
+		return;
+	}
+
+	// The base material value class is abstract and not allowed.
+	if (Value->GetClass() == UDMMaterialValue::StaticClass())
+	{
+		return;
+	}
+
+	InOutProcessedObjects.Add(InComponent);
+
+	if (TSharedPtr<SDMMaterialEditor> EditorWidget = InComponentEditorWidget->GetEditorWidget())
+	{
+		if (UDynamicMaterialModelBase* MaterialModelBase = EditorWidget->GetMaterialModelBase())
+		{
+			if (UDynamicMaterialModelDynamic* MaterialModelDynamic = Cast<UDynamicMaterialModelDynamic>(MaterialModelBase))
+			{
+				if (UDMMaterialComponentDynamic* ComponentDynamic = MaterialModelDynamic->GetComponentDynamic(Value->GetFName()))
+				{
+					FDynamicMaterialEditorModule::Get().GeneratorComponentPropertyRows(InComponentEditorWidget, ComponentDynamic, InOutPropertyRows, InOutProcessedObjects);
+				}
+
+				return;
+			}
+		}
+	}
+
+	if (Value->AllowEditValue())
+	{
+		FDMPropertyHandle Handle = FDMWidgetStatics::Get().GetPropertyHandle(&*InComponentEditorWidget, Value, UDMMaterialValue::ValueName);
+
+		Handle.ResetToDefaultOverride = FResetToDefaultOverride::Create(
+			FIsResetToDefaultVisible::CreateUObject(Value, &UDMMaterialValue::CanResetToDefault),
+			FResetToDefaultHandler::CreateUObject(Value, &UDMMaterialValue::ResetToDefault)
+		);
+
+		Handle.bEnabled = true;
+
+		InOutPropertyRows.Add(Handle);
+	}
+
+	const TArray<FName>& Properties = Value->GetEditableProperties();
+
+	for (const FName& Property : Properties)
+	{
+		if (Property == UDMMaterialValue::ValueName)
+		{
+			continue;
+		}
+			
+		if (InComponent->IsPropertyVisible(Property))
+		{
+			AddPropertyEditRows(InComponentEditorWidget, InComponent, Property, InOutPropertyRows, InOutProcessedObjects);
+		}
+	}
+}
diff --git a/Engine/Plugins/Experimental/DynamicMaterial/Source/DynamicMaterialEditor/Private/UI/PropertyGenerators/DMMaterialValuePropertyRowGenerator.h b/Engine/Plugins/Experimental/DynamicMaterial/Source/DynamicMaterialEditor/Private/UI/PropertyGenerators/DMMaterialValuePropertyRowGenerator.h
new file mode 100644
--- /dev/null
+++ b/Engine/Plugins/Experimental/DynamicMaterial/Source/DynamicMaterialEditor/Private/UI/PropertyGenerators/DMMaterialValuePropertyRowGenerator.h
@@ -0,0 +1,18 @@
+// Copyright Epic Games, Inc. All Rights Reserved.
+
+#pragma once
+
+#include "UI/PropertyGenerators/DMComponentPropertyRowGenerator.h"
+
+class FDMMaterialValuePropertyRowGenerator : public FDMComponentPropertyRowGenerator
+{
+public:
+	static const TSharedRef<FDMMaterialValuePropertyRowGenerator>& Get();
+
+	virtual ~FDMMaterialValuePropertyRowGenerator() override = default;
+
+	//~ Begin FDMComponentPropertyRowGenerator
+	virtual void AddComponentProperties(const TSharedRef<SDMMaterialComponentEditor>& InComponentEditorWidget, UDMMaterialComponent* InComponent,
+		TArray<FDMPropertyHandle>& InOutPropertyRows, TSet<UDMMaterialComponent*>& InOutProcessedObjects) override;
+	//~ End FDMComponentPropertyRowGenerator
+};
diff --git a/Engine/Plugins/Experimental/DynamicMaterial/Source/DynamicMaterialEditor/Private/UI/PropertyGenerators/DMTextureUVDynamicPropertyRowGenerator.cpp b/Engine/Plugins/Experimental/DynamicMaterial/Source/DynamicMaterialEditor/Private/UI/PropertyGenerators/DMTextureUVDynamicPropertyRowGenerator.cpp
new file mode 100644
--- /dev/null
+++ b/Engine/Plugins/Experimental/DynamicMaterial/Source/DynamicMaterialEditor/Private/UI/PropertyGenerators/DMTextureUVDynamicPropertyRowGenerator.cpp
@@ -0,0 +1,247 @@
+// Copyright Epic Games, Inc. All Rights Reserved.
+
+#include "UI/PropertyGenerators/DMTextureUVDynamicPropertyRowGenerator.h"
+
+#include "Components/DMMaterialComponent.h"
+#include "Components/DMMaterialStage.h"
+#include "Components/DMTextureUV.h"
+#include "Components/DMTextureUVDynamic.h"
+#include "DMEDefs.h"
+#include "Modules/ModuleManager.h"
+#include "PropertyCustomizationHelpers.h"
+#include "PropertyEditorModule.h"
+#include "UI/Utils/DMWidgetStatics.h"
+#include "UI/Widgets/Editor/SDMMaterialComponentEditor.h"
+#include "UI/Widgets/SDMMaterialEditor.h"
+#include "UI/Widgets/Visualizers/SDMTextureUVVisualizerProperty.h"
+#include "Widgets/Layout/SBox.h"
+
+#define LOCTEXT_NAMESPACE "DMTextureUVDynamicPropertyRowGenerator"
+
+const TSharedRef<FDMTextureUVDynamicPropertyRowGenerator>& FDMTextureUVDynamicPropertyRowGenerator::Get()
+{
+	static TSharedRef<FDMTextureUVDynamicPropertyRowGenerator> Generator = MakeShared<FDMTextureUVDynamicPropertyRowGenerator>();
+	return Generator;
+}
+
+namespace UE::DynamicMaterialEditor::Private
+{
+	void AddTextureUVDynamicPropertyRow(const TSharedRef<SWidget>& InComponentEditorWidget, UDMMaterialComponent* InTextureUVDynamic, 
+		FName InProperty, bool bInEnabled, TArray<FDMPropertyHandle>& InOutPropertyRows);
+
+	void AddTextureUVDynamicVisualizerRow(const TSharedRef<SWidget>& InComponentEditorWidget, UDMTextureUVDynamic* InTextureUVDynamic, 
+		TArray<FDMPropertyHandle>& InOutPropertyRows);
+
+	bool CanResetTextureUVDynamicDynamicPropertyToDefault(TSharedPtr<IPropertyHandle> InPropertyHandle);
+
+	void ResetTextureUVDynamicDynamicPropertyToDefault(TSharedPtr<IPropertyHandle> InPropertyHandle);
+}
+
+void FDMTextureUVDynamicPropertyRowGenerator::AddComponentProperties(const TSharedRef<SDMMaterialComponentEditor>& InComponentEditorWidget, 
+	UDMMaterialComponent* InComponent, TArray<FDMPropertyHandle>& InOutPropertyRows, TSet<UDMMaterialComponent*>& InOutProcessedObjects)
+{
+	if (!IsValid(InComponent))
+	{
+		return;
+	}
+
+	if (InOutProcessedObjects.Contains(InComponent))
+	{
+		return;
+	}
+
+	UDMTextureUVDynamic* TextureUVDynamic = Cast<UDMTextureUVDynamic>(InComponent);
+
+	if (!TextureUVDynamic)
+	{
+		return;
+	}
+
+	UDMTextureUV* TextureUV = TextureUVDynamic->GetParentTextureUV();
+
+	if (!TextureUV)
+	{
+		return;
+	}
+
+	InOutProcessedObjects.Add(InComponent);
+
+	using namespace UE::DynamicMaterialEditor::Private;
+
+	AddTextureUVDynamicPropertyRow(InComponentEditorWidget, TextureUVDynamic, UDMTextureUV::NAME_Offset, /* Enabled */ true, InOutPropertyRows);
+	AddTextureUVDynamicPropertyRow(InComponentEditorWidget, TextureUVDynamic, UDMTextureUV::NAME_Rotation, /* Enabled */ true, InOutPropertyRows);
+	AddTextureUVDynamicPropertyRow(InComponentEditorWidget, TextureUVDynamic, UDMTextureUV::NAME_Tiling, /* Enabled */ true, InOutPropertyRows);
+	AddTextureUVDynamicPropertyRow(InComponentEditorWidget, TextureUVDynamic, UDMTextureUV::NAME_Pivot, /* Enabled */ true, InOutPropertyRows);
+	AddTextureUVDynamicPropertyRow(InComponentEditorWidget, TextureUV, UDMTextureUV::NAME_bMirrorOnX, /* Enabled */ false, InOutPropertyRows);
+	AddTextureUVDynamicPropertyRow(InComponentEditorWidget, TextureUV, UDMTextureUV::NAME_bMirrorOnY, /* Enabled */ false, InOutPropertyRows);
+	AddTextureUVDynamicVisualizerRow(InComponentEditorWidget, TextureUVDynamic, InOutPropertyRows);
+}
+
+void FDMTextureUVDynamicPropertyRowGenerator::AddPopoutComponentProperties(const TSharedRef<SWidget>& InParentWidget, UDMMaterialComponent* InComponent, 
+	TArray<FDMPropertyHandle>& InOutPropertyRows)
+{
+	if (!IsValid(InComponent))
+	{
+		return;
+	}
+
+	UDMTextureUVDynamic* TextureUVDynamic = Cast<UDMTextureUVDynamic>(InComponent);
+
+	if (!TextureUVDynamic)
+	{
+		return;
+	}
+
+	UDMTextureUV* TextureUV = TextureUVDynamic->GetParentTextureUV();
+
+	if (!TextureUV)
+	{
+		return;
+	}
+
+	using namespace UE::DynamicMaterialEditor::Private;
+
+	AddTextureUVDynamicPropertyRow(InParentWidget, TextureUVDynamic, UDMTextureUV::NAME_Offset, /* Enabled */ true, InOutPropertyRows);
+	AddTextureUVDynamicPropertyRow(InParentWidget, TextureUVDynamic, UDMTextureUV::NAME_Rotation, /* Enabled */ true, InOutPropertyRows);
+	AddTextureUVDynamicPropertyRow(InParentWidget, TextureUVDynamic, UDMTextureUV::NAME_Tiling, /* Enabled */ true, InOutPropertyRows);
+	AddTextureUVDynamicPropertyRow(InParentWidget, TextureUVDynamic, UDMTextureUV::NAME_Pivot, /* Enabled */ true, InOutPropertyRows);
+	AddTextureUVDynamicPropertyRow(InParentWidget, TextureUV, UDMTextureUV::NAME_bMirrorOnX, /* Enabled */ false, InOutPropertyRows);
+	AddTextureUVDynamicPropertyRow(InParentWidget, TextureUV, UDMTextureUV::NAME_bMirrorOnY, /* Enabled */ false, InOutPropertyRows);
+}
+
+bool FDMTextureUVDynamicPropertyRowGenerator::AllowKeyframeButton(UDMMaterialComponent* InComponent, FProperty* InProperty)
+{
+	if (InProperty)
+	{
+		const bool* AddKeyframeButtonPtr = UDMTextureUV::TextureProperties.Find(InProperty->GetFName());
+
+		if (AddKeyframeButtonPtr)
+		{
+			return *AddKeyframeButtonPtr;
+		}
+	}
+
+	return FDMComponentPropertyRowGenerator::AllowKeyframeButton(InComponent, InProperty);
+}
+
+void UE::DynamicMaterialEditor::Private::AddTextureUVDynamicPropertyRow(const TSharedRef<SWidget>& InComponentEditorWidget, 
+	UDMMaterialComponent* InTextureUVDynamic, FName InProperty, bool bInEnabled, TArray<FDMPropertyHandle>& InOutPropertyRows)
+{
+	FDMPropertyHandle& NewHandle = InOutPropertyRows.Add_GetRef(FDMWidgetStatics::Get().GetPropertyHandle(&*InComponentEditorWidget, InTextureUVDynamic, InProperty));
+
+	NewHandle.ResetToDefaultOverride = FResetToDefaultOverride::Create(
+		FIsResetToDefaultVisible::CreateStatic(&UE::DynamicMaterialEditor::Private::CanResetTextureUVDynamicDynamicPropertyToDefault),
+		FResetToDefaultHandler::CreateStatic(&UE::DynamicMaterialEditor::Private::ResetTextureUVDynamicDynamicPropertyToDefault)
+	);
+
+	NewHandle.bEnabled = bInEnabled;
+}
+
+void UE::DynamicMaterialEditor::Private::AddTextureUVDynamicVisualizerRow(const TSharedRef<SWidget>& InComponentEditorWidget, UDMTextureUVDynamic* InTextureUVDynamic,
+	TArray<FDMPropertyHandle>& InOutPropertyRows)
+{
+	UDMTextureUV* TextureUV = InTextureUVDynamic->GetParentTextureUV();
+
+	if (!TextureUV)
+	{
+		return;
+	}
+
+	// Make sure we don't get a substage
+	UDMMaterialStage* Stage = TextureUV->GetTypedParent<UDMMaterialStage>(/* Allow Subclasses */ false);
+
+	if (!Stage)
+	{
+		return;
+	}
+
+	if (InComponentEditorWidget->GetWidgetClass().GetWidgetType() != SDMMaterialComponentEditor::StaticWidgetClass().GetWidgetType())
+	{
+		return;
+	}
+
+	TSharedPtr<SDMMaterialEditor> EditorWidget = StaticCastSharedRef<SDMMaterialComponentEditor>(InComponentEditorWidget)->GetEditorWidget();
+
+	if (!EditorWidget.IsValid())
+	{
+		return;
+	}
+
+	FDMPropertyHandle VisualizerHandle;
+	VisualizerHandle.NameOverride = LOCTEXT("Visualizer", "UV Visualizer");
+	VisualizerHandle.NameToolTipOverride = LOCTEXT("VisualizerToolTip", "A graphical Texture UV editor.\n\n- Offset Mode: Change the Texture UV offset.\n- Pivot Mode: Change the Texture UV pivot, rotation and tiling.\n\nControl+click to reset values to default.");
+	VisualizerHandle.ValueName = FName(*InTextureUVDynamic->GetComponentPath());
+	VisualizerHandle.ValueWidget = SNew(SDMTextureUVVisualizerProperty, EditorWidget.ToSharedRef(), Stage).TextureUVDynamic(InTextureUVDynamic);
+	VisualizerHandle.CategoryOverrideName = TEXT("Texture UV");
+	VisualizerHandle.bEnabled = true;
+	InOutPropertyRows.Add(VisualizerHandle);
+}
+
+
+bool UE::DynamicMaterialEditor::Private::CanResetTextureUVDynamicDynamicPropertyToDefault(TSharedPtr<IPropertyHandle> InPropertyHandle)
+{
+	FProperty* Property = InPropertyHandle->GetProperty();
+
+	if (!Property)
+	{
+		return false;
+	}
+
+	const FName PropertyName = Property->GetFName();
+
+	if (PropertyName.IsNone())
+	{
+		return false;
+	}
+
+	TArray<UObject*> Outers;
+	InPropertyHandle->GetOuterObjects(Outers);
+
+	if (Outers.IsEmpty())
+	{
+		return false;
+	}
+
+	const UDMTextureUVDynamic* PropertyObject = Cast<UDMTextureUVDynamic>(Outers[0]);
+
+	if (!PropertyObject)
+	{
+		return false;
+	}
+
+	const UDMTextureUVDynamic* DefaultObject = GetDefault<UDMTextureUVDynamic>();
+
+	if (!DefaultObject)
+	{
+		return false;
+	}
+
+	if (PropertyName == UDMTextureUV::NAME_Offset)
+	{
+		return !DefaultObject->GetOffset().Equals(PropertyObject->GetOffset());
+	}
+
+	if (PropertyName == UDMTextureUV::NAME_Pivot)
+	{
+		return !DefaultObject->GetPivot().Equals(PropertyObject->GetPivot());
+	}
+
+	if (PropertyName == NAME_Rotation)
+	{
+		return DefaultObject->GetRotation() != PropertyObject->GetRotation();
+	}
+
+	if (PropertyName == UDMTextureUV::NAME_Tiling)
+	{
+		return !DefaultObject->GetTiling().Equals(PropertyObject->GetTiling());
+	}
+
+	return false;
+}
+
+void UE::DynamicMaterialEditor::Private::ResetTextureUVDynamicDynamicPropertyToDefault(TSharedPtr<IPropertyHandle> InPropertyHandle)
+{
+	InPropertyHandle->ResetToDefault();
+}
+
+#undef LOCTEXT_NAMESPACE
diff --git a/Engine/Plugins/Experimental/DynamicMaterial/Source/DynamicMaterialEditor/Private/UI/PropertyGenerators/DMTextureUVDynamicPropertyRowGenerator.h b/Engine/Plugins/Experimental/DynamicMaterial/Source/DynamicMaterialEditor/Private/UI/PropertyGenerators/DMTextureUVDynamicPropertyRowGenerator.h
new file mode 100644
--- /dev/null
+++ b/Engine/Plugins/Experimental/DynamicMaterial/Source/DynamicMaterialEditor/Private/UI/PropertyGenerators/DMTextureUVDynamicPropertyRowGenerator.h
@@ -0,0 +1,25 @@
+// Copyright Epic Games, Inc. All Rights Reserved.
+
+#pragma once
+
+#include "UI/PropertyGenerators/DMComponentPropertyRowGenerator.h"
+
+class SWidget;
+
+class FDMTextureUVDynamicPropertyRowGenerator : public FDMComponentPropertyRowGenerator
+{
+public:
+	static const TSharedRef<FDMTextureUVDynamicPropertyRowGenerator>& Get();
+
+	virtual ~FDMTextureUVDynamicPropertyRowGenerator() override = default;
+
+	static void AddPopoutComponentProperties(const TSharedRef<SWidget>& InParentWidget, UDMMaterialComponent* InComponent,
+		TArray<FDMPropertyHandle>& InOutPropertyRows);
+
+	//~ Begin FDMComponentPropertyRowGenerator
+	virtual void AddComponentProperties(const TSharedRef<SDMMaterialComponentEditor>& InComponentEditorWidget, UDMMaterialComponent* InComponent,
+		TArray<FDMPropertyHandle>& InOutPropertyRows, TSet<UDMMaterialComponent*>& InOutProcessedObjects) override;
+
+	virtual bool AllowKeyframeButton(UDMMaterialComponent* InComponent, FProperty* InProperty) override;
+	//~ End FDMComponentPropertyRowGenerator
+};
diff --git a/Engine/Plugins/Experimental/DynamicMaterial/Source/DynamicMaterialEditor/Private/UI/PropertyGenerators/DMTextureUVPropertyRowGenerator.cpp b/Engine/Plugins/Experimental/DynamicMaterial/Source/DynamicMaterialEditor/Private/UI/PropertyGenerators/DMTextureUVPropertyRowGenerator.cpp
new file mode 100644
--- /dev/null
+++ b/Engine/Plugins/Experimental/DynamicMaterial/Source/DynamicMaterialEditor/Private/UI/PropertyGenerators/DMTextureUVPropertyRowGenerator.cpp
@@ -0,0 +1,258 @@
+// Copyright Epic Games, Inc. All Rights Reserved.
+
+#include "UI/PropertyGenerators/DMTextureUVPropertyRowGenerator.h"
+
+#include "Components/DMMaterialComponent.h"
+#include "Components/DMMaterialStage.h"
+#include "Components/DMTextureUV.h"
+#include "DMEDefs.h"
+#include "DynamicMaterialEditorModule.h"
+#include "Model/DynamicMaterialModelDynamic.h"
+#include "Modules/ModuleManager.h"
+#include "PropertyCustomizationHelpers.h"
+#include "PropertyEditorModule.h"
+#include "UI/Utils/DMWidgetStatics.h"
+#include "UI/Widgets/Editor/SDMMaterialComponentEditor.h"
+#include "UI/Widgets/SDMMaterialEditor.h"
+#include "UI/Widgets/Visualizers/SDMTextureUVVisualizerProperty.h"
+#include "Widgets/Layout/SBox.h"
+
+#define LOCTEXT_NAMESPACE "DMTextureUVPropertyRowGenerator"
+
+const TSharedRef<FDMTextureUVPropertyRowGenerator>& FDMTextureUVPropertyRowGenerator::Get()
+{
+	static TSharedRef<FDMTextureUVPropertyRowGenerator> Generator = MakeShared<FDMTextureUVPropertyRowGenerator>();
+	return Generator;
+}
+
+namespace UE::DynamicMaterialEditor::Private
+{
+	void AddTextureUVPropertyRow(const TSharedRef<SWidget>& InComponentEditorWidget, UDMTextureUV* InTextureUV, FName InProperty, 
+		TArray<FDMPropertyHandle>& InOutPropertyRows);
+
+	void AddTextureUVVisualizerRow(const TSharedRef<SWidget>& InComponentEditorWidget, UDMTextureUV* InTextureUV, 
+		TArray<FDMPropertyHandle>& InOutPropertyRows);
+
+	bool CanResetTextureUVPropertyToDefault(TSharedPtr<IPropertyHandle> InPropertyHandle);
+
+	void ResetTextureUVPropertyToDefault(TSharedPtr<IPropertyHandle> InPropertyHandle);
+}
+
+void FDMTextureUVPropertyRowGenerator::AddComponentProperties(const TSharedRef<SDMMaterialComponentEditor>& InComponentEditorWidget, 
+	UDMMaterialComponent* InComponent, TArray<FDMPropertyHandle>& InOutPropertyRows, TSet<UDMMaterialComponent*>& InOutProcessedObjects)
+{
+	if (!IsValid(InComponent))
+	{
+		return;
+	}
+
+	if (InOutProcessedObjects.Contains(InComponent))
+	{
+		return;
+	}
+
+	UDMTextureUV* TextureUV = Cast<UDMTextureUV>(InComponent);
+
+	if (!TextureUV)
+	{
+		return;
+	}
+
+	InOutProcessedObjects.Add(InComponent);
+
+	if (TSharedPtr<SDMMaterialEditor> EditorWidget = InComponentEditorWidget->GetEditorWidget())
+	{
+		if (UDynamicMaterialModelBase* MaterialModelBase = EditorWidget->GetMaterialModelBase())
+		{
+			if (UDynamicMaterialModelDynamic* MaterialModelDynamic = Cast<UDynamicMaterialModelDynamic>(MaterialModelBase))
+			{
+				if (UDMMaterialComponentDynamic* ComponentDynamic = MaterialModelDynamic->GetComponentDynamic(TextureUV->GetFName()))
+				{
+					FDynamicMaterialEditorModule::Get().GeneratorComponentPropertyRows(InComponentEditorWidget, ComponentDynamic, InOutPropertyRows, InOutProcessedObjects);
+				}
+
+				return;
+			}
+		}
+	}
+
+	using namespace UE::DynamicMaterialEditor::Private;
+
+	AddTextureUVPropertyRow(InComponentEditorWidget, TextureUV, UDMTextureUV::NAME_Offset, InOutPropertyRows);
+	AddTextureUVPropertyRow(InComponentEditorWidget, TextureUV, UDMTextureUV::NAME_Rotation, InOutPropertyRows);
+	AddTextureUVPropertyRow(InComponentEditorWidget, TextureUV, UDMTextureUV::NAME_Tiling, InOutPropertyRows);
+	AddTextureUVPropertyRow(InComponentEditorWidget, TextureUV, UDMTextureUV::NAME_Pivot, InOutPropertyRows);
+	AddTextureUVPropertyRow(InComponentEditorWidget, TextureUV, UDMTextureUV::NAME_bMirrorOnX, InOutPropertyRows);
+	AddTextureUVPropertyRow(InComponentEditorWidget, TextureUV, UDMTextureUV::NAME_bMirrorOnY, InOutPropertyRows);
+	AddTextureUVVisualizerRow(InComponentEditorWidget, TextureUV, InOutPropertyRows);
+}
+
+void FDMTextureUVPropertyRowGenerator::AddPopoutComponentProperties(const TSharedRef<SWidget>& InParentWidget, UDMMaterialComponent* InComponent, 
+	TArray<FDMPropertyHandle>& InOutPropertyRows)
+{
+	if (!IsValid(InComponent))
+	{
+		return;
+	}
+
+	UDMTextureUV* TextureUV = Cast<UDMTextureUV>(InComponent);
+
+	if (!TextureUV)
+	{
+		return;
+	}
+
+	using namespace UE::DynamicMaterialEditor::Private;
+
+	AddTextureUVPropertyRow(InParentWidget, TextureUV, UDMTextureUV::NAME_Offset, InOutPropertyRows);
+	AddTextureUVPropertyRow(InParentWidget, TextureUV, UDMTextureUV::NAME_Rotation, InOutPropertyRows);
+	AddTextureUVPropertyRow(InParentWidget, TextureUV, UDMTextureUV::NAME_Tiling, InOutPropertyRows);
+	AddTextureUVPropertyRow(InParentWidget, TextureUV, UDMTextureUV::NAME_Pivot, InOutPropertyRows);
+	AddTextureUVPropertyRow(InParentWidget, TextureUV, UDMTextureUV::NAME_bMirrorOnX, InOutPropertyRows);
+	AddTextureUVPropertyRow(InParentWidget, TextureUV, UDMTextureUV::NAME_bMirrorOnY, InOutPropertyRows);
+}
+
+bool FDMTextureUVPropertyRowGenerator::AllowKeyframeButton(UDMMaterialComponent* InComponent, FProperty* InProperty)
+{
+	if (InProperty)
+	{
+		const bool* AddKeyframeButtonPtr = UDMTextureUV::TextureProperties.Find(InProperty->GetFName());
+
+		if (AddKeyframeButtonPtr)
+		{
+			return *AddKeyframeButtonPtr;
+		}
+	}
+
+	return FDMComponentPropertyRowGenerator::AllowKeyframeButton(InComponent, InProperty);
+}
+
+void UE::DynamicMaterialEditor::Private::AddTextureUVPropertyRow(const TSharedRef<SWidget>& InComponentEditorWidget, UDMTextureUV* InTextureUV, 
+	FName InProperty, TArray<FDMPropertyHandle>& InOutPropertyRows)
+{
+	FDMPropertyHandle& NewHandle = InOutPropertyRows.Add_GetRef(FDMWidgetStatics::Get().GetPropertyHandle(&*InComponentEditorWidget, InTextureUV, InProperty));
+
+	NewHandle.ResetToDefaultOverride = FResetToDefaultOverride::Create(
+		FIsResetToDefaultVisible::CreateStatic(&UE::DynamicMaterialEditor::Private::CanResetTextureUVPropertyToDefault),
+		FResetToDefaultHandler::CreateStatic(&UE::DynamicMaterialEditor::Private::ResetTextureUVPropertyToDefault)
+	);
+
+	NewHandle.bEnabled = true;
+}
+
+void UE::DynamicMaterialEditor::Private::AddTextureUVVisualizerRow(const TSharedRef<SWidget>& InComponentEditorWidget, UDMTextureUV* InTextureUV, 
+	TArray<FDMPropertyHandle>& InOutPropertyRows)
+{
+	// Make sure we don't get a substage
+	UDMMaterialStage* Stage = InTextureUV->GetTypedParent<UDMMaterialStage>(/* Allow Subclasses */ false);
+
+	if (!Stage)
+	{
+		return;
+	}
+
+	if (InComponentEditorWidget->GetWidgetClass().GetWidgetType() != SDMMaterialComponentEditor::StaticWidgetClass().GetWidgetType())
+	{
+		return;
+	}
+
+	TSharedPtr<SDMMaterialEditor> EditorWidget = StaticCastSharedRef<SDMMaterialComponentEditor>(InComponentEditorWidget)->GetEditorWidget();
+
+	if (!EditorWidget.IsValid())
+	{
+		return;
+	}
+
+	FDMPropertyHandle VisualizerHandle;
+	VisualizerHandle.NameOverride = LOCTEXT("Visualizer", "UV Visualizer");
+	VisualizerHandle.NameToolTipOverride = LOCTEXT("VisualizerToolTip", "A graphical Texture UV editor.\n\n- Offset Mode: Change the Texture UV offset.\n- Pivot Mode: Change the Texture UV pivot, rotation and tiling.\n\nControl+click to reset values to default.");
+	VisualizerHandle.ValueName = FName(*InTextureUV->GetComponentPath());
+	VisualizerHandle.ValueWidget = SNew(SDMTextureUVVisualizerProperty, EditorWidget.ToSharedRef(), Stage).TextureUV(InTextureUV);
+	VisualizerHandle.CategoryOverrideName = TEXT("Texture UV");
+	VisualizerHandle.bEnabled = true;
+	InOutPropertyRows.Add(VisualizerHandle);
+}
+
+
+bool UE::DynamicMaterialEditor::Private::CanResetTextureUVPropertyToDefault(TSharedPtr<IPropertyHandle> InPropertyHandle)
+{
+	FProperty* Property = InPropertyHandle->GetProperty();
+
+	if (!Property)
+	{
+		return false;
+	}
+
+	const FName PropertyName = Property->GetFName();
+
+	if (PropertyName.IsNone())
+	{
+		return false;
+	}
+
+	TArray<UObject*> Outers;
+	InPropertyHandle->GetOuterObjects(Outers);
+
+	if (Outers.IsEmpty())
+	{
+		return false;
+	}
+
+	const UDMTextureUV* PropertyObject = Cast<UDMTextureUV>(Outers[0]);
+
+	if (!PropertyObject)
+	{
+		return false;
+	}
+
+	const UDMTextureUV* DefaultObject = GetDefault<UDMTextureUV>();
+
+	if (!DefaultObject)
+	{
+		return false;
+	}
+
+	if (PropertyName == UDMTextureUV::NAME_UVSource)
+	{
+		return DefaultObject->GetUVSource() != PropertyObject->GetUVSource();
+	}
+
+	if (PropertyName == UDMTextureUV::NAME_bMirrorOnX)
+	{
+		return DefaultObject->GetMirrorOnX() != PropertyObject->GetMirrorOnX();
+	}
+
+	if (PropertyName == UDMTextureUV::NAME_bMirrorOnY)
+	{
+		return DefaultObject->GetMirrorOnY() != PropertyObject->GetMirrorOnY();
+	}
+
+	if (PropertyName == UDMTextureUV::NAME_Offset)
+	{
+		return !DefaultObject->GetOffset().Equals(PropertyObject->GetOffset());
+	}
+
+	if (PropertyName == UDMTextureUV::NAME_Pivot)
+	{
+		return !DefaultObject->GetPivot().Equals(PropertyObject->GetPivot());
+	}
+
+	if (PropertyName == NAME_Rotation)
+	{
+		return DefaultObject->GetRotation() != PropertyObject->GetRotation();
+	}
+
+	if (PropertyName == UDMTextureUV::NAME_Tiling)
+	{
+		return !DefaultObject->GetTiling().Equals(PropertyObject->GetTiling());
+	}
+
+	return false;
+}
+
+void UE::DynamicMaterialEditor::Private::ResetTextureUVPropertyToDefault(TSharedPtr<IPropertyHandle> InPropertyHandle)
+{
+	InPropertyHandle->ResetToDefault();
+}
+
+#undef LOCTEXT_NAMESPACE
diff --git a/Engine/Plugins/Experimental/DynamicMaterial/Source/DynamicMaterialEditor/Private/UI/PropertyGenerators/DMTextureUVPropertyRowGenerator.h b/Engine/Plugins/Experimental/DynamicMaterial/Source/DynamicMaterialEditor/Private/UI/PropertyGenerators/DMTextureUVPropertyRowGenerator.h
new file mode 100644
--- /dev/null
+++ b/Engine/Plugins/Experimental/DynamicMaterial/Source/DynamicMaterialEditor/Private/UI/PropertyGenerators/DMTextureUVPropertyRowGenerator.h
@@ -0,0 +1,25 @@
+// Copyright Epic Games, Inc. All Rights Reserved.
+
+#pragma once
+
+#include "UI/PropertyGenerators/DMComponentPropertyRowGenerator.h"
+
+class SWidget;
+
+class FDMTextureUVPropertyRowGenerator : public FDMComponentPropertyRowGenerator
+{
+public:
+	static const TSharedRef<FDMTextureUVPropertyRowGenerator>& Get();
+
+	virtual ~FDMTextureUVPropertyRowGenerator() override = default;
+
+	static void AddPopoutComponentProperties(const TSharedRef<SWidget>& InParentWidget, UDMMaterialComponent* InComponent,
+		TArray<FDMPropertyHandle>& InOutPropertyRows);
+
+	//~ Begin FDMComponentPropertyRowGenerator
+	virtual void AddComponentProperties(const TSharedRef<SDMMaterialComponentEditor>& InComponentEditorWidget, UDMMaterialComponent* InComponent,
+		TArray<FDMPropertyHandle>& InOutPropertyRows, TSet<UDMMaterialComponent*>& InOutProcessedObjects) override;
+
+	virtual bool AllowKeyframeButton(UDMMaterialComponent* InComponent, FProperty* InProperty) override;
+	//~ End FDMComponentPropertyRowGenerator
+};
diff --git a/Engine/Plugins/Experimental/DynamicMaterial/Source/DynamicMaterialEditor/Private/UI/PropertyGenerators/DMThroughputPropertyRowGenerator.cpp b/Engine/Plugins/Experimental/DynamicMaterial/Source/DynamicMaterialEditor/Private/UI/PropertyGenerators/DMThroughputPropertyRowGenerator.cpp
new file mode 100644
--- /dev/null
+++ b/Engine/Plugins/Experimental/DynamicMaterial/Source/DynamicMaterialEditor/Private/UI/PropertyGenerators/DMThroughputPropertyRowGenerator.cpp
@@ -0,0 +1,123 @@
+// Copyright Epic Games, Inc. All Rights Reserved.
+
+#include "UI/PropertyGenerators/DMThroughputPropertyRowGenerator.h"
+
+#include "Components/DMMaterialStage.h"
+#include "Components/DMMaterialStageInput.h"
+#include "Components/DMMaterialStageThroughput.h"
+#include "Components/DMMaterialValue.h"
+#include "DynamicMaterialEditorModule.h"
+
+const TSharedRef<FDMThroughputPropertyRowGenerator>& FDMThroughputPropertyRowGenerator::Get()
+{
+	static TSharedRef<FDMThroughputPropertyRowGenerator> Generator = MakeShared<FDMThroughputPropertyRowGenerator>();
+	return Generator;
+}
+
+void FDMThroughputPropertyRowGenerator::AddComponentProperties(const TSharedRef<SDMMaterialComponentEditor>& InComponentEditorWidget,
+	UDMMaterialComponent* InComponent, TArray<FDMPropertyHandle>& InOutPropertyRows,
+	TSet<UDMMaterialComponent*>& InOutProcessedObjects)
+{
+	if (!IsValid(InComponent))
+	{
+		return;
+	}
+
+	if (InOutProcessedObjects.Contains(InComponent))
+	{
+		return;
+	}
+
+	UDMMaterialStageThroughput* Throughput = Cast<UDMMaterialStageThroughput>(InComponent);
+
+	if (!Throughput)
+	{
+		return;
+	}
+
+	InOutProcessedObjects.Add(Throughput);
+
+	const TArray<FName>& ThroughputProperties = Throughput->GetEditableProperties();
+
+	for (const FName& ThroughputProperty : ThroughputProperties)
+	{
+		if (InComponent->IsPropertyVisible(ThroughputProperty))
+		{
+			AddPropertyEditRows(InComponentEditorWidget, InComponent, ThroughputProperty, InOutPropertyRows, InOutProcessedObjects);
+		}
+	}
+
+	UDMMaterialStage* Stage = Throughput->GetStage();
+
+	if (!Stage)
+	{
+		return;
+	}
+
+	InOutProcessedObjects.Add(Stage);
+
+	static const FName StageInputsName = GET_MEMBER_NAME_CHECKED(UDMMaterialStage, Inputs);
+
+	{
+		const TArray<FDMMaterialStageConnector>& InputConnectors = Throughput->GetInputConnectors();
+		const TArray<FDMMaterialStageConnection>& InputMap = Stage->GetInputConnectionMap();
+		TArray<UDMMaterialStageInput*> Inputs = Stage->GetInputs();
+
+		for (int32 InputIdx = 0; InputIdx < InputConnectors.Num(); ++InputIdx)
+		{
+			if (!Throughput->IsInputVisible(InputIdx) || !Throughput->CanChangeInput(InputIdx))
+			{
+				continue;
+			}
+
+			const int32 StartRow = InOutPropertyRows.Num();
+
+			for (const FDMMaterialStageConnectorChannel& Channel : InputMap[InputIdx].Channels)
+			{
+				const int32 StageInputIdx = Channel.SourceIndex - FDMMaterialStageConnectorChannel::FIRST_STAGE_INPUT;
+
+				if (Inputs.IsValidIndex(StageInputIdx))
+				{
+					FDynamicMaterialEditorModule::GeneratorComponentPropertyRows(InComponentEditorWidget, Inputs[StageInputIdx], InOutPropertyRows, InOutProcessedObjects);
+				}
+			}
+
+			for (int32 PropertyRowIdx = StartRow; PropertyRowIdx < InOutPropertyRows.Num(); ++PropertyRowIdx)
+			{
+				FDMPropertyHandle& PropertyRow = InOutPropertyRows[PropertyRowIdx];
+
+				if (PropertyRow.NameOverride.IsSet())
+				{
+					continue;
+				}
+
+				if (!PropertyRow.PropertyHandle.IsValid()
+					|| !PropertyRow.PropertyHandle->GetProperty())
+				{
+					continue;
+				}
+
+				TArray<UObject*> Outers;
+				PropertyRow.PropertyHandle->GetOuterObjects(Outers);
+
+				if (Outers.IsEmpty() || !Outers[0]->IsA<UDMMaterialValue>())
+				{
+					continue;
+				}
+
+				PropertyRow.NameOverride = InputConnectors[InputIdx].Name;
+			}
+		}
+	}
+
+	const TArray<FName>& StageProperties = Stage->GetEditableProperties();
+
+	for (const FName& StageProperty : StageProperties)
+	{
+		if (StageProperty != StageInputsName)
+		{
+			AddPropertyEditRows(InComponentEditorWidget, Stage, StageProperty, InOutPropertyRows, InOutProcessedObjects);
+			continue;
+		}
+	}
+}
diff --git a/Engine/Plugins/Experimental/DynamicMaterial/Source/DynamicMaterialEditor/Private/UI/PropertyGenerators/DMThroughputPropertyRowGenerator.h b/Engine/Plugins/Experimental/DynamicMaterial/Source/DynamicMaterialEditor/Private/UI/PropertyGenerators/DMThroughputPropertyRowGenerator.h
new file mode 100644
--- /dev/null
+++ b/Engine/Plugins/Experimental/DynamicMaterial/Source/DynamicMaterialEditor/Private/UI/PropertyGenerators/DMThroughputPropertyRowGenerator.h
@@ -0,0 +1,18 @@
+// Copyright Epic Games, Inc. All Rights Reserved.
+
+#pragma once
+
+#include "UI/PropertyGenerators/DMComponentPropertyRowGenerator.h"
+
+class FDMThroughputPropertyRowGenerator : public FDMComponentPropertyRowGenerator
+{
+public:
+	static const TSharedRef<FDMThroughputPropertyRowGenerator>& Get();
+
+	virtual ~FDMThroughputPropertyRowGenerator() override = default;
+
+	//~ Begin FDMComponentPropertyRowGenerator
+	virtual void AddComponentProperties(const TSharedRef<SDMMaterialComponentEditor>& InComponentEditorWidget, UDMMaterialComponent* InComponent,
+		TArray<FDMPropertyHandle>& InOutPropertyRows, TSet<UDMMaterialComponent*>& InOutProcessedObjects) override;
+	//~ End FDMComponentPropertyRowGenerator
+};
diff --git a/Engine/Plugins/Experimental/DynamicMaterial/Source/DynamicMaterialEditor/Private/UI/Utils/DMWidgetStatics.cpp b/Engine/Plugins/Experimental/DynamicMaterial/Source/DynamicMaterialEditor/Private/UI/Utils/DMWidgetStatics.cpp
new file mode 100644
--- /dev/null
+++ b/Engine/Plugins/Experimental/DynamicMaterial/Source/DynamicMaterialEditor/Private/UI/Utils/DMWidgetStatics.cpp
@@ -0,0 +1,276 @@
+// Copyright Epic Games, Inc. All Rights Reserved.
+
+#include "UI/Utils/DMWidgetStatics.h"
+
+#include "Components/MaterialValues/DMMaterialValueFloat.h"
+#include "IDetailTreeNode.h"
+#include "IPropertyRowGenerator.h"
+#include "Modules/ModuleManager.h"
+#include "PropertyEditorModule.h"
+
+FDMWidgetStatics& FDMWidgetStatics::Get()
+{
+	 static FDMWidgetStatics Instance;
+	 return Instance;
+}
+
+bool FDMWidgetStatics::GetExpansionState(UObject* InOwner, FName InName, bool& bOutExpanded)
+{
+	const FExpansionItem ExpansionItem = {InOwner, InName};
+
+	if (const bool* State = ExpansionStates.Find(ExpansionItem))
+	{
+		bOutExpanded = *State;
+		return true;
+	}
+
+	return false;
+}
+
+void FDMWidgetStatics::SetExpansionState(UObject* InOwner, FName InName, bool bInIsExpanded)
+{
+	const FExpansionItem ExpansionItem = {InOwner, InName};
+	ExpansionStates.FindOrAdd(ExpansionItem) = bInIsExpanded;
+}
+
+FDMPropertyHandle FDMWidgetStatics::GetPropertyHandle(const SWidget* InOwningWidget, UObject* InObject, FName InPropertyName)
+{
+	TArray<FDMPropertyHandle>& PropertyHandles = PropertyHandleMap.FindOrAdd(InOwningWidget);
+
+	for (const FDMPropertyHandle& ExistingHandle : PropertyHandles)
+	{
+		if (ExistingHandle.PropertyHandle && ExistingHandle.PropertyHandle->GetProperty()->GetFName() == InPropertyName)
+		{
+			TArray<UObject*> Outers;
+			ExistingHandle.PropertyHandle->GetOuterObjects(Outers);
+
+			if (Outers.IsEmpty() == false && Outers[0] == InObject)
+			{
+				return ExistingHandle;
+			}
+		}
+	}
+
+	if (TSharedPtr<IPropertyRowGenerator> PropertyRowGenerator = SearchForGenerator(PropertyHandles, InObject))
+	{
+		FDMPropertyHandle PropertyHandle;
+		PropertyHandle.PropertyRowGenerator = PropertyRowGenerator;
+
+		if (TSharedPtr<IDetailTreeNode> DetailTreeNode = SearchGeneratorForNode(PropertyRowGenerator.ToSharedRef(), InPropertyName))
+		{
+			PropertyHandle.DetailTreeNode = DetailTreeNode;
+			PropertyHandle.PropertyHandle = DetailTreeNode->CreatePropertyHandle();
+
+			AddPropertyMetaData(InObject, InPropertyName, PropertyHandle);
+
+			return PropertyHandle;
+		}
+
+		return PropertyHandle;
+	}
+
+	FDMPropertyHandle NewHandle = CreatePropertyHandle(InOwningWidget, InObject, InPropertyName);
+
+	if (!NewHandle.PropertyHandle.IsValid() && NewHandle.DetailTreeNode.IsValid())
+	{
+		NewHandle.PropertyHandle = NewHandle.DetailTreeNode->CreatePropertyHandle();
+	}
+
+	AddPropertyMetaData(InObject, InPropertyName, NewHandle);
+
+	PropertyHandles.Add(NewHandle);
+
+	return NewHandle;
+}
+
+void FDMWidgetStatics::ClearPropertyHandles(const SWidget* InOwningWidget)
+{
+	PropertyHandleMap.Remove(InOwningWidget);
+}
+
+FDMPropertyHandle FDMWidgetStatics::CreatePropertyHandle(const void* InOwningWidget, UObject* InObject, FName InPropertyName)
+{
+	FDMPropertyHandle PropertyHandle;
+
+	FPropertyEditorModule& PropertyEditor = FModuleManager::Get().LoadModuleChecked<FPropertyEditorModule>(TEXT("PropertyEditor"));
+
+	FPropertyRowGeneratorArgs RowGeneratorArgs;
+
+	if (UDMMaterialComponent* Component = Cast<UDMMaterialComponent>(InObject))
+	{
+		RowGeneratorArgs.NotifyHook = Component;
+	}
+
+	PropertyHandle.PropertyRowGenerator = PropertyEditor.CreatePropertyRowGenerator(RowGeneratorArgs);
+	PropertyHandle.PropertyRowGenerator->SetObjects({InObject});
+
+	if (const TSharedPtr<IDetailTreeNode> FoundTreeNode = SearchGeneratorForNode(
+		PropertyHandle.PropertyRowGenerator.ToSharedRef(), InPropertyName))
+	{
+		PropertyHandle.DetailTreeNode = FoundTreeNode;
+		PropertyHandle.PropertyHandle = FoundTreeNode->CreatePropertyHandle();
+		return PropertyHandle;
+	}
+
+	return PropertyHandle;
+}
+
+TSharedPtr<IDetailTreeNode> FDMWidgetStatics::SearchNodesForProperty(const TArray<TSharedRef<IDetailTreeNode>>& InNodes, FName InPropertyName)
+{
+	for (const TSharedRef<IDetailTreeNode>& ChildNode : InNodes)
+	{
+		switch (ChildNode->GetNodeType())
+		{
+			case EDetailNodeType::Category:
+			{
+				TArray<TSharedRef<IDetailTreeNode>> CategoryChildNodes;
+				ChildNode->GetChildren(CategoryChildNodes);
+
+				if (TSharedPtr<IDetailTreeNode> FoundNode = SearchNodesForProperty(CategoryChildNodes, InPropertyName))
+				{
+					return FoundNode;
+				}
+
+				break;
+			}
+
+			case EDetailNodeType::Item:
+				if (ChildNode->GetNodeName() == InPropertyName)
+				{
+					return ChildNode;
+				}
+				break;
+
+			default:
+				// Do nothing
+				break;
+		}
+	}
+
+	return nullptr;
+}
+
+TSharedPtr<IDetailTreeNode> FDMWidgetStatics::SearchGeneratorForNode(const TSharedRef<IPropertyRowGenerator>& InGenerator, FName InPropertyName)
+{
+	return SearchNodesForProperty(InGenerator->GetRootTreeNodes(), InPropertyName);
+}
+
+TSharedPtr<IPropertyRowGenerator> FDMWidgetStatics::SearchForGenerator(const TArray<FDMPropertyHandle>& InPropertyHandles, UObject* InObject)
+{
+	if (!InObject)
+	{
+		return nullptr;
+	}
+
+	for (const FDMPropertyHandle& PropertyHandle : InPropertyHandles)
+	{
+		if (PropertyHandle.PropertyRowGenerator.IsValid())
+		{
+			for (const TWeakObjectPtr<UObject>& WeakObject : PropertyHandle.PropertyRowGenerator->GetSelectedObjects())
+			{
+				if (WeakObject.Get() == InObject)
+				{
+					return PropertyHandle.PropertyRowGenerator;
+				}
+			}
+		}
+	}
+
+	return nullptr;
+}
+
+void FDMWidgetStatics::AddPropertyMetaData(UObject* InObject, FName InPropertyName, FDMPropertyHandle& InPropertyHandle)
+{
+	FProperty* Property = nullptr;
+
+	if (InPropertyHandle.PropertyHandle.IsValid())
+	{
+		TSharedRef<IPropertyHandle> PropertyHandleRef = InPropertyHandle.PropertyHandle.ToSharedRef();
+
+		InPropertyHandle.Priority = GetPriority(PropertyHandleRef);
+		InPropertyHandle.bKeyframeable = IsKeyframeable(PropertyHandleRef);
+
+		Property = InPropertyHandle.PropertyHandle->GetProperty();
+
+		if (UDMMaterialValueFloat* FloatValue = Cast<UDMMaterialValueFloat>(InObject))
+		{
+			if (FloatValue->HasValueRange())
+			{
+				const FName UIMin = FName("UIMin");
+				const FName UIMax = FName("UIMax");
+				const FName ClampMin = FName("ClampMin");
+				const FName ClampMax = FName("ClampMax");
+
+				InPropertyHandle.PropertyHandle->SetInstanceMetaData(UIMin, FString::SanitizeFloat(FloatValue->GetValueRange().Min));
+				InPropertyHandle.PropertyHandle->SetInstanceMetaData(ClampMin, FString::SanitizeFloat(FloatValue->GetValueRange().Min));
+				InPropertyHandle.PropertyHandle->SetInstanceMetaData(UIMax, FString::SanitizeFloat(FloatValue->GetValueRange().Max));
+				InPropertyHandle.PropertyHandle->SetInstanceMetaData(ClampMax, FString::SanitizeFloat(FloatValue->GetValueRange().Max));
+			}
+		}
+	}
+	else if (IsValid(InObject))
+	{
+		Property = InObject->GetClass()->FindPropertyByName(InPropertyName);
+	}
+
+	if (FStructProperty* StructProperty = CastField<FStructProperty>(Property))
+	{
+		uint8 ComponentCount = 1;
+
+		if (StructProperty->Struct == TBaseStructure<FVector2D>::Get()
+			|| StructProperty->Struct == TVariantStructure<FVector2f>::Get())
+		{
+			ComponentCount = 2;
+		}
+
+		if (StructProperty->Struct == TBaseStructure<FVector>::Get()
+			|| StructProperty->Struct == TVariantStructure<FVector3f>::Get()
+			|| StructProperty->Struct == TBaseStructure<FRotator>::Get())
+		{
+			ComponentCount = 3;
+		}
+
+		// FLinearColor doesn't need the extra space
+		if (StructProperty->Struct == TBaseStructure<FVector4>::Get()
+			|| StructProperty->Struct == TVariantStructure<FVector4f>::Get())
+		{
+			ComponentCount = 4;
+		}
+
+		switch (ComponentCount)
+		{
+			case 0:
+			case 1:
+				break;
+
+			case 2:
+				InPropertyHandle.MaxWidth = 200.f;
+				break;
+
+				// 3 and above
+			default:
+				InPropertyHandle.MaxWidth = 275.f;
+				break;
+		}
+	}
+}
+
+EDMPropertyHandlePriority FDMWidgetStatics::GetPriority(const TSharedRef<IPropertyHandle>& InPropertyHandle)
+{
+	if (InPropertyHandle->HasMetaData("HighPriority"))
+	{
+		return EDMPropertyHandlePriority::High;
+	}
+
+	if (InPropertyHandle->HasMetaData("LowPriority"))
+	{
+		return EDMPropertyHandlePriority::Low;
+	}
+
+	return EDMPropertyHandlePriority::Normal;
+}
+
+bool FDMWidgetStatics::IsKeyframeable(const TSharedRef<IPropertyHandle>& InPropertyHandle)
+{
+	return !InPropertyHandle->HasMetaData("NotKeyframeable");
+}
