Material Designer

diff --git a/Engine/Plugins/Experimental/DynamicMaterial/Source/DynamicMaterialEditor/Private/UI/Utils/DMWidgetStatics.h b/Engine/Plugins/Experimental/DynamicMaterial/Source/DynamicMaterialEditor/Private/UI/Utils/DMWidgetStatics.h
new file mode 100644
--- /dev/null
+++ b/Engine/Plugins/Experimental/DynamicMaterial/Source/DynamicMaterialEditor/Private/UI/Utils/DMWidgetStatics.h
@@ -0,0 +1,68 @@
+// Copyright Epic Games, Inc. All Rights Reserved.
+
+#pragma once
+
+#include "Containers/Array.h"
+#include "Containers/Map.h"
+#include "DMEDefs.h"
+#include "Templates/SharedPointerFwd.h"
+#include "UObject/ObjectKey.h"
+
+class FName;
+class SWidget;
+class UObject;
+
+class FDMWidgetStatics
+{
+public:
+	static FDMWidgetStatics& Get();
+
+	bool GetExpansionState(UObject* InOwner, FName InName, bool& bOutExpanded);
+
+	void SetExpansionState(UObject* InOwner, FName InName, bool bInIsExpanded);
+
+	FDMPropertyHandle GetPropertyHandle(const SWidget* InOwningWidget, UObject* InObject, FName InPropertyName);
+
+	void ClearPropertyHandles(const SWidget* InOwningWidget);
+
+private:
+	struct FExpansionItem
+	{
+		TObjectKey<UObject> Owner;
+		FName Name;
+
+		bool operator==(const FExpansionItem& InOther) const
+		{
+			return Owner == InOther.Owner
+				&& Name == InOther.Name;
+		}
+
+		friend uint32 GetTypeHash(const FExpansionItem& InItem)
+		{
+			return HashCombineFast(
+				GetTypeHash(InItem.Owner),
+				GetTypeHash(InItem.Name)
+			);
+		}
+	};
+
+	TMap<FExpansionItem, bool> ExpansionStates;
+
+	TMap<const SWidget*, TArray<FDMPropertyHandle>> PropertyHandleMap;
+
+	FDMPropertyHandle CreatePropertyHandle(const void* InOwningWidget, UObject* InObject, FName InPropertyName);
+
+	static TSharedPtr<IDetailTreeNode> SearchNodesForProperty(const TArray<TSharedRef<IDetailTreeNode>>& InNodes, FName InPropertyName);
+
+	static TSharedPtr<IDetailTreeNode> SearchGeneratorForNode(const TSharedRef<IPropertyRowGenerator>& InGenerator, FName InPropertyName);
+
+	static TSharedPtr<IPropertyRowGenerator> SearchForGenerator(const TArray<FDMPropertyHandle>& InPropertyHandles, UObject* InObject);
+
+	static void AddPropertyMetaData(UObject* InObject, FName InPropertyName, FDMPropertyHandle& InPropertyHandle);
+
+	/** Checks the meta data of a property meta data to check for high and low priority specifiers. */
+	static EDMPropertyHandlePriority GetPriority(const TSharedRef<IPropertyHandle>& InPropertyHandle);
+
+	/** Checks for the NotKeyframeable meta data. */
+	static bool IsKeyframeable(const TSharedRef<IPropertyHandle>& InPropertyHandle);
+};
diff --git a/Engine/Plugins/Experimental/DynamicMaterial/Source/DynamicMaterialEditor/Private/UI/Widgets/Editor/SDMMaterialComponentEditor.cpp b/Engine/Plugins/Experimental/DynamicMaterial/Source/DynamicMaterialEditor/Private/UI/Widgets/Editor/SDMMaterialComponentEditor.cpp
new file mode 100644
--- /dev/null
+++ b/Engine/Plugins/Experimental/DynamicMaterial/Source/DynamicMaterialEditor/Private/UI/Widgets/Editor/SDMMaterialComponentEditor.cpp
@@ -0,0 +1,189 @@
+// Copyright Epic Games, Inc. All Rights Reserved.
+
+#include "UI/Widgets/Editor/SDMMaterialComponentEditor.h"
+
+#include "Components/DMMaterialComponent.h"
+#include "Components/DMMaterialLayer.h"
+#include "Components/DMMaterialStage.h"
+#include "DetailLayoutBuilder.h"
+#include "DMEDefs.h"
+#include "DynamicMaterialEditorModule.h"
+#include "DynamicMaterialEditorStyle.h"
+#include "DynamicMaterialModule.h"
+#include "Model/DynamicMaterialModel.h"
+#include "Model/DynamicMaterialModelBase.h"
+#include "UI/Menus/DMMaterialStageSourceMenus.h"
+#include "UI/Widgets/Editor/SDMMaterialSlotEditor.h"
+#include "UI/Widgets/Editor/SlotEditor/SDMMaterialSlotLayerItem.h"
+#include "UI/Widgets/Editor/SlotEditor/SDMMaterialSlotLayerView.h"
+#include "UI/Widgets/Editor/SlotEditor/SDMMaterialStage.h"
+#include "UI/Widgets/SDMMaterialEditor.h"
+#include "Widgets/Input/SComboButton.h"
+
+#define LOCTEXT_NAMESPACE "SDMMaterialComponentEditor"
+
+void SDMMaterialComponentEditor::PrivateRegisterAttributes(FSlateAttributeDescriptor::FInitializer&)
+{
+}
+
+SDMMaterialComponentEditor::~SDMMaterialComponentEditor()
+{
+	if (!FDynamicMaterialModule::AreUObjectsSafe())
+	{
+		return;
+	}
+
+	if (UDMMaterialComponent* Component = GetComponent())
+	{
+		Component->GetOnUpdate().RemoveAll(this);
+	}
+}
+
+void SDMMaterialComponentEditor::Construct(const FArguments& InArgs, const TSharedRef<SDMMaterialEditor>& InEditorWidget, 
+	UDMMaterialComponent* InMaterialComponent)
+{
+	SDMObjectEditorWidgetBase::Construct(
+		SDMObjectEditorWidgetBase::FArguments(), 
+		InEditorWidget, 
+		InMaterialComponent
+	);
+
+	if (InMaterialComponent)
+	{
+		InMaterialComponent->GetOnUpdate().AddSP(this, &SDMMaterialComponentEditor::OnComponentUpdated);
+	}
+}
+
+UDMMaterialComponent* SDMMaterialComponentEditor::GetComponent() const
+{
+	return Cast<UDMMaterialComponent>(ObjectWeak.Get());
+}
+
+TSharedRef<SWidget> SDMMaterialComponentEditor::CreateSourceTypeEditWidget()
+{
+	TWeakPtr<SDMMaterialComponentEditor> ThisWeak = SharedThis(this);
+
+	return SNew(SBox)
+		.HAlign(EHorizontalAlignment::HAlign_Left)
+		.VAlign(EVerticalAlignment::VAlign_Center)
+		.Padding(0.f, 0.f, 0.f, 0.2f)
+		[
+			SNew(SComboButton)
+			.HasDownArrow(false)
+			.IsFocusable(true)
+			.ContentPadding(4.0f)
+			.ButtonStyle(FDynamicMaterialEditorStyle::Get(), "HoverHintOnly.Bordered.Dark")
+			.ToolTipText(LOCTEXT("ChangeLayer", "Change Stage Type"))
+			.OnGetMenuContent(this, &SDMMaterialComponentEditor::MakeSourceTypeEditWidgetMenuContent)
+			.ButtonContent()
+			[
+				SNew(STextBlock)
+				.Font(IDetailLayoutBuilder::GetDetailFont())
+				.Text(this, &SDMMaterialComponentEditor::GetSourceTypeEditWidgetText)
+			]
+		];
+}
+
+TSharedRef<SWidget> SDMMaterialComponentEditor::MakeSourceTypeEditWidgetMenuContent()
+{
+	if (UDMMaterialStage* Stage = Cast<UDMMaterialStage>(ObjectWeak.Get()))
+	{
+		if (UDMMaterialLayerObject* Layer = Stage->GetLayer())
+		{
+			if (UDMMaterialSlot* Slot = Layer->GetSlot())
+			{
+				if (TSharedPtr<SDMMaterialEditor> EditorWidget = GetEditorWidget())
+				{
+					if (TSharedPtr<SDMMaterialSlotEditor> SlotEditorWidget = EditorWidget->GetSlotEditorWidget())
+					{
+						if (SlotEditorWidget->GetSlot() == Slot)
+						{
+							if (TSharedPtr<SDMMaterialSlotLayerView> SlotLayerView = SlotEditorWidget->GetLayerView())
+							{
+								if (TSharedPtr<SDMMaterialSlotLayerItem> SlotLayerItem = SlotLayerView->GetWidgetForLayer(Layer))
+								{
+									if (TSharedPtr<SDMMaterialStage> StageWidget = SlotLayerItem->GetWidgetForStage(Stage))
+									{
+										return FDMMaterialStageSourceMenus::MakeChangeSourceMenu(SlotEditorWidget, StageWidget);
+									}
+								}
+							}
+						}
+					}
+				}
+			}
+		}
+	}
+
+	return SNullWidget::NullWidget;
+}
+
+FText SDMMaterialComponentEditor::GetSourceTypeEditWidgetText() const
+{
+	if (UDMMaterialStage* Stage = Cast<UDMMaterialStage>(ObjectWeak.Get()))
+	{
+		if (UDMMaterialStageSource* Source = Stage->GetSource())
+		{
+			return Source->GetStageDescription();
+		}
+	}
+
+	return FText::GetEmpty();
+}
+
+void SDMMaterialComponentEditor::OnComponentUpdated(UDMMaterialComponent* InComponent, EDMUpdateType InUpdateType)
+{
+	if (EnumHasAnyFlags(InUpdateType, EDMUpdateType::Structure))
+	{
+		if (TSharedPtr<SDMMaterialEditor> EditorWidget = GetEditorWidget())
+		{
+			EditorWidget->EditComponent(GetComponent(), /* Force refresh */ true);
+		}
+	}
+}
+
+TArray<FDMPropertyHandle> SDMMaterialComponentEditor::GetPropertyRows()
+{
+	TArray<FDMPropertyHandle> PropertyRows;
+	TSet<UDMMaterialComponent*> ProcessedObjects;
+
+	bool bIsDynamic = false;
+
+	if (TSharedPtr<SDMMaterialEditor> EditorWidget = GetEditorWidget())
+	{
+		if (UDynamicMaterialModelBase* MaterialModelBase = EditorWidget->GetMaterialModelBase())
+		{
+			bIsDynamic = !MaterialModelBase->IsA<UDynamicMaterialModel>();
+		}
+	}
+
+	if (UDMMaterialStage* Stage = Cast<UDMMaterialStage>(ObjectWeak.Get()))
+	{
+		FDMPropertyHandle& SourceHandle = PropertyRows.AddDefaulted_GetRef();
+		SourceHandle.ValueName = TEXT("SourceType");
+		SourceHandle.NameOverride = LOCTEXT("SourceType", "Source Type");
+		SourceHandle.bEnabled = !bIsDynamic;
+		SourceHandle.ValueWidget = CreateSourceTypeEditWidget();
+		SourceHandle.ResetToDefaultOverride = FResetToDefaultOverride::Hide(true);
+		SourceHandle.Priority = EDMPropertyHandlePriority::High;
+	}
+
+	FDynamicMaterialEditorModule::GeneratorComponentPropertyRows(
+		SharedThis(this), 
+		GetComponent(), 
+		PropertyRows, 
+		ProcessedObjects
+	);
+
+	return PropertyRows;
+}
+
+void SDMMaterialComponentEditor::OnUndo()
+{
+	if (TSharedPtr<SDMMaterialEditor> EditorWidget = GetEditorWidget())
+	{
+		EditorWidget->EditComponent(GetComponent(), /* Force refresh */ true);
+	}
+}
+
+#undef LOCTEXT_NAMESPACE
diff --git a/Engine/Plugins/Experimental/DynamicMaterial/Source/DynamicMaterialEditor/Private/UI/Widgets/Editor/SDMMaterialPropertySelector.cpp b/Engine/Plugins/Experimental/DynamicMaterial/Source/DynamicMaterialEditor/Private/UI/Widgets/Editor/SDMMaterialPropertySelector.cpp
new file mode 100644
--- /dev/null
+++ b/Engine/Plugins/Experimental/DynamicMaterial/Source/DynamicMaterialEditor/Private/UI/Widgets/Editor/SDMMaterialPropertySelector.cpp
@@ -0,0 +1,394 @@
+// Copyright Epic Games, Inc. All Rights Reserved.
+
+#include "UI/Widgets/Editor/SDMMaterialPropertySelector.h"
+
+#include "Components/DMMaterialProperty.h"
+#include "DetailLayoutBuilder.h"
+#include "Framework/Application/SlateApplication.h"
+#include "Model/DynamicMaterialModel.h"
+#include "Model/DynamicMaterialModelEditorOnlyData.h"
+#include "Styling/StyleColors.h"
+#include "UI/Widgets/SDMMaterialEditor.h"
+#include "UObject/Class.h"
+#include "Utils/DMPrivate.h"
+#include "Widgets/Images/SImage.h"
+#include "Widgets/Input/SCheckBox.h"
+#include "Widgets/Layout/SBox.h"
+#include "Widgets/Layout/SScrollBox.h"
+#include "Widgets/SBoxPanel.h"
+#include "Widgets/SNullWidget.h"
+#include "Widgets/Text/STextBlock.h"
+
+#define LOCTEXT_NAMESPACE "SDMMaterialPropertySelector"
+
+void SDMMaterialPropertySelector::PrivateRegisterAttributes(FSlateAttributeDescriptor::FInitializer&)
+{
+}
+
+void SDMMaterialPropertySelector::Construct(const FArguments& InArgs, const TSharedRef<SDMMaterialEditor> InEditorWidget)
+{
+	EditorWidgetWeak = InEditorWidget;
+	SelectedProperty = EDMMaterialPropertyType::None;
+
+	ChildSlot
+	[
+		SNew(SScrollBox)
+		.Orientation(EOrientation::Orient_Vertical)
+		+ SScrollBox::Slot()
+		[
+			CreateSlot_PropertyList()
+		]
+	];
+}
+
+TSharedPtr<SDMMaterialEditor> SDMMaterialPropertySelector::GetEditorWidget() const
+{
+	return EditorWidgetWeak.Pin();
+}
+
+EDMMaterialPropertyType SDMMaterialPropertySelector::GetSelectedProperty() const
+{
+	return SelectedProperty;
+}
+
+void SDMMaterialPropertySelector::SetGlobalSettings()
+{
+	SelectedProperty = EDMMaterialPropertyType::None;
+
+	TSharedPtr<SDMMaterialEditor> EditorWidget = GetEditorWidget();
+
+	if (!EditorWidget.IsValid())
+	{
+		return;
+	}
+
+	EditorWidget->EditSlot(nullptr);
+}
+
+void SDMMaterialPropertySelector::SetSelectedProperty(EDMMaterialPropertyType InMaterialProperty)
+{
+	if (SelectedProperty == InMaterialProperty)
+	{
+		return;
+	}
+
+	SelectedProperty = InMaterialProperty;
+
+	OnSelectedPropertyChanged();
+}
+
+UDynamicMaterialModelEditorOnlyData* SDMMaterialPropertySelector::GetEditorOnlyData() const
+{
+	TSharedPtr<SDMMaterialEditor> EditorWidget = GetEditorWidget();
+
+	if (!EditorWidget.IsValid())
+	{
+		return nullptr;
+	}
+
+	UDynamicMaterialModel* MaterialModel = EditorWidget->GetMaterialModel();
+
+	if (!IsValid(MaterialModel))
+	{
+		return nullptr;
+	}
+
+	return UDynamicMaterialModelEditorOnlyData::Get(MaterialModel);
+}
+
+TSharedPtr<SDMMaterialSlotEditor> SDMMaterialPropertySelector::GetSlotEditorWidget() const
+{
+	TSharedPtr<SDMMaterialEditor> EditorWidget = GetEditorWidget();
+
+	if (!EditorWidget.IsValid())
+	{
+		return nullptr;
+	}
+
+	return EditorWidget->GetSlotEditorWidget();
+}
+
+TSharedRef<SWidget> SDMMaterialPropertySelector::CreateSlot_PropertyList()
+{
+	using namespace UE::DynamicMaterialEditor::Private;
+
+	TSharedRef<SVerticalBox> NewSlotList = SNew(SVerticalBox);
+
+	UDynamicMaterialModelEditorOnlyData* EditorOnlyData = GetEditorOnlyData();
+
+	if (!EditorOnlyData)
+	{
+		return NewSlotList;
+	}
+
+	NewSlotList->AddSlot()
+		.AutoHeight()
+		[
+			CreateSlot_SelectButton(EDMMaterialPropertyType::None)
+		];
+
+	for (const TPair<EDMMaterialPropertyType, UDMMaterialProperty*>& PropertyPair : EditorOnlyData->GetMaterialProperties())
+	{
+		if (IsCustomMaterialProperty(PropertyPair.Key))
+		{
+			continue;
+		}
+
+		NewSlotList->AddSlot()
+			.AutoHeight()
+			[
+				CreateSlot_SelectButton(PropertyPair.Key)
+			];
+	}
+
+	return NewSlotList;
+}
+
+TSharedRef<SWidget> SDMMaterialPropertySelector::CreateSlot_SelectButton(EDMMaterialPropertyType InMaterialProperty)
+{
+	const FText ButtonText = InMaterialProperty == EDMMaterialPropertyType::None
+		? LOCTEXT("GlobalSettings", "Global Settings")
+		: StaticEnum<EDMMaterialPropertyType>()->GetDisplayNameTextByValue(static_cast<int64>(InMaterialProperty));
+
+	const FText Format = LOCTEXT("PropertyFormat", "Edit the {0} property.\n\n- Control+Left click: Toggle the property\n\nProperty may not be toggleable if it is not valid for the material type..");
+	UEnum* MaterialPropertyEnum = StaticEnum<EDMMaterialPropertyType>();
+
+	const FText ToolTip = (InMaterialProperty == EDMMaterialPropertyType::None)
+		? LOCTEXT("GeneralSettingsToolTip", "Edit the material global settings.")
+		: FText::Format(Format, MaterialPropertyEnum->GetDisplayNameTextByValue(static_cast<int64>(InMaterialProperty)));
+
+	TSharedRef<SBox> Outer = SNew(SBox)
+		[		
+			SNew(SCheckBox)
+			.Style(FAppStyle::Get(), "DetailsView.SectionButton")
+			.HAlign(EHorizontalAlignment::HAlign_Center)
+			.Padding(0.f)
+			.IsEnabled(this, &SDMMaterialPropertySelector::GetPropertySelectEnabled, InMaterialProperty)
+			.IsChecked(this, &SDMMaterialPropertySelector::GetPropertySelectState, InMaterialProperty)
+			.OnCheckStateChanged(this, &SDMMaterialPropertySelector::OnPropertySelectStateChanged, InMaterialProperty)
+			.ToolTipText(ToolTip)
+			.Content()
+			[
+				SNew(SBox)
+				.WidthOverride(135.f)
+				[
+					SNew(SHorizontalBox)
+					+SHorizontalBox::Slot()
+					.VAlign(VAlign_Center)
+					.AutoWidth()
+					[
+						SNew(SImage)
+						.Image(FAppStyle::Get().GetBrush("FilterBar.FilterImage"))
+						.ColorAndOpacity(this, &SDMMaterialPropertySelector::GetPropertySelectButtonChipColor, InMaterialProperty)
+					]
+					+SHorizontalBox::Slot()
+					.Padding(10.f, 6.f)
+					.VAlign(VAlign_Center)
+					.FillWidth(1.f)
+					[
+						SNew(STextBlock)
+						.Font(IDetailLayoutBuilder::GetDetailFont())
+						.Text(ButtonText)
+					]
+				]
+			]
+		];
+
+	// This is triggered when the inner checkbox is disabled.
+	Outer->SetOnMouseButtonDown(FPointerEventHandler::CreateSP(this, &SDMMaterialPropertySelector::OnPropertySelectMouseDown, InMaterialProperty));
+
+	return Outer;
+}
+
+bool SDMMaterialPropertySelector::IsPropertyEnabled(EDMMaterialPropertyType InMaterialProperty) const
+{
+	UDynamicMaterialModelEditorOnlyData* EditorOnlyData = GetEditorOnlyData();
+
+	if (!EditorOnlyData)
+	{
+		return false;
+	}
+
+	return !!EditorOnlyData->GetMaterialProperty(InMaterialProperty);
+}
+
+bool SDMMaterialPropertySelector::DoesPropertySlotExist(EDMMaterialPropertyType InMaterialProperty) const
+{
+	UDynamicMaterialModelEditorOnlyData* EditorOnlyData = GetEditorOnlyData();
+
+	if (!EditorOnlyData)
+	{
+		return false;
+	}
+
+	UDMMaterialProperty* MaterialProperty = EditorOnlyData->GetMaterialProperty(InMaterialProperty);
+
+	if (!MaterialProperty)
+	{
+		return false;
+	}
+
+	if (!MaterialProperty->IsEnabled())
+	{
+		return false;
+	}
+
+	return !!EditorOnlyData->GetSlotForMaterialProperty(InMaterialProperty);
+}
+
+bool SDMMaterialPropertySelector::SetPropertyEnabled(EDMMaterialPropertyType InMaterialProperty, bool bInEnabled)
+{
+	UDynamicMaterialModelEditorOnlyData* EditorOnlyData = GetEditorOnlyData();
+
+	if (!EditorOnlyData)
+	{
+		return false;
+	}
+
+	UDMMaterialProperty* MaterialProperty = EditorOnlyData->GetMaterialProperty(InMaterialProperty);
+
+	if (!MaterialProperty)
+	{
+		return false;
+	}
+
+	MaterialProperty->SetEnabled(bInEnabled);
+
+	if (InMaterialProperty == SelectedProperty)
+	{
+		SetGlobalSettings();
+	}
+
+	if (!bInEnabled)
+	{
+		return true;
+	}
+
+	if (EditorOnlyData->GetSlotForMaterialProperty(InMaterialProperty))
+	{
+		return true;
+	}
+
+	// Return true if the slot was successfully added
+	return !!EditorOnlyData->AddSlotForMaterialProperty(InMaterialProperty);
+}
+
+bool SDMMaterialPropertySelector::GetPropertySelectEnabled(EDMMaterialPropertyType InMaterialProperty) const
+{
+	if (InMaterialProperty == EDMMaterialPropertyType::None)
+	{
+		return true;
+	}
+
+	return DoesPropertySlotExist(InMaterialProperty);
+}
+
+ECheckBoxState SDMMaterialPropertySelector::GetPropertySelectState(EDMMaterialPropertyType InMaterialProperty) const
+{
+	return InMaterialProperty == SelectedProperty
+		? ECheckBoxState::Checked
+		: ECheckBoxState::Unchecked;
+}
+
+void SDMMaterialPropertySelector::OnPropertySelectStateChanged(ECheckBoxState InState, EDMMaterialPropertyType InMaterialProperty)
+{
+	if (InMaterialProperty == EDMMaterialPropertyType::None)
+	{
+		SetGlobalSettings();
+		return;
+	}
+
+
+	const FSlateApplication& SlateApplication = FSlateApplication::Get();
+
+	if (SlateApplication.GetModifierKeys().IsControlDown() || SlateApplication.GetModifierKeys().IsCommandDown())
+	{
+		SetPropertyEnabled(InMaterialProperty, /* Enabled */ false);
+		return;
+	}
+
+	UDynamicMaterialModelEditorOnlyData* EditorOnlyData = GetEditorOnlyData();
+
+	if (!EditorOnlyData)
+	{
+		return;
+	}
+
+	UDMMaterialProperty* MaterialProperty = EditorOnlyData->GetMaterialProperty(InMaterialProperty);
+
+	if (!MaterialProperty)
+	{
+		return;
+	}
+
+	if (!MaterialProperty->IsEnabled())
+	{
+		return;
+	}
+
+	SetSelectedProperty(InMaterialProperty);
+}
+
+FSlateColor SDMMaterialPropertySelector::GetPropertySelectButtonChipColor(EDMMaterialPropertyType InMaterialProperty) const
+{
+	if (InMaterialProperty == EDMMaterialPropertyType::None)
+	{
+		return FStyleColors::AccentGreen;
+	}
+
+	if (GetPropertySelectEnabled(InMaterialProperty))
+	{
+		return FStyleColors::Primary;
+	}
+
+	return FStyleColors::Panel;
+}
+
+FReply SDMMaterialPropertySelector::OnPropertySelectMouseDown(const FGeometry& InGeometry, const FPointerEvent& InEvent, 
+	EDMMaterialPropertyType InMaterialProperty)
+{
+	if (!InEvent.GetModifierKeys().IsControlDown() && !InEvent.GetModifierKeys().IsCommandDown())
+	{
+		return FReply::Unhandled();
+	}
+
+	if (SetPropertyEnabled(InMaterialProperty, /* Enabled */ true))
+	{
+		SetSelectedProperty(InMaterialProperty);
+	}
+
+	return FReply::Handled();
+}
+
+void SDMMaterialPropertySelector::OnSelectedPropertyChanged()
+{
+	TSharedPtr<SDMMaterialEditor> EditorWidget = GetEditorWidget();
+
+	if (!EditorWidget.IsValid())
+	{
+		return;
+	}
+
+	UDynamicMaterialModelEditorOnlyData* EditorOnlyData = GetEditorOnlyData();
+
+	if (!EditorOnlyData)
+	{
+		return;
+	}
+
+	UDMMaterialSlot* Slot = EditorOnlyData->GetSlotForMaterialProperty(SelectedProperty);
+
+	if (!Slot)
+	{
+		Slot = EditorOnlyData->AddSlotForMaterialProperty(SelectedProperty);
+
+		if (!Slot)
+		{
+			return;
+		}
+	}
+
+	EditorWidget->EditSlot(Slot);
+}
+
+#undef LOCTEXT_NAMESPACE
diff --git a/Engine/Plugins/Experimental/DynamicMaterial/Source/DynamicMaterialEditor/Private/UI/Widgets/Editor/SDMMaterialPropertySelector.h b/Engine/Plugins/Experimental/DynamicMaterial/Source/DynamicMaterialEditor/Private/UI/Widgets/Editor/SDMMaterialPropertySelector.h
new file mode 100644
--- /dev/null
+++ b/Engine/Plugins/Experimental/DynamicMaterial/Source/DynamicMaterialEditor/Private/UI/Widgets/Editor/SDMMaterialPropertySelector.h
@@ -0,0 +1,65 @@
+// Copyright Epic Games, Inc. All Rights Reserved.
+
+#pragma once
+
+#include "Widgets/SCompoundWidget.h"
+
+#include "Templates/SharedPointer.h"
+#include "UObject/WeakObjectPtr.h"
+
+class SDMMaterialEditor;
+class SDMMaterialSlotEditor;
+class SWidget;
+class UDMMaterialProperty;
+class UDMMaterialSlot;
+class UDynamicMaterialModelEditorOnlyData;
+enum class ECheckBoxState : uint8;
+enum class EDMMaterialPropertyType : uint8;
+
+class SDMMaterialPropertySelector : public SCompoundWidget
+{
+	SLATE_DECLARE_WIDGET(SDMMaterialPropertySelector, SCompoundWidget)
+
+	SLATE_BEGIN_ARGS(SDMMaterialPropertySelector) {}
+	SLATE_END_ARGS()
+
+public:
+	virtual ~SDMMaterialPropertySelector() override = default;
+
+	void Construct(const FArguments& InArgs, const TSharedRef<SDMMaterialEditor> InEditorWidget);
+
+	TSharedPtr<SDMMaterialEditor> GetEditorWidget() const;
+
+	EDMMaterialPropertyType GetSelectedProperty() const;
+
+	void SetGlobalSettings();
+
+	void SetSelectedProperty(EDMMaterialPropertyType InMaterialProperty);
+
+protected:
+	TWeakPtr<SDMMaterialEditor> EditorWidgetWeak;
+	EDMMaterialPropertyType SelectedProperty;
+
+	UDynamicMaterialModelEditorOnlyData* GetEditorOnlyData() const;
+
+	TSharedPtr<SDMMaterialSlotEditor> GetSlotEditorWidget() const;
+
+	TSharedRef<SWidget> CreateSlot_PropertyList();
+
+	TSharedRef<SWidget> CreateSlot_SelectButton(EDMMaterialPropertyType InMaterialProperty);
+
+	bool IsPropertyEnabled(EDMMaterialPropertyType InMaterialProperty) const;
+
+	bool SetPropertyEnabled(EDMMaterialPropertyType InMaterialProperty, bool bInEnabled);
+
+	bool DoesPropertySlotExist(EDMMaterialPropertyType InMaterialProperty) const;
+
+	bool GetPropertySelectEnabled(EDMMaterialPropertyType InMaterialProperty) const;
+	ECheckBoxState GetPropertySelectState(EDMMaterialPropertyType InMaterialProperty) const;
+	void OnPropertySelectStateChanged(ECheckBoxState InState, EDMMaterialPropertyType InMaterialProperty);
+	FSlateColor GetPropertySelectButtonChipColor(EDMMaterialPropertyType InMaterialProperty) const;
+
+	FReply OnPropertySelectMouseDown(const FGeometry& InGeometry, const FPointerEvent& InEvent, EDMMaterialPropertyType InMaterialProperty);
+
+	void OnSelectedPropertyChanged();
+};
diff --git a/Engine/Plugins/Experimental/DynamicMaterial/Source/DynamicMaterialEditor/Private/UI/Widgets/SDMMaterialEditor.cpp b/Engine/Plugins/Experimental/DynamicMaterial/Source/DynamicMaterialEditor/Private/UI/Widgets/SDMMaterialEditor.cpp
new file mode 100644
--- /dev/null
+++ b/Engine/Plugins/Experimental/DynamicMaterial/Source/DynamicMaterialEditor/Private/UI/Widgets/SDMMaterialEditor.cpp
@@ -0,0 +1,866 @@
+// Copyright Epic Games, Inc. All Rights Reserved.
+
+#include "UI/Widgets/SDMMaterialEditor.h"
+
+#include "Components/DMMaterialLayer.h"
+#include "Components/DMMaterialProperty.h"
+#include "Components/DMMaterialSlot.h"
+#include "Components/DMMaterialStage.h"
+#include "DynamicMaterialEditorCommands.h"
+#include "DynamicMaterialEditorSettings.h"
+#include "DynamicMaterialModule.h"
+#include "Engine/Texture.h"
+#include "Framework/Commands/GenericCommands.h"
+#include "Materials/MaterialFunctionInterface.h"
+#include "Model/DynamicMaterialModelBase.h"
+#include "Model/DynamicMaterialModelDynamic.h"
+#include "Model/DynamicMaterialModelEditorOnlyData.h"
+#include "UI/Utils/DMPreviewMaterialManager.h"
+#include "UI/Widgets/Editor/SDMMaterialComponentEditor.h"
+#include "UI/Widgets/Editor/SDMMaterialGlobalSettingsEditor.h"
+#include "UI/Widgets/Editor/SDMMaterialPreview.h"
+#include "UI/Widgets/Editor/SDMMaterialPropertySelector.h"
+#include "UI/Widgets/Editor/SDMMaterialSlotEditor.h"
+#include "UI/Widgets/Editor/SDMStatusBar.h"
+#include "UI/Widgets/Editor/SDMToolBar.h"
+#include "UI/Widgets/SDMMaterialDesigner.h"
+#include "Utils/DMMaterialModelFunctionLibrary.h"
+#include "Widgets/Layout/SSplitter.h"
+#include "Widgets/SBoxPanel.h"
+#include "Widgets/SNullWidget.h"
+
+#define LOCTEXT_NAMESPACE "SDMMaterialEditor"
+
+namespace UE::DynamicMaterialEditor::Private
+{
+	namespace SlotList
+	{
+		constexpr int32 ToolBar = 0;
+		constexpr int32 MainLayout = 1;
+		constexpr int32 StatusBar = 2;
+
+		namespace Main
+		{
+			constexpr int32 Left = 0;
+			constexpr int32 Right = 1;
+		}
+
+		namespace Left
+		{
+			constexpr int32 Preview = 0;
+			constexpr int32 PropertySelector = 1;
+		}
+
+		namespace Right
+		{
+			constexpr int32 SlotEditor = 0;
+			constexpr int32 ComponentEditor = 1;
+		}
+	}
+}
+
+void SDMMaterialEditor::PrivateRegisterAttributes(FSlateAttributeDescriptor::FInitializer&)
+{
+}
+
+SDMMaterialEditor::SDMMaterialEditor()
+	: CommandList(MakeShared<FUICommandList>())
+	, PreviewMaterialManager(MakeShared<FDMPreviewMaterialManager>())
+	, bGlobalSettingsMode(true)
+{
+}
+
+SDMMaterialEditor::~SDMMaterialEditor()
+{
+	FCoreDelegates::OnEnginePreExit.RemoveAll(this);
+}
+
+void SDMMaterialEditor::Construct(const FArguments& InArgs, const TSharedRef<SDMMaterialDesigner>& InDesignerWidget)
+{
+	DesignerWidgetWeak = InDesignerWidget;
+	bGlobalSettingsMode = true;
+	PropertyToSelect.Reset();
+
+	Container = TDMWidgetSlot<SWidget>(SharedThis(this), 0, SNullWidget::NullWidget);
+
+	if (InArgs._MaterialProperty.IsSet())
+	{
+		SetObjectMaterialProperty(InArgs._MaterialProperty.GetValue());
+	}
+	else if (IsValid(InArgs._MaterialModelBase))
+	{
+		SetMaterialModelBase(InArgs._MaterialModelBase);
+	}
+	else
+	{
+		ensureMsgf(false, TEXT("No valid material model passed to Material DesignerWidget Editor."));
+	}
+
+	FCoreDelegates::OnEnginePreExit.AddSP(this, &SDMMaterialEditor::OnEnginePreExit);
+}
+
+TSharedPtr<SDMMaterialDesigner> SDMMaterialEditor::GetDesignerWidget() const
+{
+	return DesignerWidgetWeak.Pin();
+}
+
+UDynamicMaterialModelBase* SDMMaterialEditor::GetMaterialModelBase() const
+{
+	return MaterialModelBaseWeak.Get();
+}
+
+void SDMMaterialEditor::SetMaterialModelBase(UDynamicMaterialModelBase* InMaterialModelBase)
+{
+	MaterialModelBaseWeak = InMaterialModelBase;
+
+	if (UDynamicMaterialModelDynamic* MaterialModelDynamic = Cast<UDynamicMaterialModelDynamic>(InMaterialModelBase))
+	{
+		MaterialModelDynamic->EnsureComponents();
+	}
+
+	EditGlobalSettings();
+
+	CreateLayout();
+}
+
+UDynamicMaterialModel* SDMMaterialEditor::GetMaterialModel() const
+{
+	if (UDynamicMaterialModelBase* MaterialModelBase = MaterialModelBaseWeak.Get())
+	{
+		return MaterialModelBase->ResolveMaterialModel();
+	}
+	
+	return nullptr;
+}
+
+bool SDMMaterialEditor::IsDynamicModel() const
+{
+	return !!Cast<UDynamicMaterialModelDynamic>(MaterialModelBaseWeak.Get());
+}
+
+const FDMObjectMaterialProperty* SDMMaterialEditor::GetMaterialObjectProperty() const
+{
+	if (ObjectMaterialPropertyOpt.IsSet())
+	{
+		return &ObjectMaterialPropertyOpt.GetValue();
+	}
+
+	return nullptr;
+}
+
+void SDMMaterialEditor::SetObjectMaterialProperty(const FDMObjectMaterialProperty& InObjectProperty)
+{
+	UDynamicMaterialModelBase* MaterialModelBase = InObjectProperty.GetMaterialModelBase();
+
+	if (!ensureMsgf(MaterialModelBase, TEXT("Invalid object material property value.")))
+	{
+		ClearSlots();
+		return;
+	}
+
+	ObjectMaterialPropertyOpt = InObjectProperty;
+	SetMaterialModelBase(MaterialModelBase);
+}
+
+AActor* SDMMaterialEditor::GetMaterialActor() const
+{
+	if (ObjectMaterialPropertyOpt.IsSet())
+	{
+		return ObjectMaterialPropertyOpt.GetValue().GetTypedOuter<AActor>();
+	}
+
+	return nullptr;
+}
+
+bool SDMMaterialEditor::IsEditingGlobalSettings() const
+{
+	return bGlobalSettingsMode;
+}
+
+void SDMMaterialEditor::SetMaterialActor(AActor* InActor)
+{
+	if (GetMaterialActor() == InActor)
+	{
+		return;
+	}
+
+	TSharedRef<SDMToolBar> NewToolBar = SNew(SDMToolBar, SharedThis(this), InActor);
+
+	ToolBar << NewToolBar;
+}
+
+TSharedPtr<SDMMaterialSlotEditor> SDMMaterialEditor::GetSlotEditorWidget() const
+{
+	return &SlotEditor;
+}
+
+TSharedPtr<SDMMaterialComponentEditor> SDMMaterialEditor::GetComponentEditorWidget() const
+{
+	return &ComponentEditor;
+}
+
+void SDMMaterialEditor::SelectProperty(EDMMaterialPropertyType InProperty, bool bInForceRefresh)
+{
+	if (bInForceRefresh || !PropertySelector.IsValid())
+	{
+		PropertyToSelect = InProperty;
+		PropertySelector.Invalidate();
+		return;
+	}
+
+	if (PropertySelector->GetSelectedProperty() != InProperty)
+	{
+		PropertySelector->SetSelectedProperty(InProperty);
+	}
+	
+	PropertyToSelect.Reset();
+}
+
+const TSharedRef<FUICommandList>& SDMMaterialEditor::GetCommandList() const
+{
+	return CommandList;
+}
+
+TSharedRef<FDMPreviewMaterialManager> SDMMaterialEditor::GetPreviewMaterialManager() const
+{
+	return PreviewMaterialManager;
+}
+
+void SDMMaterialEditor::EditSlot(UDMMaterialSlot* InSlot, bool bInForceRefresh)
+{
+	if (!bInForceRefresh && SlotEditor.IsValid() && SlotEditor->GetSlot() == InSlot)
+	{
+		return;
+	}
+
+	Right.Invalidate();
+
+	SlotEditor.Invalidate();
+	SlotToEdit = InSlot;
+
+	ComponentEditor.Invalidate();
+	ComponentToEdit.Reset();
+
+	bGlobalSettingsMode = !InSlot;
+
+	if (InSlot)
+	{
+		for (const TObjectPtr<UDMMaterialLayerObject>& Layer : InSlot->GetLayers())
+		{
+			if (UDMMaterialStage* Stage = Layer->GetFirstValidStage(EDMMaterialLayerStage::All))
+			{
+				ComponentToEdit = Stage;
+				break;
+			}
+		}
+	}
+}
+
+void SDMMaterialEditor::EditComponent(UDMMaterialComponent* InComponent, bool bInForceRefresh)
+{
+	if (!bInForceRefresh && ComponentEditor.IsValid() && ComponentEditor->GetComponent() == InComponent)
+	{
+		return;
+	}
+
+	if (bGlobalSettingsMode)
+	{
+		Right.Invalidate();
+		SlotEditor.Invalidate();
+	}
+
+	bGlobalSettingsMode = false;
+
+	ComponentEditor.Invalidate();
+	ComponentToEdit = InComponent;
+}
+
+void SDMMaterialEditor::EditGlobalSettings(bool bInForceRefresh)
+{
+	if (bGlobalSettingsMode && !bInForceRefresh)
+	{
+		return;
+	}
+
+	if (!bGlobalSettingsMode)
+	{
+		Right.Invalidate();
+		SlotEditor.Invalidate();
+		ComponentEditor.Invalidate();
+	}
+
+	bGlobalSettingsMode = true;
+
+	GlobalSettingsEditor.Invalidate();
+}
+
+SDMMaterialEditor::FOnEditedSlotChanged::RegistrationType& SDMMaterialEditor::GetOnEditedSlotChanged()
+{
+	return OnEditedSlotChanged;
+}
+
+SDMMaterialEditor::FOnEditedComponentChanged::RegistrationType& SDMMaterialEditor::GetOnEditedComponentChanged()
+{
+	return OnEditedComponentChanged;
+}
+
+bool SDMMaterialEditor::SupportsKeyboardFocus() const
+{
+	return true;
+}
+
+void SDMMaterialEditor::Tick(const FGeometry& InAllottedGeometry, const double InCurrentTime, const float InDeltaTime)
+{
+	SCompoundWidget::Tick(InAllottedGeometry, InCurrentTime, InDeltaTime);
+
+	if (!FDynamicMaterialModule::AreUObjectsSafe())
+	{
+		return;
+	}
+
+	UDynamicMaterialModelBase* MaterialModelBase = GetMaterialModelBase();
+
+	if (!IsValid(MaterialModelBase))
+	{
+		Close();
+		return;
+	}
+
+	if (ObjectMaterialPropertyOpt.IsSet() && ObjectMaterialPropertyOpt->IsValid())
+	{
+		const FDMObjectMaterialProperty& ObjectMaterialProperty = ObjectMaterialPropertyOpt.GetValue();
+		UDynamicMaterialModelBase* MaterialModelBaseFromProperty = ObjectMaterialProperty.GetMaterialModelBase();
+
+		if (!UDMMaterialModelFunctionLibrary::IsModelValid(MaterialModelBaseFromProperty))
+		{
+			MaterialModelBase = nullptr;
+		}
+
+		if (MaterialModelBase != MaterialModelBaseFromProperty)
+		{
+			if (TSharedPtr<SDMMaterialDesigner> DesignerWidget = DesignerWidgetWeak.Pin())
+			{
+				DesignerWidget->OpenObjectMaterialProperty(ObjectMaterialProperty);
+				return;
+			}
+		}
+	}
+	else if (!UDMMaterialModelFunctionLibrary::IsModelValid(MaterialModelBase))
+	{
+		Close();
+		return;
+	}
+
+	ValidateSlots();
+}
+
+FReply SDMMaterialEditor::OnKeyDown(const FGeometry& InMyGeometry, const FKeyEvent& InKeyEvent)
+{
+	if (CommandList->ProcessCommandBindings(InKeyEvent))
+	{
+		return FReply::Handled();
+	}
+
+	// We accept the delete key bind, so we don't want this accidentally deleting actors and such.
+	// Always return handled to stop the event bubbling.
+	const TArray<TSharedRef<const FInputChord>> DeleteChords = {
+		FGenericCommands::Get().Delete->GetActiveChord(EMultipleKeyBindingIndex::Primary),
+		FGenericCommands::Get().Delete->GetActiveChord(EMultipleKeyBindingIndex::Secondary)
+	};
+
+	for (const TSharedRef<const FInputChord>& DeleteChord : DeleteChords)
+	{
+		if (DeleteChord->Key == InKeyEvent.GetKey())
+		{
+			return FReply::Handled();
+		}
+	}
+
+	return FReply::Unhandled();
+}
+
+void SDMMaterialEditor::PostUndo(bool bInSuccess)
+{
+	OnUndo();
+}
+
+void SDMMaterialEditor::PostRedo(bool bInSuccess)
+{
+	OnUndo();
+}
+
+void SDMMaterialEditor::BindCommands(SDMMaterialSlotEditor* InSlotEditor)
+{
+	const FGenericCommands& GenericCommands = FGenericCommands::Get();
+
+	CommandList = MakeShared<FUICommandList>();
+
+	CommandList->MapAction(
+		FDynamicMaterialEditorCommands::Get().AddDefaultLayer,
+		FExecuteAction::CreateSP(InSlotEditor, &SDMMaterialSlotEditor::AddNewLayer),
+		FCanExecuteAction::CreateSP(InSlotEditor, &SDMMaterialSlotEditor::CanAddNewLayer)
+	);
+
+	CommandList->MapAction(
+		FDynamicMaterialEditorCommands::Get().InsertDefaultLayerAbove,
+		FExecuteAction::CreateSP(InSlotEditor, &SDMMaterialSlotEditor::InsertNewLayer),
+		FCanExecuteAction::CreateSP(InSlotEditor, &SDMMaterialSlotEditor::CanInsertNewLayer)
+	);
+
+	CommandList->MapAction(
+		GenericCommands.Copy,
+		FExecuteAction::CreateSP(InSlotEditor, &SDMMaterialSlotEditor::CopySelectedLayer),
+		FCanExecuteAction::CreateSP(InSlotEditor, &SDMMaterialSlotEditor::CanCopySelectedLayer)
+	);
+
+	CommandList->MapAction(
+		GenericCommands.Cut,
+		FExecuteAction::CreateSP(InSlotEditor, &SDMMaterialSlotEditor::CutSelectedLayer),
+		FCanExecuteAction::CreateSP(InSlotEditor, &SDMMaterialSlotEditor::CanCutSelectedLayer)
+	);
+
+	CommandList->MapAction(
+		GenericCommands.Paste,
+		FExecuteAction::CreateSP(InSlotEditor, &SDMMaterialSlotEditor::PasteLayer),
+		FCanExecuteAction::CreateSP(InSlotEditor, &SDMMaterialSlotEditor::CanPasteLayer)
+	);
+
+	CommandList->MapAction(
+		GenericCommands.Duplicate,
+		FExecuteAction::CreateSP(InSlotEditor, &SDMMaterialSlotEditor::DuplicateSelectedLayer),
+		FCanExecuteAction::CreateSP(InSlotEditor, &SDMMaterialSlotEditor::CanDuplicateSelectedLayer)
+	);
+
+	CommandList->MapAction(
+		GenericCommands.Delete,
+		FExecuteAction::CreateSP(InSlotEditor, &SDMMaterialSlotEditor::DeleteSelectedLayer),
+		FCanExecuteAction::CreateSP(InSlotEditor, &SDMMaterialSlotEditor::CanDeleteSelectedLayer)
+	);
+}
+
+bool SDMMaterialEditor::IsPropertyValidForModel(EDMMaterialPropertyType InProperty) const
+{
+	UDynamicMaterialModelEditorOnlyData* EditorOnlyData = UDynamicMaterialModelEditorOnlyData::Get(MaterialModelBaseWeak);
+
+	if (!EditorOnlyData)
+	{
+		return false;
+	}
+
+	if (UDMMaterialProperty* Property = EditorOnlyData->GetMaterialProperty(InProperty))
+	{
+		if (Property->IsValidForModel(*EditorOnlyData))
+		{
+			return true;
+		}
+	}
+
+	if (InProperty == EDMMaterialPropertyType::Opacity)
+	{
+		if (UDMMaterialProperty* Property = EditorOnlyData->GetMaterialProperty(EDMMaterialPropertyType::OpacityMask))
+		{
+			return Property->IsValidForModel(*EditorOnlyData);
+		}
+	}
+
+	return false;
+}
+
+void SDMMaterialEditor::Close()
+{
+	if (TSharedPtr<SDMMaterialDesigner> DesignerWidget = DesignerWidgetWeak.Pin())
+	{
+		DesignerWidget->ShowSelectPrompt();
+	}
+}
+
+void SDMMaterialEditor::ValidateSlots()
+{
+	if (Container.HasBeenInvalidated())
+	{
+		CreateLayout();
+		return;
+	}
+
+	if (ToolBar.HasBeenInvalidated())
+	{
+		ToolBar << CreateSlot_ToolBar();
+	}
+
+	if (Main.HasBeenInvalidated())
+	{
+		Main << CreateSlot_Main();
+	}
+	else
+	{
+		if (Left.HasBeenInvalidated())
+		{
+			Left << CreateSlot_Left();
+		}
+		else
+		{
+			if (Preview.HasBeenInvalidated())
+			{
+				Preview << CreateSlot_Preview();
+			}
+
+			if (PropertySelector.HasBeenInvalidated())
+			{
+				PropertySelector << CreateSlot_PropertySelector();
+			}
+		}
+
+		if (Right.HasBeenInvalidated())
+		{
+			Right << CreateSlot_Right();
+		}
+		else if (bGlobalSettingsMode)
+		{
+			if (GlobalSettingsEditor.HasBeenInvalidated())
+			{
+				GlobalSettingsEditor << CreateSlot_GlobalSettingsEditor();
+			}
+			else
+			{
+				GlobalSettingsEditor->Validate();
+			}
+		}
+		else
+		{
+			if (SlotEditor.HasBeenInvalidated())
+			{
+				SlotEditor << CreateSlot_SlotEditor();
+			}
+			else
+			{
+				SlotEditor->ValidateSlots();
+			}
+
+			if (ComponentEditor.HasBeenInvalidated())
+			{
+				ComponentEditor << CreateSlot_ComponentEditor();
+			}
+			else
+			{
+				ComponentEditor->Validate();
+			}
+		}
+	}
+
+	if (StatusBar.HasBeenInvalidated())
+	{
+		StatusBar << CreateSlot_StatusBar();
+	}
+}
+
+void SDMMaterialEditor::ClearSlots()
+{
+	Container.ClearWidget();
+	ToolBar.ClearWidget();
+	Main.ClearWidget();
+	Left.ClearWidget();
+	Right.ClearWidget();
+	Preview.ClearWidget();
+	PropertySelector.ClearWidget();
+	SlotEditor.ClearWidget();
+	ComponentEditor.ClearWidget();
+	StatusBar.ClearWidget();
+}
+
+void SDMMaterialEditor::CreateLayout()
+{
+	Container << CreateSlot_Container();
+}
+
+TSharedRef<SWidget> SDMMaterialEditor::CreateSlot_Container()
+{
+	SVerticalBox::FSlot* ToolBarSlot = nullptr;
+	SVerticalBox::FSlot* MainSlot = nullptr;
+	SVerticalBox::FSlot* StatusBarSlot = nullptr;
+
+	TSharedRef<SVerticalBox> NewContainer = SNew(SVerticalBox)
+		+ SVerticalBox::Slot()
+		.Expose(ToolBarSlot)
+		.AutoHeight()
+		[
+			SNullWidget::NullWidget
+		]
+
+		+ SVerticalBox::Slot()
+		.Expose(MainSlot)
+		.FillHeight(1.0f)
+		[
+			SNullWidget::NullWidget
+		]
+
+		+ SVerticalBox::Slot()
+		.Expose(StatusBarSlot)
+		.AutoHeight()
+		[
+			SNullWidget::NullWidget
+		];
+
+	ToolBar = TDMWidgetSlot<SDMToolBar>(ToolBarSlot, CreateSlot_ToolBar());
+	Main = TDMWidgetSlot<SWidget>(MainSlot, CreateSlot_Main());
+	StatusBar = TDMWidgetSlot<SDMStatusBar>(StatusBarSlot, CreateSlot_StatusBar());
+
+	return NewContainer;
+}
+
+TSharedRef<SDMToolBar> SDMMaterialEditor::CreateSlot_ToolBar()
+{
+	return SNew(
+		SDMToolBar, 
+		SharedThis(this), 
+		ObjectMaterialPropertyOpt.IsSet()
+			? ObjectMaterialPropertyOpt->GetTypedOuter<AActor>()
+			: nullptr
+	);
+}
+
+TSharedRef<SWidget> SDMMaterialEditor::CreateSlot_Main()
+{
+	SHorizontalBox::FSlot* LeftSlot = nullptr;
+	SHorizontalBox::FSlot* RightSlot = nullptr;
+
+	TSharedRef<SHorizontalBox> NewMain = SNew(SHorizontalBox)
+		+ SHorizontalBox::Slot()
+		.Expose(LeftSlot)
+		.AutoWidth()
+		[
+			SNullWidget::NullWidget
+		]
+
+		+ SHorizontalBox::Slot()
+		.Expose(RightSlot)
+		.FillWidth(1.0f)
+		[
+			SNullWidget::NullWidget
+		];
+
+	Left = TDMWidgetSlot<SWidget>(LeftSlot, CreateSlot_Left());
+	Right = TDMWidgetSlot<SWidget>(RightSlot, CreateSlot_Right());
+
+	return NewMain;
+}
+
+TSharedRef<SWidget> SDMMaterialEditor::CreateSlot_Left()
+{
+	SVerticalBox::FSlot* PreviewSlot = nullptr;
+	SVerticalBox::FSlot* PropertySelectorSlot = nullptr;
+
+	TSharedRef<SVerticalBox> NewLeft = SNew(SVerticalBox)
+		+ SVerticalBox::Slot()
+		.Expose(PreviewSlot)
+		.AutoHeight()
+		.Padding(5.f)
+		[
+			SNullWidget::NullWidget
+		]
+
+		+ SVerticalBox::Slot()
+		.Expose(PropertySelectorSlot)
+		.FillHeight(1.0f)
+		.Padding(5.f, 0.f, 5.f, 5.f)
+		[
+			SNullWidget::NullWidget
+		];
+
+	Preview = TDMWidgetSlot<SDMMaterialPreview>(PreviewSlot, CreateSlot_Preview());
+	PropertySelector = TDMWidgetSlot<SDMMaterialPropertySelector>(PropertySelectorSlot, CreateSlot_PropertySelector());
+
+	return NewLeft;
+}
+
+TSharedRef<SWidget> SDMMaterialEditor::CreateSlot_Right()
+{
+	const bool bHasSlotToEdit = SlotToEdit.IsValid();
+
+	if (!bGlobalSettingsMode && !bHasSlotToEdit)
+	{
+		bGlobalSettingsMode = true;
+	}
+	else if (bHasSlotToEdit)
+	{
+		bGlobalSettingsMode = false;
+	}
+
+	return bGlobalSettingsMode
+		? CreateSlot_Right_GlobalSettings()
+		: CreateSlot_Right_Slot();
+}
+
+TSharedRef<SWidget> SDMMaterialEditor::CreateSlot_Right_GlobalSettings()
+{
+	SVerticalBox::FSlot* GlobalSettingsSlot = nullptr;
+
+	TSharedRef<SVerticalBox> NewRight = SNew(SVerticalBox)
+		+ SVerticalBox::Slot()
+		.Expose(GlobalSettingsSlot)
+		.AutoHeight()
+		[
+			SNullWidget::NullWidget
+		];
+
+	GlobalSettingsEditor = TDMWidgetSlot<SDMMaterialGlobalSettingsEditor>(GlobalSettingsSlot, CreateSlot_GlobalSettingsEditor());
+
+	return NewRight;
+}
+
+TSharedRef<SDMMaterialGlobalSettingsEditor> SDMMaterialEditor::CreateSlot_GlobalSettingsEditor()
+{
+	return SNew(SDMMaterialGlobalSettingsEditor, SharedThis(this), GetMaterialModelBase());
+}
+
+TSharedRef<SWidget> SDMMaterialEditor::CreateSlot_Right_Slot()
+{
+	float SplitterValue = 0.5;
+
+	if (UDynamicMaterialEditorSettings* Settings = UDynamicMaterialEditorSettings::Get())
+	{
+		SplitterValue = Settings->SplitterLocation;
+	}
+
+	SSplitter::FSlot* SlotEditorSlot = nullptr;
+	SSplitter::FSlot* ComponentEditorSlot = nullptr;
+
+	TSharedRef<SSplitter> NewRight = SNew(SSplitter)
+		.Style(FAppStyle::Get(), "DetailsView.Splitter")
+		.Orientation(Orient_Vertical)
+		.ResizeMode(ESplitterResizeMode::Fill)
+		.PhysicalSplitterHandleSize(3.0f)
+		.HitDetectionSplitterHandleSize(4.0f)
+		.OnSplitterFinishedResizing(this, &SDMMaterialEditor::OnRightSlotSplitterResized)
+
+		+ SSplitter::Slot()
+		.Expose(SlotEditorSlot)
+		.Resizable(true)
+		.SizeRule(SSplitter::ESizeRule::FractionOfParent)
+		.MinSize(50)
+		.Value(SplitterValue)
+		[
+			SNullWidget::NullWidget
+		]
+
+		+ SSplitter::Slot()
+		.Expose(ComponentEditorSlot)
+		.Resizable(true)
+		.SizeRule(SSplitter::ESizeRule::FractionOfParent)
+		.MinSize(50)
+		.Value(1.f - SplitterValue)
+		[
+			SNullWidget::NullWidget
+		];
+
+	SlotEditor = TDMWidgetSlot<SDMMaterialSlotEditor>(SlotEditorSlot, CreateSlot_SlotEditor());
+	ComponentEditor = TDMWidgetSlot<SDMMaterialComponentEditor>(ComponentEditorSlot, CreateSlot_ComponentEditor());
+
+	return NewRight;
+}
+
+TSharedRef<SDMMaterialPreview> SDMMaterialEditor::CreateSlot_Preview()
+{
+	return SNew(SDMMaterialPreview, SharedThis(this), GetMaterialModelBase());
+}
+
+TSharedRef<SDMMaterialPropertySelector> SDMMaterialEditor::CreateSlot_PropertySelector()
+{
+	TSharedRef<SDMMaterialPropertySelector> NewPropertySelector = SNew(SDMMaterialPropertySelector, SharedThis(this));
+
+	if (!PropertyToSelect.IsSet())
+	{
+		if (UDynamicMaterialModel* MaterialModel = GetMaterialModel())
+		{
+			if (UDynamicMaterialModelEditorOnlyData* EditorOnlyData = UDynamicMaterialModelEditorOnlyData::Get(MaterialModel))
+			{
+				for (const TPair<EDMMaterialPropertyType, UDMMaterialProperty*>& PropertyPair : EditorOnlyData->GetMaterialProperties())
+				{
+					if (PropertyPair.Value->IsEnabled() && PropertyPair.Value->IsValidForModel(*EditorOnlyData))
+					{
+						PropertyToSelect = PropertyPair.Key;
+						break;
+					}
+				}
+			}
+		}
+	}
+
+	if (PropertyToSelect.IsSet())
+	{
+		NewPropertySelector->SetSelectedProperty(PropertyToSelect.GetValue());
+		PropertyToSelect.Reset();
+	}
+
+	return NewPropertySelector;
+}
+
+TSharedRef<SDMMaterialSlotEditor> SDMMaterialEditor::CreateSlot_SlotEditor()
+{
+	UDMMaterialSlot* Slot = SlotToEdit.Get();
+	SlotToEdit.Reset();
+
+	TSharedRef<SDMMaterialSlotEditor> NewSlotEditor = SNew(SDMMaterialSlotEditor, SharedThis(this), Slot);
+
+	BindCommands(&*NewSlotEditor);
+
+	OnEditedSlotChanged.Broadcast(NewSlotEditor, Slot);
+
+	return NewSlotEditor;
+}
+
+TSharedRef<SDMMaterialComponentEditor> SDMMaterialEditor::CreateSlot_ComponentEditor()
+{
+	UDMMaterialComponent* Component = ComponentToEdit.Get();
+	ComponentToEdit.Reset();
+
+	TSharedRef<SDMMaterialComponentEditor> NewComponentEditor = SNew(SDMMaterialComponentEditor, SharedThis(this), Component);
+
+	OnEditedComponentChanged.Broadcast(NewComponentEditor, Component);
+
+	return NewComponentEditor;
+}
+
+TSharedRef<SDMStatusBar> SDMMaterialEditor::CreateSlot_StatusBar()
+{
+	return SNew(SDMStatusBar, SharedThis(this), GetMaterialModelBase());
+}
+
+void SDMMaterialEditor::OnUndo()
+{
+	UDynamicMaterialModelBase* MaterialModel = GetMaterialModelBase();
+
+	if (!IsValid(MaterialModel))
+	{
+		Close();
+		return;
+	}
+
+	if (UDynamicMaterialModelEditorOnlyData* ModelEditorOnlyData = UDynamicMaterialModelEditorOnlyData::Get(MaterialModelBaseWeak))
+	{
+		for (const TPair<EDMMaterialPropertyType, UDMMaterialProperty*>& PropertyPair : ModelEditorOnlyData->GetMaterialProperties())
+		{
+			if (PropertyPair.Value->IsEnabled())
+			{
+				PropertySelector->SetSelectedProperty(PropertyPair.Key);
+			}
+		}
+	}
+}
+
+void SDMMaterialEditor::OnEnginePreExit()
+{
+	Preview.ClearWidget();
+}
+
+void SDMMaterialEditor::OnRightSlotSplitterResized()
+{
+	UDynamicMaterialEditorSettings* Settings = UDynamicMaterialEditorSettings::Get();
+
+	if (SlotEditor.GetSlot())
+	{
+		Settings->SplitterLocation = static_cast<SSplitter::FSlot*>(SlotEditor.GetSlot())->GetSizeValue();
+		Settings->SaveConfig();
+	}
+}
+
+#undef LOCTEXT_NAMESPACE
diff --git a/Engine/Plugins/Experimental/DynamicMaterial/Source/DynamicMaterialEditor/Private/UI/Widgets/Visualizers/SDMMaterialComponentPreview.cpp b/Engine/Plugins/Experimental/DynamicMaterial/Source/DynamicMaterialEditor/Private/UI/Widgets/Visualizers/SDMMaterialComponentPreview.cpp
new file mode 100644
--- /dev/null
+++ b/Engine/Plugins/Experimental/DynamicMaterial/Source/DynamicMaterialEditor/Private/UI/Widgets/Visualizers/SDMMaterialComponentPreview.cpp
@@ -0,0 +1,187 @@
+// Copyright Epic Games, Inc. All Rights Reserved.
+
+#include "UI/Widgets/Visualizers/SDMMaterialComponentPreview.h"
+
+#include "Components/DMMaterialStage.h"
+#include "Components/DMMaterialValue.h"
+#include "Components/DMTextureUV.h"
+#include "Components/DMTextureUVDynamic.h"
+#include "DynamicMaterialModule.h"
+#include "Materials/Material.h"
+#include "Materials/MaterialInstanceDynamic.h"
+#include "Model/DynamicMaterialModel.h"
+#include "Model/DynamicMaterialModelDynamic.h"
+#include "UI/Utils/DMPreviewMaterialManager.h"
+#include "UI/Widgets/SDMMaterialEditor.h"
+#include "Widgets/Images/SImage.h"
+
+SDMMaterialComponentPreview::SDMMaterialComponentPreview()
+	: Brush(FSlateMaterialBrush(FVector2D(1.f, 1.f)))
+{
+	Brush.SetUVRegion(FBox2f(FVector2f::ZeroVector, FVector2f::UnitVector));
+}
+
+SDMMaterialComponentPreview::~SDMMaterialComponentPreview()
+{
+	if (!FDynamicMaterialModule::AreUObjectsSafe())
+	{
+		return;
+	}
+
+	if (UDMMaterialComponent* Component = ComponentWeak.Get())
+	{
+		Component->GetOnUpdate().RemoveAll(this);
+
+		if (TSharedPtr<SDMMaterialEditor> EditorWidget = EditorWidgetWeak.Pin())
+		{
+			EditorWidget->GetPreviewMaterialManager()->CreatePreviewMaterial(Component);
+		}
+	}
+
+	if (UDynamicMaterialModelBase* MaterialModelBase = MaterialModelBaseWeak.Get())
+	{
+		if (UDynamicMaterialModel* MaterialModel = Cast<UDynamicMaterialModel>(MaterialModelBase))
+		{
+			MaterialModel->GetOnValueUpdateDelegate().RemoveAll(this);
+			MaterialModel->GetOnTextureUVUpdateDelegate().RemoveAll(this);
+		}
+		else if (UDynamicMaterialModelDynamic* MaterialModelDynamic = Cast<UDynamicMaterialModelDynamic>(MaterialModelBase))
+		{
+			MaterialModelDynamic->GetOnValueDynamicUpdateDelegate().RemoveAll(this);
+			MaterialModelDynamic->GetOnTextureUVDynamicUpdateDelegate().RemoveAll(this);
+		}
+	}
+}
+
+void SDMMaterialComponentPreview::Construct(const FArguments& InArgs, const TSharedRef<SDMMaterialEditor>& InEditorWidget, UDMMaterialComponent* InComponent)
+{
+	EditorWidgetWeak = InEditorWidget;
+	ComponentWeak = InComponent;
+
+	PreviewMaterialBaseWeak = InEditorWidget->GetPreviewMaterialManager()->CreatePreviewMaterial(InComponent);
+	PreviewMaterialDynamicWeak = InEditorWidget->GetPreviewMaterialManager()->CreatePreviewMaterialDynamic(PreviewMaterialBaseWeak.Get());
+	MaterialModelBaseWeak = InEditorWidget->GetMaterialModelBase();
+
+	if (UDynamicMaterialModel* MaterialModel = Cast<UDynamicMaterialModel>(MaterialModelBaseWeak.Get()))
+	{
+		MaterialModel->GetOnValueUpdateDelegate().AddSP(this, &SDMMaterialComponentPreview::OnValueUpdated);
+		MaterialModel->GetOnTextureUVUpdateDelegate().AddSP(this, &SDMMaterialComponentPreview::OnTextureUVUpdated);
+		MaterialModel->ApplyComponents(PreviewMaterialDynamicWeak.Get());
+	}
+	else if (UDynamicMaterialModelDynamic* MaterialModelDynamic = Cast<UDynamicMaterialModelDynamic>(MaterialModelBaseWeak.Get()))
+	{
+		MaterialModelDynamic->GetOnValueDynamicUpdateDelegate().AddSP(this, &SDMMaterialComponentPreview::OnValueDynamicUpdated);
+		MaterialModelDynamic->GetOnTextureUVDynamicUpdateDelegate().AddSP(this, &SDMMaterialComponentPreview::OnTextureUVDynamicUpdated);
+		MaterialModelDynamic->ApplyComponents(PreviewMaterialDynamicWeak.Get());
+	}
+	else
+	{
+		return;
+	}
+
+	PreviewSize = InArgs._PreviewSize;
+
+	SetCanTick(true);
+
+	Brush.SetImageSize(PreviewSize.Get());
+
+	if (ensure(IsValid(InComponent)))
+	{
+		InComponent->GetOnUpdate().AddSP(this, &SDMMaterialComponentPreview::OnComponentUpdated);
+		OnComponentUpdated(InComponent, EDMUpdateType::Structure);
+	}
+
+	ChildSlot
+	[
+		SNew(SImage)
+		.Image(&Brush)
+		.DesiredSizeOverride(this, &SDMMaterialComponentPreview::GetPreviewSize)
+	];
+}
+
+void SDMMaterialComponentPreview::Tick(const FGeometry& AllottedGeometry, const double InCurrentTime, const float InDeltaTime)
+{
+	SCompoundWidget::Tick(AllottedGeometry, InCurrentTime, InDeltaTime);
+
+	if (!PreviewMaterialBaseWeak.IsValid() || !PreviewMaterialDynamicWeak.IsValid())
+	{
+		Brush.SetMaterial(nullptr);
+	}	
+}
+
+void SDMMaterialComponentPreview::OnComponentUpdated(UDMMaterialComponent* InComponent, EDMUpdateType InUpdateType)
+{
+	if (UDMMaterialStage* Stage = Cast<UDMMaterialStage>(InComponent))
+	{
+		if (Stage == ComponentWeak.Get() && IsValid(Stage) && Stage->IsComponentValid())
+		{
+			if (TSharedPtr<SDMMaterialEditor> EditorWidget = EditorWidgetWeak.Pin())
+			{
+				UMaterial* PreviewMaterialBase = PreviewMaterialBaseWeak.Get();
+
+				if (!PreviewMaterialBase)
+				{
+					PreviewMaterialBase = EditorWidget->GetPreviewMaterialManager()->CreatePreviewMaterial(Stage);
+				}
+
+				if (EnumHasAnyFlags(InUpdateType, EDMUpdateType::Structure))
+				{
+					Stage->GeneratePreviewMaterial(PreviewMaterialBase);
+
+					EditorWidget->GetPreviewMaterialManager()->FreePreviewMaterialDynamic(PreviewMaterialBase);
+					PreviewMaterialDynamicWeak = EditorWidget->GetPreviewMaterialManager()->CreatePreviewMaterialDynamic(PreviewMaterialBase);
+
+					UDynamicMaterialModelBase* MaterialModelBase = EditorWidget->GetMaterialModelBase();
+
+					if (UDynamicMaterialModel* MaterialModel = Cast<UDynamicMaterialModel>(MaterialModelBase))
+					{
+						MaterialModel->ApplyComponents(PreviewMaterialDynamicWeak.Get());
+					}
+					else if (UDynamicMaterialModelDynamic* MaterialModelDynamic = Cast<UDynamicMaterialModelDynamic>(MaterialModelBase))
+					{
+						MaterialModelDynamic->ApplyComponents(PreviewMaterialDynamicWeak.Get());
+					}
+
+					Brush.SetMaterial(PreviewMaterialDynamicWeak.Get());
+				}
+			}
+		}
+	}
+}
+
+TOptional<FVector2D> SDMMaterialComponentPreview::GetPreviewSize() const
+{
+	return PreviewSize.Get();
+}
+
+void SDMMaterialComponentPreview::OnValueUpdated(UDynamicMaterialModel* InMaterialModel, UDMMaterialValue* InValue)
+{
+	if (UMaterialInstanceDynamic* PreviewMaterialDynamic = PreviewMaterialDynamicWeak.Get())
+	{
+		InValue->SetMIDParameter(PreviewMaterialDynamic);
+	}
+}
+
+void SDMMaterialComponentPreview::OnTextureUVUpdated(UDynamicMaterialModel* InMaterialModel, UDMTextureUV* InTextureUV)
+{
+	if (UMaterialInstanceDynamic* PreviewMaterialDynamic = PreviewMaterialDynamicWeak.Get())
+	{
+		InTextureUV->SetMIDParameters(PreviewMaterialDynamic);
+	}
+}
+
+void SDMMaterialComponentPreview::OnValueDynamicUpdated(UDynamicMaterialModelDynamic* InMaterialModel, UDMMaterialValueDynamic* InValueDynamic)
+{
+	if (UMaterialInstanceDynamic* PreviewMaterialDynamic = PreviewMaterialDynamicWeak.Get())
+	{
+		InValueDynamic->SetMIDParameter(PreviewMaterialDynamic);
+	}
+}
+
+void SDMMaterialComponentPreview::OnTextureUVDynamicUpdated(UDynamicMaterialModelDynamic* InMaterialModel, UDMTextureUVDynamic* InTextureUVDynamic)
+{
+	if (UMaterialInstanceDynamic* PreviewMaterialDynamic = PreviewMaterialDynamicWeak.Get())
+	{
+		InTextureUVDynamic->SetMIDParameters(PreviewMaterialDynamic);
+	}
+}
diff --git a/Engine/Plugins/Experimental/DynamicMaterial/Source/DynamicMaterialEditor/Private/UI/Widgets/Visualizers/SDMMaterialComponentPreview.h b/Engine/Plugins/Experimental/DynamicMaterial/Source/DynamicMaterialEditor/Private/UI/Widgets/Visualizers/SDMMaterialComponentPreview.h
new file mode 100644
--- /dev/null
+++ b/Engine/Plugins/Experimental/DynamicMaterial/Source/DynamicMaterialEditor/Private/UI/Widgets/Visualizers/SDMMaterialComponentPreview.h
@@ -0,0 +1,61 @@
+// Copyright Epic Games, Inc. All Rights Reserved.
+
+#pragma once
+
+#include "Widgets/SCompoundWidget.h"
+
+#include "DMEDefs.h"
+#include "SlateMaterialBrush.h"
+#include "UObject/WeakObjectPtr.h"
+#include "UObject/WeakObjectPtrTemplates.h"
+
+class SDMMaterialEditor;
+class UDMMaterialComponent;
+class UDMMaterialValueDynamic;
+class UDMTextureUV;
+class UDMTextureUVDynamic;
+class UDynamicMaterialModelBase;
+class UDynamicMaterialModelDynamic;
+
+class SDMMaterialComponentPreview : public SCompoundWidget
+{
+public:
+	SLATE_BEGIN_ARGS(SDMMaterialComponentPreview)
+		: _PreviewSize(FVector2D(48.f, 48.f))
+		{}
+		SLATE_ATTRIBUTE(FVector2D, PreviewSize)
+	SLATE_END_ARGS()
+
+	SDMMaterialComponentPreview();
+
+	virtual ~SDMMaterialComponentPreview() override;
+
+	void Construct(const FArguments& InArgs, const TSharedRef<SDMMaterialEditor>& InEditorWidget, UDMMaterialComponent* InComponent);
+
+	FSlateMaterialBrush& GetBrush() { return Brush; }
+
+	//~ Begin SWidget
+	virtual void Tick(const FGeometry& InAllottedGeometry, const double InCurrentTime, const float InDeltaTime) override;
+	//~ End SWidget
+
+protected:
+	TWeakPtr<SDMMaterialEditor> EditorWidgetWeak;
+	TWeakObjectPtr<UDMMaterialComponent> ComponentWeak;
+	TWeakObjectPtr<UDynamicMaterialModelBase> MaterialModelBaseWeak;
+	TWeakObjectPtr<UMaterial> PreviewMaterialBaseWeak;
+	TWeakObjectPtr<UMaterialInstanceDynamic> PreviewMaterialDynamicWeak;
+	FSlateMaterialBrush Brush;
+	TAttribute<FVector2D> PreviewSize;
+
+	void OnComponentUpdated(UDMMaterialComponent* InComponent, EDMUpdateType InUpdateType);
+
+	TOptional<FVector2D> GetPreviewSize() const;
+
+	void OnValueUpdated(UDynamicMaterialModel* InMaterialModel, UDMMaterialValue* InValue);
+
+	void OnTextureUVUpdated(UDynamicMaterialModel* InMaterialModel, UDMTextureUV* InTextureUV);
+
+	void OnValueDynamicUpdated(UDynamicMaterialModelDynamic* InMaterialModel, UDMMaterialValueDynamic* InValueDynamic);
+
+	void OnTextureUVDynamicUpdated(UDynamicMaterialModelDynamic* InMaterialModel, UDMTextureUVDynamic* InTextureUVDynamic);
+};
diff --git a/Engine/Plugins/Experimental/DynamicMaterial/Source/DynamicMaterialEditor/Private/Utils/DMMaterialModelFunctionLibrary.cpp b/Engine/Plugins/Experimental/DynamicMaterial/Source/DynamicMaterialEditor/Private/Utils/DMMaterialModelFunctionLibrary.cpp
new file mode 100644
--- /dev/null
+++ b/Engine/Plugins/Experimental/DynamicMaterial/Source/DynamicMaterialEditor/Private/Utils/DMMaterialModelFunctionLibrary.cpp
@@ -0,0 +1,531 @@
+// Copyright Epic Games, Inc. All Rights Reserved.
+
+#include "Utils/DMMaterialModelFunctionLibrary.h"
+
+#include "AssetRegistry/AssetRegistryModule.h"
+#include "AssetToolsModule.h"
+#include "Components/DMMaterialLayer.h"
+#include "Components/PrimitiveComponent.h"
+#include "ContentBrowserModule.h"
+#include "DMObjectMaterialProperty.h"
+#include "DMWorldSubsystem.h"
+#include "DynamicMaterialEditorModule.h"
+#include "Engine/Texture.h"
+#include "Engine/World.h"
+#include "EngineAnalytics.h"
+#include "GameFramework/Actor.h"
+#include "IAssetTools.h"
+#include "IContentBrowserSingleton.h"
+#include "Material/DynamicMaterialInstance.h"
+#include "Material/DynamicMaterialInstanceFactory.h"
+#include "Materials/Material.h"
+#include "Model/DynamicMaterialModel.h"
+#include "Model/DynamicMaterialModelDynamic.h"
+#include "Model/IDynamicMaterialModelEditorOnlyDataInterface.h"
+#include "UObject/Package.h"
+#include "UObject/UObjectGlobals.h"
+#include "Utils/DMPrivate.h"
+
+#define LOCTEXT_NAMESPACE "DMMaterialModelFunctionLibrary"
+
+TArray<FDMObjectMaterialProperty> UDMMaterialModelFunctionLibrary::GetActorMaterialProperties(AActor* InActor)
+{
+	TArray<FDMObjectMaterialProperty> ActorProperties;
+
+	if (!IsValid(InActor))
+	{
+		return ActorProperties;
+	}
+
+	FDMGetObjectMaterialPropertiesDelegate PropertyGenerator = FDynamicMaterialEditorModule::GetCustomMaterialPropertyGenerator(InActor->GetClass());
+
+	if (PropertyGenerator.IsBound())
+	{
+		ActorProperties = PropertyGenerator.Execute(InActor);
+
+		if (!ActorProperties.IsEmpty())
+		{
+			return ActorProperties;
+		}
+	}
+
+	InActor->ForEachComponent<UPrimitiveComponent>(false, [&ActorProperties](UPrimitiveComponent* InComp)
+		{
+			for (int32 MaterialIdx = 0; MaterialIdx < InComp->GetNumMaterials(); ++MaterialIdx)
+			{
+				ActorProperties.Add({InComp, MaterialIdx});
+			}
+		});
+
+	return ActorProperties;
+}
+
+UDynamicMaterialModel* UDMMaterialModelFunctionLibrary::CreateDynamicMaterialInObject(FDMObjectMaterialProperty& InMaterialProperty)
+{
+	if (!InMaterialProperty.IsValid())
+	{
+		return nullptr;
+	}
+
+	UObject* const Outer = InMaterialProperty.GetOuter();
+
+	UDynamicMaterialInstanceFactory* const InstanceFactory = NewObject<UDynamicMaterialInstanceFactory>();
+	check(InstanceFactory);
+
+	UDynamicMaterialInstance* const NewInstance = Cast<UDynamicMaterialInstance>(InstanceFactory->FactoryCreateNew(UDynamicMaterialInstance::StaticClass(),
+		Outer, NAME_None, RF_Transactional, nullptr, GWarn));
+	check(NewInstance);
+
+	bool bSubsystemTakenOver = false;
+
+	if (const UWorld* const World = Outer->GetWorld())
+	{
+		if (IsValid(World))
+		{
+			if (UDMWorldSubsystem* const WorldSubsystem = World->GetSubsystem<UDMWorldSubsystem>())
+			{
+				bSubsystemTakenOver = WorldSubsystem->ExecuteMaterialValueSetterDelegate(InMaterialProperty, NewInstance);
+			}
+		}
+	}
+
+	if (!bSubsystemTakenOver)
+	{
+		InMaterialProperty.SetMaterial(NewInstance);
+	}
+
+	if (UDynamicMaterialModel* MaterialModel = NewInstance->GetMaterialModel())
+	{
+		if (IDynamicMaterialModelEditorOnlyDataInterface* EditorOnlyData = MaterialModel->GetEditorOnlyData())
+		{
+			EditorOnlyData->RequestMaterialBuild();
+		}
+
+		return MaterialModel;
+	}
+
+	return nullptr;
+}
+
+UDynamicMaterialInstance* UDMMaterialModelFunctionLibrary::ExportMaterialInstance(UDynamicMaterialModelBase* InMaterialModelBase)
+{
+	if (!IsValid(InMaterialModelBase))
+	{
+		return nullptr;
+	}
+
+	UDynamicMaterialInstance* MaterialInstance = InMaterialModelBase->GetDynamicMaterialInstance();
+
+	if (!IsValid(MaterialInstance))
+	{
+		return nullptr;
+	}
+
+	FString CurrentName = MaterialInstance->GetName();
+
+	if (InMaterialModelBase->IsA<UDynamicMaterialModel>())
+	{
+		CurrentName = CurrentName.StartsWith(TEXT("MDI_"))
+			? CurrentName
+			: (TEXT("MDI_") + CurrentName);
+	}
+	else
+	{
+		CurrentName = CurrentName.StartsWith(TEXT("MDD_"))
+			? CurrentName
+			: (TEXT("MDD_") + CurrentName);
+	}
+
+	IAssetTools& AssetTools = FModuleManager::LoadModuleChecked<FAssetToolsModule>("AssetTools").Get();
+	FString PackageName, AssetName;
+	AssetTools.CreateUniqueAssetName(CurrentName, TEXT(""), PackageName, AssetName);
+
+	IContentBrowserSingleton& ContentBrowserModule = FModuleManager::LoadModuleChecked<FContentBrowserModule>("ContentBrowser").Get();
+	const FContentBrowserItemPath CurrentPath = ContentBrowserModule.GetCurrentPath();
+	const FString PathStr = CurrentPath.HasInternalPath() ? CurrentPath.GetInternalPathString() : TEXT("/Game");
+
+	FSaveAssetDialogConfig SaveAssetDialogConfig;
+	SaveAssetDialogConfig.DialogTitleOverride = LOCTEXT("SaveAssetDialogTitle", "Save Asset As");
+	SaveAssetDialogConfig.DefaultPath = PathStr;
+	SaveAssetDialogConfig.ExistingAssetPolicy = ESaveAssetDialogExistingAssetPolicy::Disallow;
+	SaveAssetDialogConfig.DefaultAssetName = AssetName;
+
+	const FString SaveObjectPath = ContentBrowserModule.Get().CreateModalSaveAssetDialog(SaveAssetDialogConfig);
+
+	if (SaveObjectPath.IsEmpty())
+	{
+		return nullptr;
+	}
+
+	return ExportMaterialInstance(MaterialInstance->GetMaterialModelBase(), SaveObjectPath);
+}
+
+UDynamicMaterialInstance* UDMMaterialModelFunctionLibrary::ExportMaterialInstance(UDynamicMaterialModelBase* InMaterialModel, const FString& InSavePath)
+{
+	if (!IsValid(InMaterialModel))
+	{
+		UE::DynamicMaterialEditor::Private::LogError(TEXT("Invalid material to export."));
+		return nullptr;
+	}
+
+	if (InSavePath.IsEmpty())
+	{
+		UE::DynamicMaterialEditor::Private::LogError(TEXT("Invalid material save path to export."));
+		return nullptr;
+	}
+
+	UDynamicMaterialInstance* Instance = InMaterialModel->GetDynamicMaterialInstance();
+	if (!Instance)
+	{
+		UE::DynamicMaterialEditor::Private::LogError(TEXT("Failed to find a Material Designer Instance to export."), true, InMaterialModel);
+		return nullptr;
+	}
+
+	const FString PackagePath = FPaths::GetBaseFilename(*InSavePath, false);
+
+	UPackage* Package = CreatePackage(*PackagePath);
+	if (!Package)
+	{
+		UE::DynamicMaterialEditor::Private::LogError(FString::Printf(TEXT("Failed to create package for Material Designer Instance (%s)."), *PackagePath));
+		return nullptr;
+	}
+
+	const FString AssetName = FPaths::GetBaseFilename(*InSavePath, true);
+
+	FObjectDuplicationParameters Params = InitStaticDuplicateObjectParams(Instance, Package, *AssetName,
+		RF_Public | RF_Standalone, nullptr, EDuplicateMode::Normal, EInternalObjectFlags::None);
+
+	UObject* NewAsset = StaticDuplicateObjectEx(Params);
+	if (!NewAsset)
+	{
+		UE::DynamicMaterialEditor::Private::LogError(TEXT("Failed to create new Material Designer Instance asset."));
+		return nullptr;
+	}
+
+	// Not sure why these need to be set again!
+	NewAsset->SetFlags(RF_Public | RF_Standalone);
+
+	UDynamicMaterialInstance* NewInstance = Cast<UDynamicMaterialInstance>(NewAsset);
+
+	if (NewInstance)
+	{
+		if (InMaterialModel->IsA<UDynamicMaterialModel>())
+		{
+			if (UDynamicMaterialModel* NewModel = NewInstance->GetMaterialModel())
+			{
+				if (IDynamicMaterialModelEditorOnlyDataInterface* ModelEditorOnlyData = NewModel->GetEditorOnlyData())
+				{
+					ModelEditorOnlyData->RequestMaterialBuild();
+				}
+			}
+		}
+		else if (InMaterialModel->IsA<UDynamicMaterialModelDynamic>())
+		{
+			NewInstance->InitializeMIDPublic();
+		}
+	}
+
+	FAssetRegistryModule::AssetCreated(NewAsset);
+
+	if (FEngineAnalytics::IsAvailable())
+	{
+		FEngineAnalytics::GetProvider().RecordEvent(TEXT("Editor.Usage.MaterialDesigner.ExportedMaterialInstance"));
+	}
+
+	return NewInstance;
+}
+
+UMaterial* UDMMaterialModelFunctionLibrary::ExportGeneratedMaterial(UDynamicMaterialModelBase* InMaterialModelBase)
+{
+	return nullptr;
+}
+
+UMaterial* UDMMaterialModelFunctionLibrary::ExportGeneratedMaterial(UDynamicMaterialModelBase* InMaterialModelBase, const FString& InSavePath)
+{
+	if (!IsValid(InMaterialModelBase))
+	{
+		UE::DynamicMaterialEditor::Private::LogError(TEXT("Invalid material to export."));
+		return nullptr;
+	}
+
+	if (InSavePath.IsEmpty())
+	{
+		UE::DynamicMaterialEditor::Private::LogError(TEXT("Invalid material save path to export."));
+		return nullptr;
+	}
+
+	UMaterial* GeneratedMaterial = InMaterialModelBase->GetGeneratedMaterial();
+
+	if (!GeneratedMaterial)
+	{
+		UE::DynamicMaterialEditor::Private::LogError(TEXT("Failed to find a generated material to export."));
+		return nullptr;
+	}
+
+	const FString PackagePath = FPaths::GetBaseFilename(*InSavePath, false);
+
+	UPackage* Package = CreatePackage(*PackagePath);
+	if (!Package)
+	{
+		UE::DynamicMaterialEditor::Private::LogError(FString::Printf(TEXT("Failed to create package for exported material (%s)."), *PackagePath));
+		return nullptr;
+	}
+
+	const FString AssetName = FPaths::GetBaseFilename(*InSavePath, true);
+
+	FObjectDuplicationParameters Params = InitStaticDuplicateObjectParams(GeneratedMaterial, Package, *AssetName,
+		RF_Public | RF_Standalone, nullptr, EDuplicateMode::Normal, EInternalObjectFlags::None);
+
+	UObject* NewAsset = StaticDuplicateObjectEx(Params);
+
+	if (!NewAsset)
+	{
+		UE::DynamicMaterialEditor::Private::LogError(TEXT("Failed to create new material asset."));
+		return nullptr;
+	}
+
+	// Not sure why these need to be set again!
+	NewAsset->SetFlags(RF_Public | RF_Standalone);
+
+	FAssetRegistryModule::AssetCreated(NewAsset);
+
+	if (FEngineAnalytics::IsAvailable())
+	{
+		FEngineAnalytics::GetProvider().RecordEvent(TEXT("Editor.Usage.MaterialDesigner.ExportedGeneratedMaterial"));
+	}
+
+	return Cast<UMaterial>(NewAsset);
+}
+
+UDynamicMaterialModel* UDMMaterialModelFunctionLibrary::ExportToTemplateMaterialModel(UDynamicMaterialModelDynamic* InMaterialModelDynamic)
+{
+	UDynamicMaterialModel* ParentModel = InMaterialModelDynamic->GetParentModel();
+
+	if (!ParentModel)
+	{
+		UE_LOG(LogDynamicMaterialEditor, Error, TEXT("Failed to find parent model."));
+		return nullptr;
+	}
+
+	// Where should we save it? (Always export to CB)
+	FString CurrentName = InMaterialModelDynamic->GetName();
+
+	CurrentName = CurrentName.StartsWith(TEXT("MDD_"))
+		? (TEXT("MDM_") + CurrentName.RightChop(4))
+		: (TEXT("MDM_") + CurrentName);
+
+	IAssetTools& AssetTools = FModuleManager::LoadModuleChecked<FAssetToolsModule>("AssetTools").Get();
+	FString PackageName, AssetName;
+	AssetTools.CreateUniqueAssetName(CurrentName, TEXT(""), PackageName, AssetName);
+
+	IContentBrowserSingleton& ContentBrowserModule = FModuleManager::LoadModuleChecked<FContentBrowserModule>("ContentBrowser").Get();
+	const FContentBrowserItemPath CurrentPath = ContentBrowserModule.GetCurrentPath();
+	const FString PathStr = CurrentPath.HasInternalPath() ? CurrentPath.GetInternalPathString() : TEXT("/Game");
+
+	FSaveAssetDialogConfig SaveAssetDialogConfig;
+	SaveAssetDialogConfig.DialogTitleOverride = LOCTEXT("SaveAssetDialogTitle", "Save Asset As");
+	SaveAssetDialogConfig.DefaultPath = PathStr;
+	SaveAssetDialogConfig.ExistingAssetPolicy = ESaveAssetDialogExistingAssetPolicy::Disallow;
+	SaveAssetDialogConfig.DefaultAssetName = AssetName;
+
+	const FString SaveObjectPath = ContentBrowserModule.Get().CreateModalSaveAssetDialog(SaveAssetDialogConfig);
+
+	if (SaveObjectPath.IsEmpty())
+	{
+		UE_LOG(LogDynamicMaterialEditor, Warning, TEXT("No path was chosen for saving the new editable asset, cancelling."));
+		return nullptr;
+	}
+
+	return ExportToTemplateMaterialModel(InMaterialModelDynamic, SaveObjectPath);
+}
+
+UDynamicMaterialModel* UDMMaterialModelFunctionLibrary::ExportToTemplateMaterialModel(UDynamicMaterialModelDynamic* InMaterialModelDynamic, const FString& InSavePath)
+{
+	UDynamicMaterialModel* ParentModel = InMaterialModelDynamic->GetParentModel();
+
+	if (!ParentModel)
+	{
+		UE_LOG(LogDynamicMaterialEditor, Error, TEXT("Failed to find parent model."));
+		return nullptr;
+	}
+
+	// Create new model
+	UDynamicMaterialModel* NewModel = InMaterialModelDynamic->ToEditable(GetTransientPackage());
+
+	if (!NewModel)
+	{
+		UE_LOG(LogDynamicMaterialEditor, Error, TEXT("Failed to convert dynamic asset to editable."));
+		return nullptr;
+	}
+
+	// Create a package for it
+	const FString PackageName = FPaths::GetBaseFilename(*InSavePath, false);
+	UPackage* Package = CreatePackage(*PackageName);
+
+	if (!Package)
+	{
+		UE_LOG(LogDynamicMaterialEditor, Error, TEXT("Failed to create new package for editable asset."));
+		return nullptr;
+	}
+
+	const FString AssetName = FPaths::GetBaseFilename(*InSavePath, true);
+
+	NewModel->Rename(nullptr, Package, UE::DynamicMaterial::RenameFlags);
+	NewModel->SetFlags(RF_Transactional | RF_Public | RF_Standalone);
+
+	FAssetRegistryModule::AssetCreated(NewModel);
+
+	if (FEngineAnalytics::IsAvailable())
+	{
+		FEngineAnalytics::GetProvider().RecordEvent(TEXT("Editor.Usage.MaterialDesigner.ExportToTemplateMaterialModel"));
+	}
+
+	return NewModel;
+}
+
+UDynamicMaterialInstance* UDMMaterialModelFunctionLibrary::ExportToTemplateMaterialInstance(UDynamicMaterialModelDynamic* InMaterialModelDynamic)
+{
+	UDynamicMaterialModel* ParentModel = InMaterialModelDynamic->GetParentModel();
+
+	if (!ParentModel)
+	{
+		UE_LOG(LogDynamicMaterialEditor, Error, TEXT("Failed to find parent model."));
+		return nullptr;
+	}
+
+	UDynamicMaterialInstance* OldInstance = InMaterialModelDynamic->GetDynamicMaterialInstance();
+
+	if (!OldInstance)
+	{
+		UE_LOG(LogDynamicMaterialEditor, Error, TEXT("Failed to find material instance."));
+		return nullptr;
+	}
+
+	// Where should we save it? (Always export to CB)
+	FString CurrentName = OldInstance ? OldInstance->GetName() : InMaterialModelDynamic->GetName();
+
+	CurrentName = CurrentName.StartsWith(TEXT("MDD_"))
+		? (TEXT("MDI_") + CurrentName.RightChop(4))
+		: (TEXT("MDI_") + CurrentName);
+
+	IAssetTools& AssetTools = FModuleManager::LoadModuleChecked<FAssetToolsModule>("AssetTools").Get();
+	FString PackageName, AssetName;
+	AssetTools.CreateUniqueAssetName(CurrentName, TEXT(""), PackageName, AssetName);
+
+	IContentBrowserSingleton& ContentBrowserModule = FModuleManager::LoadModuleChecked<FContentBrowserModule>("ContentBrowser").Get();
+	const FContentBrowserItemPath CurrentPath = ContentBrowserModule.GetCurrentPath();
+	const FString PathStr = CurrentPath.HasInternalPath() ? CurrentPath.GetInternalPathString() : TEXT("/Game");
+
+	FSaveAssetDialogConfig SaveAssetDialogConfig;
+	SaveAssetDialogConfig.DialogTitleOverride = LOCTEXT("SaveAssetDialogTitle", "Save Asset As");
+	SaveAssetDialogConfig.DefaultPath = PathStr;
+	SaveAssetDialogConfig.ExistingAssetPolicy = ESaveAssetDialogExistingAssetPolicy::Disallow;
+	SaveAssetDialogConfig.DefaultAssetName = AssetName;
+
+	const FString SaveObjectPath = ContentBrowserModule.Get().CreateModalSaveAssetDialog(SaveAssetDialogConfig);
+
+	if (SaveObjectPath.IsEmpty())
+	{
+		UE_LOG(LogDynamicMaterialEditor, Warning, TEXT("No path was chosen for saving the new editable asset, cancelling."));
+		return nullptr;
+	}
+
+	return ExportToTemplateMaterialInstance(InMaterialModelDynamic, SaveObjectPath);	
+}
+
+UDynamicMaterialInstance* UDMMaterialModelFunctionLibrary::ExportToTemplateMaterialInstance(UDynamicMaterialModelDynamic* InMaterialModelDynamic, const FString& InSavePath)
+{
+	UDynamicMaterialModel* ParentModel = InMaterialModelDynamic->GetParentModel();
+
+	if (!ParentModel)
+	{
+		UE_LOG(LogDynamicMaterialEditor, Error, TEXT("Failed to find parent model."));
+		return nullptr;
+	}
+
+	UDynamicMaterialInstance* OldInstance = InMaterialModelDynamic->GetDynamicMaterialInstance();
+
+	if (!OldInstance)
+	{
+		UE_LOG(LogDynamicMaterialEditor, Error, TEXT("Failed to find material instance."));
+		return nullptr;
+	}
+
+	// Create new model
+	UDynamicMaterialModel* NewModel = InMaterialModelDynamic->ToEditable(GetTransientPackage());
+
+	if (!NewModel)
+	{
+		UE_LOG(LogDynamicMaterialEditor, Error, TEXT("Failed to convert dynamic asset to editable."));
+		return nullptr;
+	}
+
+	// Create a package for it
+	const FString PackageName = FPaths::GetBaseFilename(*InSavePath, false);
+	UPackage* Package = CreatePackage(*PackageName);
+
+	if (!Package)
+	{
+		UE_LOG(LogDynamicMaterialEditor, Error, TEXT("Failed to create new package for editable asset."));
+		return nullptr;
+	}
+
+	const FString AssetName = FPaths::GetBaseFilename(*InSavePath, true);
+
+	UDynamicMaterialInstance* NewInstance = Cast<UDynamicMaterialInstance>(GetMutableDefault<UDynamicMaterialInstanceFactory>()->FactoryCreateNew(
+		UDynamicMaterialInstance::StaticClass(),
+		Package,
+		*AssetName,
+		RF_Transactional | RF_Public | RF_Standalone,
+		NewModel,
+		nullptr
+	));
+
+	FAssetRegistryModule::AssetCreated(NewInstance);
+
+	if (FEngineAnalytics::IsAvailable())
+	{
+		FEngineAnalytics::GetProvider().RecordEvent(TEXT("Editor.Usage.MaterialDesigner.ExportToTemplateMaterialInstance"));
+	}
+
+	return NewInstance;
+}
+
+bool UDMMaterialModelFunctionLibrary::IsModelValid(UDynamicMaterialModelBase* InMaterialModelBase)
+{
+	if (UWorld* World = InMaterialModelBase->GetWorld())
+	{
+		if (UDMWorldSubsystem* WorldSubsystem = World->GetSubsystem<UDMWorldSubsystem>())
+		{
+			if (!WorldSubsystem->ExecuteIsValidDelegate(InMaterialModelBase))
+			{
+				return false;
+			}
+		}
+	}
+
+	UActorComponent* ComponentOuter = InMaterialModelBase->GetTypedOuter<UActorComponent>();
+
+	if (ComponentOuter && !IsValid(ComponentOuter))
+	{
+		return false;
+	}
+
+	AActor* ActorOuter = InMaterialModelBase->GetTypedOuter<AActor>();
+
+	if (ActorOuter && !IsValid(ActorOuter))
+	{
+		return false;
+	}
+
+	UPackage* PackageOuter = InMaterialModelBase->GetPackage();
+
+	if (PackageOuter && !IsValid(PackageOuter))
+	{
+		return false;
+	}
+
+	return true;
+}
+
+#undef LOCTEXT_NAMESPACE
diff --git a/Engine/Plugins/Experimental/DynamicMaterial/Source/DynamicMaterialEditor/Private/Utils/DMMaterialStageFunctionLibrary.cpp b/Engine/Plugins/Experimental/DynamicMaterial/Source/DynamicMaterialEditor/Private/Utils/DMMaterialStageFunctionLibrary.cpp
new file mode 100644
--- /dev/null
+++ b/Engine/Plugins/Experimental/DynamicMaterial/Source/DynamicMaterialEditor/Private/Utils/DMMaterialStageFunctionLibrary.cpp
@@ -0,0 +1,96 @@
+// Copyright Epic Games, Inc. All Rights Reserved.
+
+#include "Utils/DMMaterialStageFunctionLibrary.h"
+
+#include "Components/DMMaterialStage.h"
+#include "Components/DMMaterialStageThroughput.h"
+#include "Components/DMMaterialSubStage.h"
+#include "Components/DMRenderTargetRenderer.h"
+#include "Components/MaterialStageExpressions/DMMSETextureSample.h"
+#include "Components/MaterialStageInputs/DMMSIExpression.h"
+#include "Components/MaterialStageInputs/DMMSIValue.h"
+#include "Components/MaterialValues/DMMaterialValueRenderTarget.h"
+#include "Components/MaterialValues/DMMaterialValueTexture.h"
+
+UDMMaterialStageInputValue* UDMMaterialStageFunctionLibrary::FindDefaultStageOpacityInputValue(UDMMaterialStage* InStage)
+{
+	if (!IsValid(InStage))
+	{
+		return nullptr;
+	}
+
+	int32 OpacityInputIndex = -1;
+
+	if (UDMMaterialStageThroughput* const ThroughputSource = Cast<UDMMaterialStageThroughput>(InStage->GetSource()))
+	{
+		const TArray<FDMMaterialStageConnector>& Connectors = ThroughputSource->GetInputConnectors();
+		for (const FDMMaterialStageConnector& Connector : Connectors)
+		{
+			if (Connector.Name.ToString() == TEXT("Opacity"))
+			{
+				OpacityInputIndex = Connector.Index;
+			}
+		}
+	}
+
+	const TArray<UDMMaterialStageInput*>& StageInputs = InStage->GetInputs();
+	const TArray<FDMMaterialStageConnection>& InputConnectionMap = InStage->GetInputConnectionMap();
+
+	if (InputConnectionMap.IsValidIndex(OpacityInputIndex) && !InputConnectionMap[OpacityInputIndex].Channels.IsEmpty())
+	{
+		if (InputConnectionMap[OpacityInputIndex].Channels[0].SourceIndex >= FDMMaterialStageConnectorChannel::FIRST_STAGE_INPUT)
+		{
+			const int32 StageInputIndex = InputConnectionMap[OpacityInputIndex].Channels[0].SourceIndex - FDMMaterialStageConnectorChannel::FIRST_STAGE_INPUT;
+
+			return Cast<UDMMaterialStageInputValue>(StageInputs[StageInputIndex]);
+		}
+	}
+
+	return nullptr;
+}
+
+void UDMMaterialStageFunctionLibrary::SetStageInputToRenderer(UDMMaterialStage* InStage, TSubclassOf<UDMRenderTargetRenderer> InRendererClass, int32 InInputIndex)
+{
+	UDMMaterialStageInputExpression* InputExpression = UDMMaterialStageInputExpression::ChangeStageInput_Expression(
+		InStage,
+		UDMMaterialStageExpressionTextureSample::StaticClass(),
+		InInputIndex,
+		FDMMaterialStageConnectorChannel::WHOLE_CHANNEL,
+		0,
+		FDMMaterialStageConnectorChannel::THREE_CHANNELS
+	);
+
+	if (!ensure(InputExpression))
+	{
+		return;
+	}
+
+	UDMMaterialSubStage* SubStage = InputExpression->GetSubStage();
+
+	if (!ensure(SubStage))
+	{
+		return;
+	}
+
+	UDMMaterialStageInputValue* InputValue = UDMMaterialStageInputValue::ChangeStageInput_NewLocalValue(
+		SubStage,
+		0,
+		FDMMaterialStageConnectorChannel::WHOLE_CHANNEL,
+		UDMMaterialValueRenderTarget::StaticClass(),
+		FDMMaterialStageConnectorChannel::WHOLE_CHANNEL
+	);
+
+	if (!ensure(InputValue))
+	{
+		return;
+	}
+
+	UDMMaterialValueRenderTarget* RenderTargetValue = Cast<UDMMaterialValueRenderTarget>(InputValue->GetValue());
+
+	if (!RenderTargetValue)
+	{
+		return;
+	}
+
+	UDMRenderTargetRenderer::CreateRenderTargetRenderer(InRendererClass, RenderTargetValue);
+}
diff --git a/Engine/Plugins/Experimental/DynamicMaterial/Source/DynamicMaterialEditor/Private/Utils/DMPrivate.cpp b/Engine/Plugins/Experimental/DynamicMaterial/Source/DynamicMaterialEditor/Private/Utils/DMPrivate.cpp
--- a/Engine/Plugins/Experimental/DynamicMaterial/Source/DynamicMaterialEditor/Private/Utils/DMPrivate.cpp
+++ b/Engine/Plugins/Experimental/DynamicMaterial/Source/DynamicMaterialEditor/Private/Utils/DMPrivate.cpp
@@ -1,8 +1,8 @@
 // Copyright Epic Games, Inc. All Rights Reserved.
 
 #include "Utils/DMPrivate.h"
+
 #include "Components/DMMaterialLayer.h"
-#include "Components/DMMaterialStage.h"
 #include "DMDefs.h"
 #include "DynamicMaterialEditorModule.h"
 #include "Framework/Notifications/NotificationManager.h"
@@ -66,6 +66,11 @@ namespace UE::DynamicMaterialEditor::Private
 		}
 	}
 
+	bool IsCustomMaterialProperty(EDMMaterialPropertyType InMaterialProperty)
+	{
+		return InMaterialProperty >= EDMMaterialPropertyType::Custom1 && InMaterialProperty <= EDMMaterialPropertyType::Custom4;
+	}
+
 	static bool bAllowUIFeedback = false;
 	static FText LogErrorObjectFormat = LOCTEXT("LogErrorObjectFormat", "%s (Source: %s)");
 
@@ -105,41 +110,11 @@ UDMMaterialLayerObject* FDMMaterialLayerReference::GetLayer() const
 	return LayerWeak.Get();
 }
 
-bool FDMMaterialLayerReference::IsBaseEnabled() const
-{
-	if (const UDMMaterialLayerObject* Layer = GetLayer())
-	{
-		return Layer->IsStageEnabled(EDMMaterialLayerStage::Base);
-	}
-
-	return false;
-}
-
-bool FDMMaterialLayerReference::IsBaseBeingEdited() const
-{
-	if (const UDMMaterialLayerObject* Layer = GetLayer())
-	{
-		return Layer->IsStageBeingEdited(EDMMaterialLayerStage::Base);
-	}
-
-	return false;
-}
-
-bool FDMMaterialLayerReference::IsMaskEnabled() const
-{
-	if (const UDMMaterialLayerObject* Layer = GetLayer())
-	{
-		return Layer->IsStageEnabled(EDMMaterialLayerStage::Mask);
-	}
-
-	return false;
-}
-
-bool FDMMaterialLayerReference::IsMaskBeingEdited() const
+bool FDMMaterialLayerReference::IsValid() const
 {
-	if (const UDMMaterialLayerObject* Layer = GetLayer())
+	if (UDMMaterialLayerObject* Layer = GetLayer())
 	{
-		return Layer->IsStageBeingEdited(EDMMaterialLayerStage::Mask);
+		return Layer->FindIndex() != INDEX_NONE;
 	}
 
 	return false;
diff --git a/Engine/Plugins/Experimental/DynamicMaterial/Source/DynamicMaterialEditor/Private/Utils/DMPrivate.h b/Engine/Plugins/Experimental/DynamicMaterial/Source/DynamicMaterialEditor/Private/Utils/DMPrivate.h
--- a/Engine/Plugins/Experimental/DynamicMaterial/Source/DynamicMaterialEditor/Private/Utils/DMPrivate.h
+++ b/Engine/Plugins/Experimental/DynamicMaterial/Source/DynamicMaterialEditor/Private/Utils/DMPrivate.h
@@ -5,13 +5,14 @@
 #include "Containers/Array.h"
 #include "HAL/Platform.h"
 #include "ScopedTransaction.h"
-#include "Slate/SDMSlot.h"
+#include "UI/Widgets/Editor/SDMMaterialSlotEditor.h"
 #include "UObject/WeakObjectPtr.h"
 
 class UDMMaterialLayerObject;
 class UDMMaterialSlot;
 class UDMMaterialStage;
 class UDMMaterialStageInput;
+enum class EDMMaterialPropertyType : uint8;
 struct FExpressionInput;
 struct FExpressionOutput;
 
@@ -42,10 +43,7 @@ struct FDMMaterialLayerReference
 
 	UDMMaterialLayerObject* GetLayer() const;
 
-	bool IsBaseEnabled() const;
-	bool IsBaseBeingEdited() const;
-	bool IsMaskEnabled() const;
-	bool IsMaskBeingEdited() const;
+	bool IsValid() const;
 };
 
 struct FDMScopedUITransaction
diff --git a/Engine/Plugins/Experimental/DynamicMaterial/Source/DynamicMaterialEditor/Public/Components/DMMaterialLayer.h b/Engine/Plugins/Experimental/DynamicMaterial/Source/DynamicMaterialEditor/Public/Components/DMMaterialLayer.h
--- a/Engine/Plugins/Experimental/DynamicMaterial/Source/DynamicMaterialEditor/Public/Components/DMMaterialLayer.h
+++ b/Engine/Plugins/Experimental/DynamicMaterial/Source/DynamicMaterialEditor/Public/Components/DMMaterialLayer.h
@@ -3,8 +3,10 @@
 #pragma once
 
 #include "Components/DMMaterialComponent.h"
-#include "DMDefs.h"
+
 #include "Components/DMMaterialStage.h"
+#include "DMDefs.h"
+
 #include "DMMaterialLayer.generated.h"
 
 class UDMMaterialEffectStack;
@@ -128,10 +130,6 @@ public:
 	UFUNCTION(BlueprintCallable, Category = "Material Designer")
 	DYNAMICMATERIALEDITOR_API bool AreAllStagesEnabled(EDMMaterialLayerStage InStageScope) const;
 
-	bool IsStageBeingEdited(EDMMaterialLayerStage InStageScope = EDMMaterialLayerStage::All) const;
-
-	UDMMaterialStage* GetFirstStageBeingEdited(EDMMaterialLayerStage InStageScope) const;
-
 	/** Iterate over all the valid stages, whether they are enabled or not. */
 	DYNAMICMATERIALEDITOR_API void ForEachValidStage(EDMMaterialLayerStage InStageScope, FStageCallbackFunc InCallback) const;
 
diff --git a/Engine/Plugins/Experimental/DynamicMaterial/Source/DynamicMaterialEditor/Public/Components/DMMaterialProperty.h b/Engine/Plugins/Experimental/DynamicMaterial/Source/DynamicMaterialEditor/Public/Components/DMMaterialProperty.h
--- a/Engine/Plugins/Experimental/DynamicMaterial/Source/DynamicMaterialEditor/Public/Components/DMMaterialProperty.h
+++ b/Engine/Plugins/Experimental/DynamicMaterial/Source/DynamicMaterialEditor/Public/Components/DMMaterialProperty.h
@@ -2,10 +2,12 @@
 
 #pragma once
 
-#include "DMEDefs.h"
 #include "Components/DMMaterialComponent.h"
+
+#include "DMEDefs.h"
 #include "Math/MathFwd.h"
 #include "Templates/SharedPointerFwd.h"
+
 #include "DMMaterialProperty.generated.h"
 
 class FText;
@@ -47,6 +49,12 @@ public:
 	UFUNCTION(BlueprintPure, Category = "Material Designer")
 	EDMMaterialPropertyType GetMaterialProperty() const { return MaterialProperty; }
 
+	UFUNCTION(BlueprintPure, Category = "Material Designer")
+	bool IsEnabled() const { return bEnabled; }
+
+	UFUNCTION(BlueprintCallable, Category = "Material Designer")
+	void SetEnabled(bool bInEnabled);
+
 	/**
 	 * The description of this property based on the EDMMaterialPropertyType enum.
 	 */
@@ -210,6 +218,12 @@ protected:
 	UPROPERTY(VisibleInstanceOnly, BlueprintReadOnly, Category = "Material Designer")
 	EDMMaterialPropertyType MaterialProperty;
 
+	/**
+	 * Whether this property is enabled. If it is not enabled, it will generate no expressions.
+	 */
+	UPROPERTY(VisibleInstanceOnly, BlueprintReadOnly, Category = "Material Designer")
+	bool bEnabled;
+
 	/**
 	 * The value type used to connect to this property. Will be either VT_Float1, VT_Float3_RGB or VT_Float3_XYZ.
 	 */
diff --git a/Engine/Plugins/Experimental/DynamicMaterial/Source/DynamicMaterialEditor/Public/Components/DMMaterialSlot.h b/Engine/Plugins/Experimental/DynamicMaterial/Source/DynamicMaterialEditor/Public/Components/DMMaterialSlot.h
--- a/Engine/Plugins/Experimental/DynamicMaterial/Source/DynamicMaterialEditor/Public/Components/DMMaterialSlot.h
+++ b/Engine/Plugins/Experimental/DynamicMaterial/Source/DynamicMaterialEditor/Public/Components/DMMaterialSlot.h
@@ -3,7 +3,9 @@
 #pragma once
 
 #include "Components/DMMaterialComponent.h"
+
 #include "DMEDefs.h"
+
 #include "DMMaterialSlot.generated.h"
 
 class FScopedTransaction;
@@ -26,7 +28,7 @@ UCLASS(MinimalAPI, BlueprintType, ClassGroup = "Material Designer", Meta = (Disp
 class UDMMaterialSlot : public UDMMaterialComponent
 {
 	friend class SDMMaterialSlot;
-	friend class SDMSlot;
+	friend class SDMMaterialSlotEditor;
 
 	GENERATED_BODY()
 
@@ -84,7 +86,7 @@ public:
 
 	/** Can't be removed if it is the last remaining layer. */
 	UFUNCTION(BlueprintCallable, Category = "Material Designer")
-	DYNAMICMATERIALEDITOR_API bool CanRemoveLayer(const UDMMaterialLayerObject* InLayer) const;
+	DYNAMICMATERIALEDITOR_API bool CanRemoveLayer(UDMMaterialLayerObject* InLayer) const;
 
 	/** Removes the layer, if possible. */
 	UFUNCTION(BlueprintCallable, Category = "Material Designer")
@@ -134,10 +136,6 @@ public:
 	/** Returns true if all associations have been removed */
 	bool UnreferencedBySlot(UDMMaterialSlot* InOtherSlot);
 
-	bool IsEditingLayers() const { return bIsEditingLayers; }
-
-	void SetEditingLayers(bool bInIsEditing) { bIsEditingLayers = bInIsEditing; }
-
 	UFUNCTION(BlueprintCallable, Category = "Material Designer")
 	DYNAMICMATERIALEDITOR_API void GeneralPreviewMaterial(UMaterial* InPreviewMaterial, EDMMaterialLayerStage InLayerStage);
 
@@ -176,9 +174,6 @@ protected:
 	UPROPERTY()
 	TMap<TWeakObjectPtr<UDMMaterialSlot>, int32> SlotsReferencedBy;
 
-	UPROPERTY(VisibleInstanceOnly, BlueprintReadOnly, Transient, TextExportTransient, Category = "Material Designer")
-	bool bIsEditingLayers;
-
 	FDMOnMaterialSlotConnectorsUpdated OnConnectorsUpdateDelegate;
 	FDMOnMaterialSlotPropertiesUpdated OnPropertiesUpdateDelegate;
 	FDMOnMaterialSlotLayersUpdated OnLayersUpdateDelegate;
diff --git a/Engine/Plugins/Experimental/DynamicMaterial/Source/DynamicMaterialEditor/Public/Components/DMMaterialStage.h b/Engine/Plugins/Experimental/DynamicMaterial/Source/DynamicMaterialEditor/Public/Components/DMMaterialStage.h
--- a/Engine/Plugins/Experimental/DynamicMaterial/Source/DynamicMaterialEditor/Public/Components/DMMaterialStage.h
+++ b/Engine/Plugins/Experimental/DynamicMaterial/Source/DynamicMaterialEditor/Public/Components/DMMaterialStage.h
@@ -3,9 +3,11 @@
 #pragma once
 
 #include "Components/DMMaterialComponent.h"
+
 #include "Components/DMMaterialStageSource.h"
 #include "DMEDefs.h"
 #include "Templates/SubclassOf.h"
+
 #include "DMMaterialStage.generated.h"
 
 class FAssetThumbnailPool;
@@ -122,12 +124,6 @@ public:
 
 	DYNAMICMATERIALEDITOR_API void GenerateExpressions(const TSharedRef<FDMMaterialBuildState>& InBuildState) const;
 
-	UFUNCTION(BlueprintPure, Category = "Material Designer")
-	bool IsBeingEdited() const { return bIsBeingEdited; }
-
-	UFUNCTION(BlueprintCallable, Category = "Material Designer")
-	bool SetBeingEdited(bool bInBeingEdited);
-
 	/** Get the last layer for each property type from the previous stages. */
 	DYNAMICMATERIALEDITOR_API TMap<EDMMaterialPropertyType, UDMMaterialLayerObject*> GetPreviousStagesPropertyMap();
 
@@ -265,9 +261,6 @@ protected:
 	UPROPERTY(VisibleInstanceOnly, BlueprintReadOnly, Category = "Material Designer")
 	bool bCanChangeSource;
 
-	UPROPERTY(VisibleInstanceOnly, BlueprintReadOnly, Transient, TextExportTransient, Category = "Material Designer")
-	bool bIsBeingEdited;
-
 	//~ Begin UDMMaterialComponent
 	DYNAMICMATERIALEDITOR_API virtual void OnComponentAdded() override;
 	DYNAMICMATERIALEDITOR_API virtual void OnComponentRemoved() override;
diff --git a/Engine/Plugins/Experimental/DynamicMaterial/Source/DynamicMaterialEditor/Public/Components/MaterialStageExpressions/DMMSETextureSample.h b/Engine/Plugins/Experimental/DynamicMaterial/Source/DynamicMaterialEditor/Public/Components/MaterialStageExpressions/DMMSETextureSample.h
--- a/Engine/Plugins/Experimental/DynamicMaterial/Source/DynamicMaterialEditor/Public/Components/MaterialStageExpressions/DMMSETextureSample.h
+++ b/Engine/Plugins/Experimental/DynamicMaterial/Source/DynamicMaterialEditor/Public/Components/MaterialStageExpressions/DMMSETextureSample.h
@@ -3,6 +3,7 @@
 #pragma once
 
 #include "Components/MaterialStageExpressions/DMMSETextureSampleBase.h"
+
 #include "DMMSETextureSample.generated.h"
 
 UCLASS(MinimalAPI, BlueprintType, ClassGroup = "Material Designer")
@@ -10,7 +11,7 @@ class UDMMaterialStageExpressionTextureSample : public UDMMaterialStageExpressio
 {
 	GENERATED_BODY()
 
-	friend class SDMComponentEdit;
+	friend class SDMMaterialComponentEditor;
 
 public:
 	UDMMaterialStageExpressionTextureSample();
diff --git a/Engine/Plugins/Experimental/DynamicMaterial/Source/DynamicMaterialEditor/Public/DMEDefs.h b/Engine/Plugins/Experimental/DynamicMaterial/Source/DynamicMaterialEditor/Public/DMEDefs.h
--- a/Engine/Plugins/Experimental/DynamicMaterial/Source/DynamicMaterialEditor/Public/DMEDefs.h
+++ b/Engine/Plugins/Experimental/DynamicMaterial/Source/DynamicMaterialEditor/Public/DMEDefs.h
@@ -2,19 +2,21 @@
 
 #pragma once
 
+#include "UObject/Interface.h"
+
 #include "DMDefs.h"
 #include "DMValueDefinition.h"
 #include "IDetailPropertyRow.h"
 #include "Misc/Optional.h"
 #include "Templates/SharedPointer.h"
-#include "UObject/Interface.h"
+
 #include "DMEDefs.generated.h"
 
 class FAssetThumbnailPool;
 class FProperty;
 class IDetailTreeNode;
+class SDMMaterialComponentEditor;
 class SDMMaterialStageEdit;
-class SDMComponentEdit;
 class SVerticalBox;
 class SWidget;
 class UDMMaterialStageInput;
@@ -118,6 +120,13 @@ struct FDMMaterialSlotOutputConnectorTypes
 	TArray<EDMValueType> ConnectorTypes;
 };
 
+enum class EDMPropertyHandlePriority : uint8
+{
+	Low,
+	Normal,
+	High
+};
+
 struct FDMPropertyHandle
 {
 	TSharedPtr<IPropertyRowGenerator> PropertyRowGenerator;
@@ -131,6 +140,8 @@ struct FDMPropertyHandle
 	FName CategoryOverrideName = NAME_None;
 	TOptional<float> MaxWidth;
 	bool bEnabled = true;
+	EDMPropertyHandlePriority Priority = EDMPropertyHandlePriority::Normal;
+	bool bKeyframeable = true;
 };
 
 UENUM(BlueprintType, meta = (Bitflags, UseEnumValuesAsMaskValuesInEditor = "true"))
diff --git a/Engine/Plugins/Experimental/DynamicMaterial/Source/DynamicMaterialEditor/Public/DMObjectMaterialProperty.h b/Engine/Plugins/Experimental/DynamicMaterial/Source/DynamicMaterialEditor/Public/DMObjectMaterialProperty.h
--- a/Engine/Plugins/Experimental/DynamicMaterial/Source/DynamicMaterialEditor/Public/DMObjectMaterialProperty.h
+++ b/Engine/Plugins/Experimental/DynamicMaterial/Source/DynamicMaterialEditor/Public/DMObjectMaterialProperty.h
@@ -2,10 +2,10 @@
 
 #pragma once
 
-#include "UObject/WeakObjectPtr.h"
-#include "UObject/WeakObjectPtrTemplates.h"
 #include "UObject/NameTypes.h"
 #include "UObject/Object.h"
+#include "UObject/WeakObjectPtr.h"
+
 #include "DMObjectMaterialProperty.generated.h"
 
 class FProperty;
@@ -18,7 +18,7 @@ class UPrimitiveComponent;
 /**
  * Defines a material property slot that can be a Material Designer Instance.
  */
-USTRUCT(BlueprintType)
+USTRUCT()
 struct FDMObjectMaterialProperty
 {
 	GENERATED_BODY()
@@ -31,19 +31,11 @@ struct FDMObjectMaterialProperty
 	/** Class Property (including potential array index) */
 	DYNAMICMATERIALEDITOR_API FDMObjectMaterialProperty(UObject* InOuter, FProperty* InProperty, int32 InIndex = INDEX_NONE);
 
-	UPROPERTY(VisibleInstanceOnly, BlueprintReadOnly, Category = "Material")
-	TWeakObjectPtr<UObject> OuterWeak;
+	DYNAMICMATERIALEDITOR_API UObject* GetOuter() const;
 
-	/** C++ version of property */
-	FProperty* Property;
+	DYNAMICMATERIALEDITOR_API FProperty* GetProperty() const;
 
-	/** Blueprint version of property */
-	UPROPERTY(VisibleInstanceOnly, BlueprintReadOnly, Category = "Material")
-	FName PropertyName;
-
-	/** Component or array property index. */
-	UPROPERTY(VisibleInstanceOnly, BlueprintReadOnly, Category = "Material")
-	int32 Index;
+	DYNAMICMATERIALEDITOR_API int32 GetIndex() const;
 
 	DYNAMICMATERIALEDITOR_API UDynamicMaterialModelBase* GetMaterialModelBase() const;
 
@@ -57,6 +49,12 @@ struct FDMObjectMaterialProperty
 
 	DYNAMICMATERIALEDITOR_API void Reset();
 
+	/** Whether this is a property on the object. */
+	DYNAMICMATERIALEDITOR_API bool IsProperty() const;
+
+	/** Whether this is an element of a primitive component's material override array. */
+	DYNAMICMATERIALEDITOR_API bool IsElement() const;
+
 	template<typename InClass>
 	InClass* GetTypedOuter() const
 	{
@@ -72,4 +70,14 @@ struct FDMObjectMaterialProperty
 
 		return nullptr;
 	}
+
+protected:
+	UPROPERTY()
+	TWeakObjectPtr<UObject> OuterWeak = nullptr;
+
+	/** C++ version of property */
+	FProperty* Property = nullptr;
+
+	/** Component or array property index. */
+	int32 Index = INDEX_NONE;
 };
diff --git a/Engine/Plugins/Experimental/DynamicMaterial/Source/DynamicMaterialEditor/Public/Model/DynamicMaterialModelEditorOnlyData.h b/Engine/Plugins/Experimental/DynamicMaterial/Source/DynamicMaterialEditor/Public/Model/DynamicMaterialModelEditorOnlyData.h
--- a/Engine/Plugins/Experimental/DynamicMaterial/Source/DynamicMaterialEditor/Public/Model/DynamicMaterialModelEditorOnlyData.h
+++ b/Engine/Plugins/Experimental/DynamicMaterial/Source/DynamicMaterialEditor/Public/Model/DynamicMaterialModelEditorOnlyData.h
@@ -9,8 +9,8 @@
 
 #include "Engine/EngineTypes.h"
 #include "MaterialDomain.h"
+#include "MaterialEditingLibrary.h"
 #include "UObject/WeakObjectPtrFwd.h"
-#include "UObject/WeakObjectPtrTemplatesFwd.h"
 
 #include "DynamicMaterialModelEditorOnlyData.generated.h"
 
@@ -48,7 +48,7 @@ class UDynamicMaterialModelEditorOnlyData : public UObject, public IDynamicMater
 	GENERATED_BODY()
 
 	friend class UDynamicMaterialModelFactory;
-	friend class SDMComponentEdit;
+	friend class FDMMaterialModelPropertyRowGenerator;
 
 public:
 	DYNAMICMATERIALEDITOR_API static const FString SlotsPathToken;
@@ -151,6 +151,9 @@ public:
 	UFUNCTION(BlueprintCallable, Category = "Material Designer")
 	DYNAMICMATERIALEDITOR_API void SetChannelListPreset(FName InPresetName);
 
+	UFUNCTION(BlueprintPure, Category = "Material Designer")
+	const FMaterialStatistics& GetMaterialStats() const;
+
 	UFUNCTION(BlueprintCallable, Category = "Material Designer")
 	DYNAMICMATERIALEDITOR_API void OpenMaterialEditor() const;
 
@@ -226,7 +229,7 @@ public:
 	DYNAMICMATERIALEDITOR_API virtual void PostLoad() override;
 	DYNAMICMATERIALEDITOR_API virtual void PostEditUndo() override;
 	DYNAMICMATERIALEDITOR_API virtual void PostEditImport() override;
-	DYNAMICMATERIALEDITOR_API virtual void PostDuplicate(bool bDuplicateForPIE) override;
+	DYNAMICMATERIALEDITOR_API virtual void PostDuplicate(bool bInDuplicateForPIE) override;
 	DYNAMICMATERIALEDITOR_API virtual void PostEditChangeChainProperty(FPropertyChangedChainEvent& PropertyChangedEvent) override;
 	DYNAMICMATERIALEDITOR_API virtual void Serialize(FArchive& Ar) override;
 	//~ End UObject
@@ -312,6 +315,9 @@ protected:
 	UPROPERTY(VisibleInstanceOnly, BlueprintReadOnly, Category = "Material Designer")
 	bool bCreateMaterialPackage;
 
+	UPROPERTY(VisibleInstanceOnly, BlueprintReadOnly, Category = "Material Designer")
+	FMaterialStatistics MaterialStats;
+
 	FDMOnMaterialBuilt OnMaterialBuiltDelegate;
 	FDMOnValueListUpdated OnValueListUpdateDelegate;
 	FDMOnSlotListUpdated OnSlotListUpdateDelegate;
diff --git a/Engine/Plugins/Experimental/DynamicMaterial/Source/DynamicMaterialEditor/Public/UI/PropertyGenerators/DMComponentPropertyRowGenerator.h b/Engine/Plugins/Experimental/DynamicMaterial/Source/DynamicMaterialEditor/Public/UI/PropertyGenerators/DMComponentPropertyRowGenerator.h
new file mode 100644
--- /dev/null
+++ b/Engine/Plugins/Experimental/DynamicMaterial/Source/DynamicMaterialEditor/Public/UI/PropertyGenerators/DMComponentPropertyRowGenerator.h
@@ -0,0 +1,55 @@
+// Copyright Epic Games, Inc. All Rights Reserved.
+
+#pragma once
+
+#include "Containers/ContainersFwd.h"
+#include "Templates/SharedPointerFwd.h"
+
+enum class EDMPropertyHandlePriority : uint8;
+class FName;
+class FProperty;
+class SDMMaterialComponentEditor;
+class UDMMaterialComponent;
+struct FDMPropertyHandle;
+
+/** Used to generate editable properties for objects editable in the Material Designer. */
+class FDMComponentPropertyRowGenerator
+{
+public:
+	DYNAMICMATERIALEDITOR_API static const TSharedRef<FDMComponentPropertyRowGenerator>& Get();
+
+	virtual ~FDMComponentPropertyRowGenerator() = default;
+
+	/**
+	 * Generate properties for the given component.
+	 * @param InComponentEditorWidget The edit widget generating the properties.
+	 * @param InComponent The component being edited.
+	 * @param InOutPropertyRows The generated rows.
+	 * @param InOutProcessedObjects The already processed objects - add to this to avoid possible recursive generation.
+	 */
+	DYNAMICMATERIALEDITOR_API virtual void AddComponentProperties(const TSharedRef<SDMMaterialComponentEditor>& InComponentEditorWidget, 
+		UDMMaterialComponent* InComponent, TArray<FDMPropertyHandle>& InOutPropertyRows, TSet<UDMMaterialComponent*>& InOutProcessedObjects);
+
+	/**
+	 * Add the rows needed for a specific property.
+	 * @param InProperty The name of the property.
+	 */
+	DYNAMICMATERIALEDITOR_API virtual void AddPropertyEditRows(const TSharedRef<SDMMaterialComponentEditor>& InComponentEditorWidget, 
+		UDMMaterialComponent* InComponent, const FName& InProperty, TArray<FDMPropertyHandle>& InOutPropertyRows, 
+		TSet<UDMMaterialComponent*>& InOutProcessedObjects);
+
+	/**
+	 * Adds the rows needed for a specific property.
+	 * @param InProperty The FProperty for this property.
+	 * @param MemoryPtr The pointer to the property's value in memory
+	 */
+	DYNAMICMATERIALEDITOR_API virtual void AddPropertyEditRows(const TSharedRef<SDMMaterialComponentEditor>& InComponentEditorWidget, 
+		UDMMaterialComponent* InComponent, FProperty* InProperty, void* MemoryPtr, TArray<FDMPropertyHandle>& InOutPropertyRows, 
+		TSet<UDMMaterialComponent*>& InOutProcessedObjects);
+
+	DYNAMICMATERIALEDITOR_API virtual bool AllowKeyframeButton(UDMMaterialComponent* InComponent, FProperty* InProperty);
+
+protected:
+	/** Returns true if the component edit widget is editing a Material Designer Dynamic. */
+	DYNAMICMATERIALEDITOR_API static bool IsDynamic(const TSharedRef<SDMMaterialComponentEditor>& InComponentEditorWidget);
+};
