Material Designer

diff --git a/Engine/Plugins/Experimental/DynamicMaterial/Source/DynamicMaterialEditor/Public/Utils/DMMaterialModelFunctionLibrary.h b/Engine/Plugins/Experimental/DynamicMaterial/Source/DynamicMaterialEditor/Public/Utils/DMMaterialModelFunctionLibrary.h
new file mode 100644
--- /dev/null
+++ b/Engine/Plugins/Experimental/DynamicMaterial/Source/DynamicMaterialEditor/Public/Utils/DMMaterialModelFunctionLibrary.h
@@ -0,0 +1,48 @@
+// Copyright Epic Games, Inc. All Rights Reserved.
+
+#pragma once
+
+#include "Kismet/BlueprintFunctionLibrary.h"
+
+#include "DMEDefs.h"
+
+#include "DMMaterialModelFunctionLibrary.generated.h"
+
+class AActor;
+class FString;
+class UDynamicMaterialModelBase;
+class UDynamicMaterialModelDynamic;
+class UMaterial;
+struct FDMObjectMaterialProperty;
+
+/**
+ * Material Model / Instance Function Library
+ */
+UCLASS()
+class UDMMaterialModelFunctionLibrary : public UBlueprintFunctionLibrary
+{
+	GENERATED_BODY()
+
+public:
+	DYNAMICMATERIALEDITOR_API static TArray<FDMObjectMaterialProperty> GetActorMaterialProperties(AActor* InActor);
+
+	DYNAMICMATERIALEDITOR_API static UDynamicMaterialModel* CreateDynamicMaterialInObject(FDMObjectMaterialProperty& InMaterialProperty);
+
+	DYNAMICMATERIALEDITOR_API static UDynamicMaterialInstance* ExportMaterialInstance(UDynamicMaterialModelBase* InMaterialModelBase);
+
+	DYNAMICMATERIALEDITOR_API static UDynamicMaterialInstance* ExportMaterialInstance(UDynamicMaterialModelBase* InMaterialModel, const FString& InSavePath);
+
+	DYNAMICMATERIALEDITOR_API static UMaterial* ExportGeneratedMaterial(UDynamicMaterialModelBase* InMaterialModelBase);
+
+	DYNAMICMATERIALEDITOR_API static UMaterial* ExportGeneratedMaterial(UDynamicMaterialModelBase* InMaterialModelBase, const FString& InSavePath);
+
+	DYNAMICMATERIALEDITOR_API static UDynamicMaterialModel* ExportToTemplateMaterialModel(UDynamicMaterialModelDynamic* InMaterialModelDynamic);
+
+	DYNAMICMATERIALEDITOR_API static UDynamicMaterialModel* ExportToTemplateMaterialModel(UDynamicMaterialModelDynamic* InMaterialModelDynamic, const FString& InSavePath);
+
+	DYNAMICMATERIALEDITOR_API static UDynamicMaterialInstance* ExportToTemplateMaterialInstance(UDynamicMaterialModelDynamic* InMaterialModelDynamic);
+
+	DYNAMICMATERIALEDITOR_API static UDynamicMaterialInstance* ExportToTemplateMaterialInstance(UDynamicMaterialModelDynamic* InMaterialModelDynamic, const FString& InSavePath);
+
+	DYNAMICMATERIALEDITOR_API static bool IsModelValid(UDynamicMaterialModelBase* InMaterialModelBase);
+};
