Material Designer
- Horde issue 533743

diff --git a/Engine/Plugins/Experimental/DynamicMaterial/Source/DynamicMaterialEditor/Private/UI/PropertyGenerators/DMMaterialStageFunctionPropertyRowGenerator.cpp b/Engine/Plugins/Experimental/DynamicMaterial/Source/DynamicMaterialEditor/Private/UI/PropertyGenerators/DMMaterialStageFunctionPropertyRowGenerator.cpp
--- a/Engine/Plugins/Experimental/DynamicMaterial/Source/DynamicMaterialEditor/Private/UI/PropertyGenerators/DMMaterialStageFunctionPropertyRowGenerator.cpp
+++ b/Engine/Plugins/Experimental/DynamicMaterial/Source/DynamicMaterialEditor/Private/UI/PropertyGenerators/DMMaterialStageFunctionPropertyRowGenerator.cpp
@@ -60,7 +60,7 @@ void FDMMaterialStageFunctionPropertyRowGenerator::ApplyMetaData(const FFunction
 
 		for (int32 Index = 0; Index < MetaDataValue.Len(); ++Index)
 		{
-			if (MetaDataValue[Index] != '.'	&& MetaDataValue[Index] != '-' && MetaDataValue[Index] < '0' && MetaDataValue[Index] > '9')
+			if (MetaDataValue[Index] != '.'	&& MetaDataValue[Index] != '-' && (MetaDataValue[Index] < '0' || MetaDataValue[Index] > '9'))
 			{
 				bValidValue = false;
 				break;
