Material Designer
- Fixed UDMMaterialModelFunctionLibrary::IsModelValid not checking if the model pointer was actually valid (only checking its properties / context).

diff --git a/Engine/Plugins/Experimental/DynamicMaterial/Source/DynamicMaterialEditor/Private/Utils/DMMaterialModelFunctionLibrary.cpp b/Engine/Plugins/Experimental/DynamicMaterial/Source/DynamicMaterialEditor/Private/Utils/DMMaterialModelFunctionLibrary.cpp
--- a/Engine/Plugins/Experimental/DynamicMaterial/Source/DynamicMaterialEditor/Private/Utils/DMMaterialModelFunctionLibrary.cpp
+++ b/Engine/Plugins/Experimental/DynamicMaterial/Source/DynamicMaterialEditor/Private/Utils/DMMaterialModelFunctionLibrary.cpp
@@ -493,6 +493,11 @@ UDynamicMaterialInstance* UDMMaterialModelFunctionLibrary::ExportToTemplateMater
 
 bool UDMMaterialModelFunctionLibrary::IsModelValid(UDynamicMaterialModelBase* InMaterialModelBase)
 {
+	if (!IsValid(InMaterialModelBase))
+	{
+		return false;
+	}
+
 	if (UWorld* World = InMaterialModelBase->GetWorld())
 	{
 		if (UDMWorldSubsystem* WorldSubsystem = World->GetSubsystem<UDMWorldSubsystem>())
