Material Designer
- Reworked how the material editor widget invalidated itself (now through the designer tick, rather than its own) to prevent slate crash.
- Renamed TDMWidgetSlots to be called ...Slot for more clarity.
- Made sure that widgets that don't need to tick do not tick.

diff --git a/Engine/Plugins/Experimental/DynamicMaterial/Source/DynamicMaterialEditor/Private/UI/Widgets/Editor/SDMMaterialComponentEditor.cpp b/Engine/Plugins/Experimental/DynamicMaterial/Source/DynamicMaterialEditor/Private/UI/Widgets/Editor/SDMMaterialComponentEditor.cpp
--- a/Engine/Plugins/Experimental/DynamicMaterial/Source/DynamicMaterialEditor/Private/UI/Widgets/Editor/SDMMaterialComponentEditor.cpp
+++ b/Engine/Plugins/Experimental/DynamicMaterial/Source/DynamicMaterialEditor/Private/UI/Widgets/Editor/SDMMaterialComponentEditor.cpp
@@ -42,6 +42,8 @@ SDMMaterialComponentEditor::~SDMMaterialComponentEditor()
 void SDMMaterialComponentEditor::Construct(const FArguments& InArgs, const TSharedRef<SDMMaterialEditor>& InEditorWidget, 
 	UDMMaterialComponent* InMaterialComponent)
 {
+	SetCanTick(false);
+
 	SDMObjectEditorWidgetBase::Construct(
 		SDMObjectEditorWidgetBase::FArguments(), 
 		InEditorWidget, 
diff --git a/Engine/Plugins/Experimental/DynamicMaterial/Source/DynamicMaterialEditor/Private/UI/Widgets/Editor/SDMMaterialGlobalSettingsEditor.cpp b/Engine/Plugins/Experimental/DynamicMaterial/Source/DynamicMaterialEditor/Private/UI/Widgets/Editor/SDMMaterialGlobalSettingsEditor.cpp
--- a/Engine/Plugins/Experimental/DynamicMaterial/Source/DynamicMaterialEditor/Private/UI/Widgets/Editor/SDMMaterialGlobalSettingsEditor.cpp
+++ b/Engine/Plugins/Experimental/DynamicMaterial/Source/DynamicMaterialEditor/Private/UI/Widgets/Editor/SDMMaterialGlobalSettingsEditor.cpp
@@ -13,6 +13,8 @@ void SDMMaterialGlobalSettingsEditor::PrivateRegisterAttributes(FSlateAttributeD
 
 void SDMMaterialGlobalSettingsEditor::Construct(const FArguments& InArgs, const TSharedRef<SDMMaterialEditor>& InEditorWidget, UDynamicMaterialModelBase* InMaterialModelBase)
 {
+	SetCanTick(false);
+
 	SDMObjectEditorWidgetBase::Construct(
 		SDMObjectEditorWidgetBase::FArguments(), 
 		InEditorWidget, 
diff --git a/Engine/Plugins/Experimental/DynamicMaterial/Source/DynamicMaterialEditor/Private/UI/Widgets/Editor/SDMMaterialPreview.cpp b/Engine/Plugins/Experimental/DynamicMaterial/Source/DynamicMaterialEditor/Private/UI/Widgets/Editor/SDMMaterialPreview.cpp
--- a/Engine/Plugins/Experimental/DynamicMaterial/Source/DynamicMaterialEditor/Private/UI/Widgets/Editor/SDMMaterialPreview.cpp
+++ b/Engine/Plugins/Experimental/DynamicMaterial/Source/DynamicMaterialEditor/Private/UI/Widgets/Editor/SDMMaterialPreview.cpp
@@ -87,6 +87,8 @@ void SDMMaterialPreview::Construct(const FArguments& InArgs, const TSharedRef<SD
 	PreviewMeshComponent = nullptr;
 	PostProcessVolumeActor = nullptr;
 
+	SetCanTick(false);
+
 	SEditorViewport::Construct(SEditorViewport::FArguments().ViewportSize(FVector2D(135.f)));
 
 	// restore last used feature level
diff --git a/Engine/Plugins/Experimental/DynamicMaterial/Source/DynamicMaterialEditor/Private/UI/Widgets/Editor/SDMMaterialPropertySelector.cpp b/Engine/Plugins/Experimental/DynamicMaterial/Source/DynamicMaterialEditor/Private/UI/Widgets/Editor/SDMMaterialPropertySelector.cpp
--- a/Engine/Plugins/Experimental/DynamicMaterial/Source/DynamicMaterialEditor/Private/UI/Widgets/Editor/SDMMaterialPropertySelector.cpp
+++ b/Engine/Plugins/Experimental/DynamicMaterial/Source/DynamicMaterialEditor/Private/UI/Widgets/Editor/SDMMaterialPropertySelector.cpp
@@ -30,6 +30,8 @@ void SDMMaterialPropertySelector::Construct(const FArguments& InArgs, const TSha
 	EditorWidgetWeak = InEditorWidget;
 	SelectedProperty = EDMMaterialPropertyType::None;
 
+	SetCanTick(false);
+
 	ChildSlot
 	[
 		SNew(SScrollBox)
diff --git a/Engine/Plugins/Experimental/DynamicMaterial/Source/DynamicMaterialEditor/Private/UI/Widgets/Editor/SDMMaterialSlotEditor.cpp b/Engine/Plugins/Experimental/DynamicMaterial/Source/DynamicMaterialEditor/Private/UI/Widgets/Editor/SDMMaterialSlotEditor.cpp
--- a/Engine/Plugins/Experimental/DynamicMaterial/Source/DynamicMaterialEditor/Private/UI/Widgets/Editor/SDMMaterialSlotEditor.cpp
+++ b/Engine/Plugins/Experimental/DynamicMaterial/Source/DynamicMaterialEditor/Private/UI/Widgets/Editor/SDMMaterialSlotEditor.cpp
@@ -77,9 +77,11 @@ void SDMMaterialSlotEditor::Construct(const FArguments& InArgs, const TSharedRef
 	EditorWidgetWeak = InEditorWidget;
 	MaterialSlotWeak = InSlot;
 
+	SetCanTick(false);
+
 	bIsDynamic = !Cast<UDynamicMaterialModel>(InEditorWidget->GetMaterialModelBase());
 
-	Container = TDMWidgetSlot<SWidget>(SharedThis(this), 0, SNullWidget::NullWidget);
+	ContentSlot = TDMWidgetSlot<SWidget>(SharedThis(this), 0, SNullWidget::NullWidget);
 
 	if (!IsValid(InSlot))
 	{
@@ -89,40 +91,40 @@ void SDMMaterialSlotEditor::Construct(const FArguments& InArgs, const TSharedRef
 	InSlot->GetOnPropertiesUpdateDelegate().AddSP(this, &SDMMaterialSlotEditor::OnSlotPropertiesUpdated);
 	InSlot->GetOnLayersUpdateDelegate().AddSP(this, &SDMMaterialSlotEditor::OnSlotLayersUpdated);
 
-	Container << CreateSlot_Container();
+	ContentSlot << CreateSlot_Container();
 }
 
 void SDMMaterialSlotEditor::ValidateSlots()
 {
 	if (!MaterialSlotWeak.IsValid())
 	{
-		if (Container.HasWidget())
+		if (ContentSlot.HasWidget())
 		{
-			Container.ClearWidget();
+			ContentSlot.ClearWidget();
 		}
 
 		return;
 	}
 
-	if (Container.HasBeenInvalidated())
+	if (ContentSlot.HasBeenInvalidated())
 	{
-		Container << CreateSlot_Container();
+		ContentSlot << CreateSlot_Container();
 	}
 	else
 	{
-		if (SlotSettings.HasBeenInvalidated())
+		if (SlotSettingsSlot.HasBeenInvalidated())
 		{
-			SlotSettings << CreateSlot_SlotSettings();
+			SlotSettingsSlot << CreateSlot_SlotSettings();
 		}
 
-		if (LayerView.HasBeenInvalidated())
+		if (LayerViewSlot.HasBeenInvalidated())
 		{
-			LayerView << CreateSlot_LayerView();
+			LayerViewSlot << CreateSlot_LayerView();
 		}
 
-		if (LayerSettings.HasBeenInvalidated())
+		if (LayerSettingsSlot.HasBeenInvalidated())
 		{
-			LayerSettings << CreateSlot_LayerSettings();
+			LayerSettingsSlot << CreateSlot_LayerSettings();
 		}
 	}
 }
@@ -139,7 +141,7 @@ UDMMaterialSlot* SDMMaterialSlotEditor::GetSlot() const
 
 void SDMMaterialSlotEditor::ClearSelection()
 {
-	LayerView->ClearSelection();
+	LayerViewSlot->ClearSelection();
 }
 
 bool SDMMaterialSlotEditor::CanAddNewLayer() const
@@ -196,12 +198,12 @@ void SDMMaterialSlotEditor::AddNewLayer()
 
 bool SDMMaterialSlotEditor::CanInsertNewLayer() const
 {
-	return !!LayerView->GetSelectedLayer();
+	return !!LayerViewSlot->GetSelectedLayer();
 }
 
 void SDMMaterialSlotEditor::InsertNewLayer()
 {
-	UDMMaterialLayerObject* SelectedLayer = LayerView->GetSelectedLayer();
+	UDMMaterialLayerObject* SelectedLayer = LayerViewSlot->GetSelectedLayer();
 
 	if (!SelectedLayer)
 	{
@@ -241,12 +243,12 @@ void SDMMaterialSlotEditor::InsertNewLayer()
 
 bool SDMMaterialSlotEditor::CanCopySelectedLayer() const
 {
-	return !!LayerView->GetSelectedLayer();
+	return !!LayerViewSlot->GetSelectedLayer();
 }
 
 void SDMMaterialSlotEditor::CopySelectedLayer()
 {
-	UDMMaterialLayerObject* SelectedLayer = LayerView->GetSelectedLayer();
+	UDMMaterialLayerObject* SelectedLayer = LayerViewSlot->GetSelectedLayer();
 
 	FPlatformApplicationMisc::ClipboardCopy(*SelectedLayer->SerializeToString());
 }
@@ -336,7 +338,7 @@ bool SDMMaterialSlotEditor::CanDeleteSelectedLayer() const
 		return false;
 	}
 
-	UDMMaterialLayerObject* SelectedLayer = LayerView->GetSelectedLayer();
+	UDMMaterialLayerObject* SelectedLayer = LayerViewSlot->GetSelectedLayer();
 
 	if (!SelectedLayer)
 	{
@@ -349,7 +351,7 @@ bool SDMMaterialSlotEditor::CanDeleteSelectedLayer() const
 void SDMMaterialSlotEditor::DeleteSelectedLayer()
 {
 	UDMMaterialSlot* Slot = GetSlot();
-	UDMMaterialLayerObject* SelectedLayer = LayerView->GetSelectedLayer();
+	UDMMaterialLayerObject* SelectedLayer = LayerViewSlot->GetSelectedLayer();
 
 	FDMScopedUITransaction Transaction(LOCTEXT("DeleteLayer", "Delete Layer"));
 	Slot->Modify();
@@ -360,29 +362,29 @@ void SDMMaterialSlotEditor::DeleteSelectedLayer()
 
 TSharedRef<SDMMaterialSlotLayerView> SDMMaterialSlotEditor::GetLayerView() const
 {
-	return *LayerView;
+	return *LayerViewSlot;
 }
 
 void SDMMaterialSlotEditor::InvalidateSlotSettings()
 {
-	SlotSettings.Invalidate();
+	SlotSettingsSlot.Invalidate();
 }
 
 void SDMMaterialSlotEditor::InvalidateLayerView()
 {
-	LayerView.Invalidate();
+	LayerViewSlot.Invalidate();
 }
 
 void SDMMaterialSlotEditor::InvalidateLayerSettings()
 {
-	LayerSettings.Invalidate();
+	LayerSettingsSlot.Invalidate();
 }
 
 TSharedRef<SWidget> SDMMaterialSlotEditor::CreateSlot_Container()
 {
-	SVerticalBox::FSlot* SettingsSlot = nullptr;
-	SScrollBox::FSlot* LayerViewSlot = nullptr;
-	SVerticalBox::FSlot* LayerSettingsSlot = nullptr;
+	SVerticalBox::FSlot* SettingsSlotPtr = nullptr;
+	SScrollBox::FSlot* LayerViewSlotPtr = nullptr;
+	SVerticalBox::FSlot* LayerSettingsSlotPtr = nullptr;
 
 	TSharedPtr<SAssetDropTarget> DropTarget;
 
@@ -405,7 +407,7 @@ TSharedRef<SWidget> SDMMaterialSlotEditor::CreateSlot_Container()
 			[
 				SNew(SVerticalBox)
 				+ SVerticalBox::Slot()
-				.Expose(SettingsSlot)
+				.Expose(SettingsSlotPtr)
 				.AutoHeight()
 				[
 					SNullWidget::NullWidget
@@ -448,7 +450,7 @@ TSharedRef<SWidget> SDMMaterialSlotEditor::CreateSlot_Container()
 									.Orientation(EOrientation::Orient_Vertical)
 									.ExternalScrollbar(VerticalScrollBar)
 									+ SScrollBox::Slot()
-									.Expose(LayerViewSlot)
+									.Expose(LayerViewSlotPtr)
 									.VAlign(EVerticalAlignment::VAlign_Fill)
 									.Padding(0.f, 0.f, 0.f, 20.f)
 									[
@@ -485,15 +487,15 @@ TSharedRef<SWidget> SDMMaterialSlotEditor::CreateSlot_Container()
 		]
 
 		+ SVerticalBox::Slot()
-		.Expose(LayerSettingsSlot)
+		.Expose(LayerSettingsSlotPtr)
 		.AutoHeight()
 		[
 			SNullWidget::NullWidget
 		];
 
-	SlotSettings = TDMWidgetSlot<SWidget>(SettingsSlot, CreateSlot_SlotSettings());
-	LayerView = TDMWidgetSlot<SDMMaterialSlotLayerView>(LayerViewSlot, CreateSlot_LayerView());
-	LayerSettings = TDMWidgetSlot<SWidget>(LayerSettingsSlot, CreateSlot_LayerSettings());
+	SlotSettingsSlot = TDMWidgetSlot<SWidget>(SettingsSlotPtr, CreateSlot_SlotSettings());
+	LayerViewSlot = TDMWidgetSlot<SDMMaterialSlotLayerView>(LayerViewSlotPtr, CreateSlot_LayerView());
+	LayerSettingsSlot = TDMWidgetSlot<SWidget>(LayerSettingsSlotPtr, CreateSlot_LayerSettings());
 
 	if (UDMMaterialSlot* Slot = GetSlot())
 	{
@@ -501,7 +503,7 @@ TSharedRef<SWidget> SDMMaterialSlotEditor::CreateSlot_Container()
 
 		if (!Layers.IsEmpty())
 		{
-			LayerView->SetSelectedLayer(Layers[0]);
+			LayerViewSlot->SetSelectedLayer(Layers[0]);
 		}
 	}
 
@@ -548,9 +550,9 @@ TSharedRef<SWidget> SDMMaterialSlotEditor::CreateSlot_LayerBlendMode()
 {
 	TSubclassOf<UDMMaterialStageBlend> SelectedBlendMode = nullptr;
 
-	if (LayerView.IsValid())
+	if (LayerViewSlot.IsValid())
 	{
-		if (const UDMMaterialLayerObject* SelectedLayer = LayerView->GetSelectedLayer())
+		if (const UDMMaterialLayerObject* SelectedLayer = LayerViewSlot->GetSelectedLayer())
 		{
 			if (UDMMaterialStage* BaseStage = SelectedLayer->GetFirstEnabledStage(EDMMaterialLayerStage::Base))
 			{
@@ -604,9 +606,9 @@ TSharedRef<SWidget> SDMMaterialSlotEditor::CreateSlot_LayerOpacity()
 {
 	LayerOpacityItem = nullptr;
 
-	if (LayerView.IsValid())
+	if (LayerViewSlot.IsValid())
 	{
-		if (const UDMMaterialLayerObject* SelectedLayer = LayerView->GetSelectedLayer())
+		if (const UDMMaterialLayerObject* SelectedLayer = LayerViewSlot->GetSelectedLayer())
 		{
 			if (UDMMaterialStage* ValidStage = SelectedLayer->GetFirstValidStage(EDMMaterialLayerStage::All))
 			{
@@ -857,7 +859,7 @@ void SDMMaterialSlotEditor::OnSlotPropertiesUpdated(UDMMaterialSlot* InSlot)
 void SDMMaterialSlotEditor::OnLayerSelected(const TSharedRef<SDMMaterialSlotLayerView>& InLayerView, 
 	const TSharedPtr<FDMMaterialLayerReference>& InLayerReference)
 {
-	SlotSettings.Invalidate();
+	SlotSettingsSlot.Invalidate();
 }
 
 FText SDMMaterialSlotEditor::GetLayerButtonsDescription() const
@@ -878,7 +880,7 @@ FText SDMMaterialSlotEditor::GetLayerButtonsDescription() const
 
 TSharedRef<SWidget> SDMMaterialSlotEditor::GetLayerButtonsMenuContent()
 {
-	if (UDMMaterialLayerObject* LayerObject = LayerView->GetSelectedLayer())
+	if (UDMMaterialLayerObject* LayerObject = LayerViewSlot->GetSelectedLayer())
 	{
 		UToolMenu* ContextMenu = FDMMaterialSlotLayerMenus::GenerateSlotLayerMenu(SharedThis(this), LayerObject);
 		return UToolMenus::Get()->GenerateWidget(ContextMenu);
@@ -889,12 +891,12 @@ TSharedRef<SWidget> SDMMaterialSlotEditor::GetLayerButtonsMenuContent()
 
 bool SDMMaterialSlotEditor::GetLayerCanAddEffect() const
 {
-	return !!LayerView->GetSelectedLayer();
+	return !!LayerViewSlot->GetSelectedLayer();
 }
 
 TSharedRef<SWidget> SDMMaterialSlotEditor::GetLayerEffectsMenuContent()
 {
-	if (UDMMaterialLayerObject* LayerObject = LayerView->GetSelectedLayer())
+	if (UDMMaterialLayerObject* LayerObject = LayerViewSlot->GetSelectedLayer())
 	{
 		return FDMMaterialSlotLayerAddEffectMenus::OpenAddEffectMenu(EditorWidgetWeak.Pin(), LayerObject);
 	}
diff --git a/Engine/Plugins/Experimental/DynamicMaterial/Source/DynamicMaterialEditor/Private/UI/Widgets/Editor/SDMMaterialSlotEditor.h b/Engine/Plugins/Experimental/DynamicMaterial/Source/DynamicMaterialEditor/Private/UI/Widgets/Editor/SDMMaterialSlotEditor.h
--- a/Engine/Plugins/Experimental/DynamicMaterial/Source/DynamicMaterialEditor/Private/UI/Widgets/Editor/SDMMaterialSlotEditor.h
+++ b/Engine/Plugins/Experimental/DynamicMaterial/Source/DynamicMaterialEditor/Private/UI/Widgets/Editor/SDMMaterialSlotEditor.h
@@ -75,10 +75,10 @@ protected:
 	TWeakObjectPtr<UDMMaterialSlot> MaterialSlotWeak;
 	bool bIsDynamic;
 
-	TDMWidgetSlot<SWidget> Container;
-	TDMWidgetSlot<SWidget> SlotSettings;
-	TDMWidgetSlot<SDMMaterialSlotLayerView> LayerView;
-	TDMWidgetSlot<SWidget> LayerSettings;
+	TDMWidgetSlot<SWidget> ContentSlot;
+	TDMWidgetSlot<SWidget> SlotSettingsSlot;
+	TDMWidgetSlot<SDMMaterialSlotLayerView> LayerViewSlot;
+	TDMWidgetSlot<SWidget> LayerSettingsSlot;
 
 	TSharedPtr<ICustomDetailsViewItem> LayerOpacityItem;
 
diff --git a/Engine/Plugins/Experimental/DynamicMaterial/Source/DynamicMaterialEditor/Private/UI/Widgets/Editor/SDMObjectEditorWidgetBase.cpp b/Engine/Plugins/Experimental/DynamicMaterial/Source/DynamicMaterialEditor/Private/UI/Widgets/Editor/SDMObjectEditorWidgetBase.cpp
--- a/Engine/Plugins/Experimental/DynamicMaterial/Source/DynamicMaterialEditor/Private/UI/Widgets/Editor/SDMObjectEditorWidgetBase.cpp
+++ b/Engine/Plugins/Experimental/DynamicMaterial/Source/DynamicMaterialEditor/Private/UI/Widgets/Editor/SDMObjectEditorWidgetBase.cpp
@@ -56,27 +56,27 @@ void SDMObjectEditorWidgetBase::Construct(const FArguments& InArgs, const TShare
 		}
 	}
 
-	SScrollBox::FSlot* WidgetSlot = nullptr;
+	SScrollBox::FSlot* ContentSlotPtr = nullptr;
 
 	ChildSlot
 	[
 		SNew(SScrollBox)
 		+ SScrollBox::Slot()
-		.Expose(WidgetSlot)
+		.Expose(ContentSlotPtr)
 		.VAlign(EVerticalAlignment::VAlign_Fill)
 		[
 			SNullWidget::NullWidget
 		]		
 	];
 
-	Container = TDMWidgetSlot<SWidget>(WidgetSlot, CreateWidget());
+	ContentSlot = TDMWidgetSlot<SWidget>(ContentSlotPtr, CreateWidget());
 }
 
 void SDMObjectEditorWidgetBase::Validate()
 {
 	if (!ObjectWeak.IsValid())
 	{
-		Container.ClearWidget();
+		ContentSlot.ClearWidget();
 	}
 }
 
diff --git a/Engine/Plugins/Experimental/DynamicMaterial/Source/DynamicMaterialEditor/Private/UI/Widgets/Editor/SDMObjectEditorWidgetBase.h b/Engine/Plugins/Experimental/DynamicMaterial/Source/DynamicMaterialEditor/Private/UI/Widgets/Editor/SDMObjectEditorWidgetBase.h
--- a/Engine/Plugins/Experimental/DynamicMaterial/Source/DynamicMaterialEditor/Private/UI/Widgets/Editor/SDMObjectEditorWidgetBase.h
+++ b/Engine/Plugins/Experimental/DynamicMaterial/Source/DynamicMaterialEditor/Private/UI/Widgets/Editor/SDMObjectEditorWidgetBase.h
@@ -48,7 +48,7 @@ protected:
 	TWeakPtr<SDMMaterialEditor> EditorWidgetWeak;
 	TWeakObjectPtr<UObject> ObjectWeak;
 
-	TDMWidgetSlot<SWidget> Container;
+	TDMWidgetSlot<SWidget> ContentSlot;
 
 	TSharedPtr<IDetailKeyframeHandler> KeyframeHandler;
 	bool bConstructing;
diff --git a/Engine/Plugins/Experimental/DynamicMaterial/Source/DynamicMaterialEditor/Private/UI/Widgets/Editor/SDMStatusBar.cpp b/Engine/Plugins/Experimental/DynamicMaterial/Source/DynamicMaterialEditor/Private/UI/Widgets/Editor/SDMStatusBar.cpp
--- a/Engine/Plugins/Experimental/DynamicMaterial/Source/DynamicMaterialEditor/Private/UI/Widgets/Editor/SDMStatusBar.cpp
+++ b/Engine/Plugins/Experimental/DynamicMaterial/Source/DynamicMaterialEditor/Private/UI/Widgets/Editor/SDMStatusBar.cpp
@@ -49,6 +49,8 @@ void SDMStatusBar::Construct(const FArguments& InArgs, const TSharedRef<SDMMater
 	MaterialModelBaseWeak = InMaterialModelBase;
 	EditorWidgetWeak = InEditorWidget;
 
+	SetCanTick(false);
+
 	UDynamicMaterialModel* MaterialModel = InEditorWidget->GetMaterialModel();
 
 	if (!MaterialModel)
@@ -67,12 +69,12 @@ void SDMStatusBar::Construct(const FArguments& InArgs, const TSharedRef<SDMMater
 
 	EditorOnlyData->GetOnMaterialBuiltDelegate().AddSP(this, &SDMStatusBar::OnMaterialBuilt);
 
-	Content = TDMWidgetSlot<SWidget>(SharedThis(this), 0, SNullWidget::NullWidget);
+	ContentSlot = TDMWidgetSlot<SWidget>(SharedThis(this), 0, SNullWidget::NullWidget);
 
 	// No stats!
 	if (CachedMaterialStats.NumPixelShaderInstructions > 0)
 	{
-		Content << CreateContent();
+		ContentSlot << CreateContent();
 	}
 }
 
@@ -164,7 +166,7 @@ void SDMStatusBar::OnMaterialBuilt(UDynamicMaterialModelBase* InMaterialModelBas
 
 	CachedMaterialStats = EditorOnlyData->GetMaterialStats();
 
-	Content << CreateContent();
+	ContentSlot << CreateContent();
 }
 
 #undef LOCTEXT_NAMESPACE
diff --git a/Engine/Plugins/Experimental/DynamicMaterial/Source/DynamicMaterialEditor/Private/UI/Widgets/Editor/SDMStatusBar.h b/Engine/Plugins/Experimental/DynamicMaterial/Source/DynamicMaterialEditor/Private/UI/Widgets/Editor/SDMStatusBar.h
--- a/Engine/Plugins/Experimental/DynamicMaterial/Source/DynamicMaterialEditor/Private/UI/Widgets/Editor/SDMStatusBar.h
+++ b/Engine/Plugins/Experimental/DynamicMaterial/Source/DynamicMaterialEditor/Private/UI/Widgets/Editor/SDMStatusBar.h
@@ -25,7 +25,7 @@ protected:
 	TWeakPtr<SDMMaterialEditor> EditorWidgetWeak;
 	TWeakObjectPtr<UDynamicMaterialModelBase> MaterialModelBaseWeak;
 
-	TDMWidgetSlot<SWidget> Content;
+	TDMWidgetSlot<SWidget> ContentSlot;
 
 	int32 CachedSlotCount = 0;
 	int32 CachedCurrentLayerCount = 0;
diff --git a/Engine/Plugins/Experimental/DynamicMaterial/Source/DynamicMaterialEditor/Private/UI/Widgets/Editor/SDMToolBar.cpp b/Engine/Plugins/Experimental/DynamicMaterial/Source/DynamicMaterialEditor/Private/UI/Widgets/Editor/SDMToolBar.cpp
--- a/Engine/Plugins/Experimental/DynamicMaterial/Source/DynamicMaterialEditor/Private/UI/Widgets/Editor/SDMToolBar.cpp
+++ b/Engine/Plugins/Experimental/DynamicMaterial/Source/DynamicMaterialEditor/Private/UI/Widgets/Editor/SDMToolBar.cpp
@@ -48,6 +48,8 @@ void SDMToolBar::Construct(const FArguments& InArgs, const TSharedRef<SDMMateria
 	MaterialActorWeak = InActor;
 	SelectedMaterialElementIndex = INDEX_NONE;
 
+	SetCanTick(false);
+
 	if (IsValid(InActor))
 	{
 		TArray<FDMObjectMaterialProperty> ActorProperties = UDMMaterialModelFunctionLibrary::GetActorMaterialProperties(InActor);
diff --git a/Engine/Plugins/Experimental/DynamicMaterial/Source/DynamicMaterialEditor/Private/UI/Widgets/Editor/SlotEditor/SDMMaterialLayerBlendMode.cpp b/Engine/Plugins/Experimental/DynamicMaterial/Source/DynamicMaterialEditor/Private/UI/Widgets/Editor/SlotEditor/SDMMaterialLayerBlendMode.cpp
--- a/Engine/Plugins/Experimental/DynamicMaterial/Source/DynamicMaterialEditor/Private/UI/Widgets/Editor/SlotEditor/SDMMaterialLayerBlendMode.cpp
+++ b/Engine/Plugins/Experimental/DynamicMaterial/Source/DynamicMaterialEditor/Private/UI/Widgets/Editor/SlotEditor/SDMMaterialLayerBlendMode.cpp
@@ -144,6 +144,8 @@ void SDMMaterialLayerBlendMode::Construct(const FArguments& InArgs, const TShare
 	SlotEditorWidgetWeak = InSlotEditor;
 	SelectedItem = InArgs._SelectedItem;
 
+	SetCanTick(false);
+
 	EnsureBlendMap();
 	EnsureMenuRegistered();
 
diff --git a/Engine/Plugins/Experimental/DynamicMaterial/Source/DynamicMaterialEditor/Private/UI/Widgets/Editor/SlotEditor/SDMMaterialStage.cpp b/Engine/Plugins/Experimental/DynamicMaterial/Source/DynamicMaterialEditor/Private/UI/Widgets/Editor/SlotEditor/SDMMaterialStage.cpp
--- a/Engine/Plugins/Experimental/DynamicMaterial/Source/DynamicMaterialEditor/Private/UI/Widgets/Editor/SlotEditor/SDMMaterialStage.cpp
+++ b/Engine/Plugins/Experimental/DynamicMaterial/Source/DynamicMaterialEditor/Private/UI/Widgets/Editor/SlotEditor/SDMMaterialStage.cpp
@@ -36,6 +36,8 @@ void SDMMaterialStage::Construct(const FArguments& InArgs, const TSharedRef<SDMM
 	SlotLayerItemWeak = InSlotLayerItem;
 	StageWeak = InStage;
 
+	SetCanTick(false);
+
 	if (!IsValid(InStage))
 	{
 		return;
diff --git a/Engine/Plugins/Experimental/DynamicMaterial/Source/DynamicMaterialEditor/Private/UI/Widgets/SDMActorMaterialSelector.cpp b/Engine/Plugins/Experimental/DynamicMaterial/Source/DynamicMaterialEditor/Private/UI/Widgets/SDMActorMaterialSelector.cpp
--- a/Engine/Plugins/Experimental/DynamicMaterial/Source/DynamicMaterialEditor/Private/UI/Widgets/SDMActorMaterialSelector.cpp
+++ b/Engine/Plugins/Experimental/DynamicMaterial/Source/DynamicMaterialEditor/Private/UI/Widgets/SDMActorMaterialSelector.cpp
@@ -33,6 +33,8 @@ void SDMActorMaterialSelector::Construct(const FArguments& InArgs, const TShared
 	ActorWeak = InActor;
 	ActorProperties = MoveTemp(InActorProperties);
 
+	SetCanTick(false);
+
 	ChildSlot
 	[
 		SNew(SBox)
diff --git a/Engine/Plugins/Experimental/DynamicMaterial/Source/DynamicMaterialEditor/Private/UI/Widgets/SDMMaterialDesigner.cpp b/Engine/Plugins/Experimental/DynamicMaterial/Source/DynamicMaterialEditor/Private/UI/Widgets/SDMMaterialDesigner.cpp
--- a/Engine/Plugins/Experimental/DynamicMaterial/Source/DynamicMaterialEditor/Private/UI/Widgets/SDMMaterialDesigner.cpp
+++ b/Engine/Plugins/Experimental/DynamicMaterial/Source/DynamicMaterialEditor/Private/UI/Widgets/SDMMaterialDesigner.cpp
@@ -24,7 +24,9 @@ void SDMMaterialDesigner::PrivateRegisterAttributes(FSlateAttributeDescriptor::F
 
 void SDMMaterialDesigner::Construct(const FArguments& InArgs)
 {
-	Content = TDMWidgetSlot<SWidget>(SharedThis(this), 0, SNullWidget::NullWidget);
+	SetCanTick(true);
+
+	ContentSlot = TDMWidgetSlot<SWidget>(SharedThis(this), 0, SNullWidget::NullWidget);
 
 	SetSelectPromptView();
 }
@@ -120,24 +122,34 @@ void SDMMaterialDesigner::OnActorSelected(AActor* InActor)
 
 UDynamicMaterialModelBase* SDMMaterialDesigner::GetMaterialModelBase() const
 {
-	if (!Content.HasWidget())
+	if (!Content.IsValid())
 	{
 		return nullptr;
 	}
 
 	if (Content->GetWidgetClass().GetWidgetType() == SDMMaterialEditor::StaticWidgetClass().GetWidgetType())
 	{
-		return StaticCastSharedRef<SDMMaterialEditor>(*Content)->GetMaterialModelBase();
+		return StaticCastSharedPtr<SDMMaterialEditor>(Content)->GetMaterialModelBase();
 	}
 
 	if (Content->GetWidgetClass().GetWidgetType() == SDMMaterialWizard::StaticWidgetClass().GetWidgetType())
 	{
-		return StaticCastSharedRef<SDMMaterialWizard>(*Content)->GetMaterialModel();
+		return StaticCastSharedPtr<SDMMaterialWizard>(Content)->GetMaterialModel();
 	}
 
 	return nullptr;
 }
 
+void SDMMaterialDesigner::Tick(const FGeometry& InAllottedGeometry, const double InCurrentTime, const float InDeltaTime)
+{
+	if (Content.IsValid() && Content->GetWidgetClass().GetWidgetType() == SDMMaterialEditor::StaticWidgetClass().GetWidgetType())
+	{
+		StaticCastSharedPtr<SDMMaterialEditor>(Content)->Validate();
+	}
+
+	SCompoundWidget::Tick(InAllottedGeometry, InCurrentTime, InDeltaTime);
+}
+
 void SDMMaterialDesigner::OpenMaterialModelBase_Internal(UDynamicMaterialModelBase* InMaterialModelBase)
 {
 	if (NeedsWizard(InMaterialModelBase))
@@ -255,9 +267,11 @@ void SDMMaterialDesigner::SetEditorView(const FDMObjectMaterialProperty& InObjec
 
 void SDMMaterialDesigner::SetWidget(const TSharedRef<SWidget>& InWidget, bool bInIncludeAssetDropTarget)
 {
+	Content = InWidget;
+
 	if (!bInIncludeAssetDropTarget)
 	{
-		Content << InWidget;
+		ContentSlot << InWidget;
 	}
 	else
 	{
@@ -271,7 +285,7 @@ void SDMMaterialDesigner::SetWidget(const TSharedRef<SWidget>& InWidget, bool bI
 		using namespace UE::DynamicMaterialEditor::Private;
 		DropTarget::SetInvalidColor(&DropTarget.Get(), FStyleColors::Transparent);
 
-		Content << DropTarget;
+		ContentSlot << DropTarget;
 	}
 }
 
diff --git a/Engine/Plugins/Experimental/DynamicMaterial/Source/DynamicMaterialEditor/Private/UI/Widgets/SDMMaterialDesigner.h b/Engine/Plugins/Experimental/DynamicMaterial/Source/DynamicMaterialEditor/Private/UI/Widgets/SDMMaterialDesigner.h
--- a/Engine/Plugins/Experimental/DynamicMaterial/Source/DynamicMaterialEditor/Private/UI/Widgets/SDMMaterialDesigner.h
+++ b/Engine/Plugins/Experimental/DynamicMaterial/Source/DynamicMaterialEditor/Private/UI/Widgets/SDMMaterialDesigner.h
@@ -57,8 +57,13 @@ public:
 	/** Getters */
 	UDynamicMaterialModelBase* GetMaterialModelBase() const;
 
+	//~ Begin SWidget
+	virtual void Tick(const FGeometry& InAllottedGeometry, const double InCurrentTime, const float InDeltaTime) override;
+	//~ End SWidget
+
 protected:
-	TDMWidgetSlot<SWidget> Content;
+	TDMWidgetSlot<SWidget> ContentSlot;
+	TSharedPtr<SWidget> Content;
 
 	/** Opens assets in specific widget modes. */
 	void OpenMaterialModelBase_Internal(UDynamicMaterialModelBase* InMaterialModelBase);
diff --git a/Engine/Plugins/Experimental/DynamicMaterial/Source/DynamicMaterialEditor/Private/UI/Widgets/SDMMaterialEditor.cpp b/Engine/Plugins/Experimental/DynamicMaterial/Source/DynamicMaterialEditor/Private/UI/Widgets/SDMMaterialEditor.cpp
--- a/Engine/Plugins/Experimental/DynamicMaterial/Source/DynamicMaterialEditor/Private/UI/Widgets/SDMMaterialEditor.cpp
+++ b/Engine/Plugins/Experimental/DynamicMaterial/Source/DynamicMaterialEditor/Private/UI/Widgets/SDMMaterialEditor.cpp
@@ -81,7 +81,9 @@ void SDMMaterialEditor::Construct(const FArguments& InArgs, const TSharedRef<SDM
 	bGlobalSettingsMode = true;
 	PropertyToSelect.Reset();
 
-	Container = TDMWidgetSlot<SWidget>(SharedThis(this), 0, SNullWidget::NullWidget);
+	SetCanTick(false);
+
+	ContentSlot = TDMWidgetSlot<SWidget>(SharedThis(this), 0, SNullWidget::NullWidget);
 
 	if (InArgs._MaterialProperty.IsSet())
 	{
@@ -186,31 +188,31 @@ void SDMMaterialEditor::SetMaterialActor(AActor* InActor)
 
 	TSharedRef<SDMToolBar> NewToolBar = SNew(SDMToolBar, SharedThis(this), InActor);
 
-	ToolBar << NewToolBar;
+	ToolBarSlot << NewToolBar;
 }
 
 TSharedPtr<SDMMaterialSlotEditor> SDMMaterialEditor::GetSlotEditorWidget() const
 {
-	return &SlotEditor;
+	return &SlotEditorSlot;
 }
 
 TSharedPtr<SDMMaterialComponentEditor> SDMMaterialEditor::GetComponentEditorWidget() const
 {
-	return &ComponentEditor;
+	return &ComponentEditorSlot;
 }
 
 void SDMMaterialEditor::SelectProperty(EDMMaterialPropertyType InProperty, bool bInForceRefresh)
 {
-	if (bInForceRefresh || !PropertySelector.IsValid())
+	if (bInForceRefresh || !PropertySelectorSlot.IsValid())
 	{
 		PropertyToSelect = InProperty;
-		PropertySelector.Invalidate();
+		PropertySelectorSlot.Invalidate();
 		return;
 	}
 
-	if (PropertySelector->GetSelectedProperty() != InProperty)
+	if (PropertySelectorSlot->GetSelectedProperty() != InProperty)
 	{
-		PropertySelector->SetSelectedProperty(InProperty);
+		PropertySelectorSlot->SetSelectedProperty(InProperty);
 	}
 	
 	PropertyToSelect.Reset();
@@ -228,17 +230,17 @@ TSharedRef<FDMPreviewMaterialManager> SDMMaterialEditor::GetPreviewMaterialManag
 
 void SDMMaterialEditor::EditSlot(UDMMaterialSlot* InSlot, bool bInForceRefresh)
 {
-	if (!bInForceRefresh && SlotEditor.IsValid() && SlotEditor->GetSlot() == InSlot)
+	if (!bInForceRefresh && SlotEditorSlot.IsValid() && SlotEditorSlot->GetSlot() == InSlot)
 	{
 		return;
 	}
 
-	Right.Invalidate();
+	RightSlot.Invalidate();
 
-	SlotEditor.Invalidate();
+	SlotEditorSlot.Invalidate();
 	SlotToEdit = InSlot;
 
-	ComponentEditor.Invalidate();
+	ComponentEditorSlot.Invalidate();
 	ComponentToEdit.Reset();
 
 	bGlobalSettingsMode = !InSlot;
@@ -258,20 +260,20 @@ void SDMMaterialEditor::EditSlot(UDMMaterialSlot* InSlot, bool bInForceRefresh)
 
 void SDMMaterialEditor::EditComponent(UDMMaterialComponent* InComponent, bool bInForceRefresh)
 {
-	if (!bInForceRefresh && ComponentEditor.IsValid() && ComponentEditor->GetComponent() == InComponent)
+	if (!bInForceRefresh && ComponentEditorSlot.IsValid() && ComponentEditorSlot->GetComponent() == InComponent)
 	{
 		return;
 	}
 
 	if (bGlobalSettingsMode)
 	{
-		Right.Invalidate();
-		SlotEditor.Invalidate();
+		RightSlot.Invalidate();
+		SlotEditorSlot.Invalidate();
 	}
 
 	bGlobalSettingsMode = false;
 
-	ComponentEditor.Invalidate();
+	ComponentEditorSlot.Invalidate();
 	ComponentToEdit = InComponent;
 }
 
@@ -284,35 +286,18 @@ void SDMMaterialEditor::EditGlobalSettings(bool bInForceRefresh)
 
 	if (!bGlobalSettingsMode)
 	{
-		Right.Invalidate();
-		SlotEditor.Invalidate();
-		ComponentEditor.Invalidate();
+		RightSlot.Invalidate();
+		SlotEditorSlot.Invalidate();
+		ComponentEditorSlot.Invalidate();
 	}
 
 	bGlobalSettingsMode = true;
 
-	GlobalSettingsEditor.Invalidate();
-}
-
-SDMMaterialEditor::FOnEditedSlotChanged::RegistrationType& SDMMaterialEditor::GetOnEditedSlotChanged()
-{
-	return OnEditedSlotChanged;
-}
-
-SDMMaterialEditor::FOnEditedComponentChanged::RegistrationType& SDMMaterialEditor::GetOnEditedComponentChanged()
-{
-	return OnEditedComponentChanged;
-}
-
-bool SDMMaterialEditor::SupportsKeyboardFocus() const
-{
-	return true;
+	GlobalSettingsEditorSlot.Invalidate();
 }
 
-void SDMMaterialEditor::Tick(const FGeometry& InAllottedGeometry, const double InCurrentTime, const float InDeltaTime)
+void SDMMaterialEditor::Validate()
 {
-	SCompoundWidget::Tick(InAllottedGeometry, InCurrentTime, InDeltaTime);
-
 	if (!FDynamicMaterialModule::AreUObjectsSafe())
 	{
 		return;
@@ -354,6 +339,21 @@ void SDMMaterialEditor::Tick(const FGeometry& InAllottedGeometry, const double I
 	ValidateSlots();
 }
 
+SDMMaterialEditor::FOnEditedSlotChanged::RegistrationType& SDMMaterialEditor::GetOnEditedSlotChanged()
+{
+	return OnEditedSlotChanged;
+}
+
+SDMMaterialEditor::FOnEditedComponentChanged::RegistrationType& SDMMaterialEditor::GetOnEditedComponentChanged()
+{
+	return OnEditedComponentChanged;
+}
+
+bool SDMMaterialEditor::SupportsKeyboardFocus() const
+{
+	return true;
+}
+
 FReply SDMMaterialEditor::OnKeyDown(const FGeometry& InMyGeometry, const FKeyEvent& InKeyEvent)
 {
 	if (CommandList->ProcessCommandBindings(InKeyEvent))
@@ -476,133 +476,133 @@ void SDMMaterialEditor::Close()
 
 void SDMMaterialEditor::ValidateSlots()
 {
-	if (Container.HasBeenInvalidated())
+	if (ContentSlot.HasBeenInvalidated())
 	{
 		CreateLayout();
 		return;
 	}
 
-	if (ToolBar.HasBeenInvalidated())
+	if (ToolBarSlot.HasBeenInvalidated())
 	{
-		ToolBar << CreateSlot_ToolBar();
+		ToolBarSlot << CreateSlot_ToolBar();
 	}
 
-	if (Main.HasBeenInvalidated())
+	if (MainSlot.HasBeenInvalidated())
 	{
-		Main << CreateSlot_Main();
+		MainSlot << CreateSlot_Main();
 	}
 	else
 	{
-		if (Left.HasBeenInvalidated())
+		if (LeftSlot.HasBeenInvalidated())
 		{
-			Left << CreateSlot_Left();
+			LeftSlot << CreateSlot_Left();
 		}
 		else
 		{
-			if (Preview.HasBeenInvalidated())
+			if (MaterialPreviewSlot.HasBeenInvalidated())
 			{
-				Preview << CreateSlot_Preview();
+				MaterialPreviewSlot << CreateSlot_Preview();
 			}
 
-			if (PropertySelector.HasBeenInvalidated())
+			if (PropertySelectorSlot.HasBeenInvalidated())
 			{
-				PropertySelector << CreateSlot_PropertySelector();
+				PropertySelectorSlot << CreateSlot_PropertySelector();
 			}
 		}
 
-		if (Right.HasBeenInvalidated())
+		if (RightSlot.HasBeenInvalidated())
 		{
-			Right << CreateSlot_Right();
+			RightSlot << CreateSlot_Right();
 		}
 		else if (bGlobalSettingsMode)
 		{
-			if (GlobalSettingsEditor.HasBeenInvalidated())
+			if (GlobalSettingsEditorSlot.HasBeenInvalidated())
 			{
-				GlobalSettingsEditor << CreateSlot_GlobalSettingsEditor();
+				GlobalSettingsEditorSlot << CreateSlot_GlobalSettingsEditor();
 			}
 			else
 			{
-				GlobalSettingsEditor->Validate();
+				GlobalSettingsEditorSlot->Validate();
 			}
 		}
 		else
 		{
-			if (SlotEditor.HasBeenInvalidated())
+			if (SlotEditorSlot.HasBeenInvalidated())
 			{
-				SlotEditor << CreateSlot_SlotEditor();
+				SlotEditorSlot << CreateSlot_SlotEditor();
 			}
 			else
 			{
-				SlotEditor->ValidateSlots();
+				SlotEditorSlot->ValidateSlots();
 			}
 
-			if (ComponentEditor.HasBeenInvalidated())
+			if (ComponentEditorSlot.HasBeenInvalidated())
 			{
-				ComponentEditor << CreateSlot_ComponentEditor();
+				ComponentEditorSlot << CreateSlot_ComponentEditor();
 			}
 			else
 			{
-				ComponentEditor->Validate();
+				ComponentEditorSlot->Validate();
 			}
 		}
 	}
 
-	if (StatusBar.HasBeenInvalidated())
+	if (StatusBarSlot.HasBeenInvalidated())
 	{
-		StatusBar << CreateSlot_StatusBar();
+		StatusBarSlot << CreateSlot_StatusBar();
 	}
 }
 
 void SDMMaterialEditor::ClearSlots()
 {
-	Container.ClearWidget();
-	ToolBar.ClearWidget();
-	Main.ClearWidget();
-	Left.ClearWidget();
-	Right.ClearWidget();
-	Preview.ClearWidget();
-	PropertySelector.ClearWidget();
-	SlotEditor.ClearWidget();
-	ComponentEditor.ClearWidget();
-	StatusBar.ClearWidget();
+	ContentSlot.ClearWidget();
+	ToolBarSlot.ClearWidget();
+	MainSlot.ClearWidget();
+	LeftSlot.ClearWidget();
+	RightSlot.ClearWidget();
+	MaterialPreviewSlot.ClearWidget();
+	PropertySelectorSlot.ClearWidget();
+	SlotEditorSlot.ClearWidget();
+	ComponentEditorSlot.ClearWidget();
+	StatusBarSlot.ClearWidget();
 }
 
 void SDMMaterialEditor::CreateLayout()
 {
-	Container << CreateSlot_Container();
+	ContentSlot << CreateSlot_Container();
 }
 
 TSharedRef<SWidget> SDMMaterialEditor::CreateSlot_Container()
 {
-	SVerticalBox::FSlot* ToolBarSlot = nullptr;
-	SVerticalBox::FSlot* MainSlot = nullptr;
-	SVerticalBox::FSlot* StatusBarSlot = nullptr;
+	SVerticalBox::FSlot* ToolBarSlotPtr = nullptr;
+	SVerticalBox::FSlot* MainSlotPtr = nullptr;
+	SVerticalBox::FSlot* StatusBarSlotPtr = nullptr;
 
 	TSharedRef<SVerticalBox> NewContainer = SNew(SVerticalBox)
 		+ SVerticalBox::Slot()
-		.Expose(ToolBarSlot)
+		.Expose(ToolBarSlotPtr)
 		.AutoHeight()
 		[
 			SNullWidget::NullWidget
 		]
 
 		+ SVerticalBox::Slot()
-		.Expose(MainSlot)
+		.Expose(MainSlotPtr)
 		.FillHeight(1.0f)
 		[
 			SNullWidget::NullWidget
 		]
 
 		+ SVerticalBox::Slot()
-		.Expose(StatusBarSlot)
+		.Expose(StatusBarSlotPtr)
 		.AutoHeight()
 		[
 			SNullWidget::NullWidget
 		];
 
-	ToolBar = TDMWidgetSlot<SDMToolBar>(ToolBarSlot, CreateSlot_ToolBar());
-	Main = TDMWidgetSlot<SWidget>(MainSlot, CreateSlot_Main());
-	StatusBar = TDMWidgetSlot<SDMStatusBar>(StatusBarSlot, CreateSlot_StatusBar());
+	ToolBarSlot = TDMWidgetSlot<SDMToolBar>(ToolBarSlotPtr, CreateSlot_ToolBar());
+	MainSlot = TDMWidgetSlot<SWidget>(MainSlotPtr, CreateSlot_Main());
+	StatusBarSlot = TDMWidgetSlot<SDMStatusBar>(StatusBarSlotPtr, CreateSlot_StatusBar());
 
 	return NewContainer;
 }
@@ -620,38 +620,38 @@ TSharedRef<SDMToolBar> SDMMaterialEditor::CreateSlot_ToolBar()
 
 TSharedRef<SWidget> SDMMaterialEditor::CreateSlot_Main()
 {
-	SHorizontalBox::FSlot* LeftSlot = nullptr;
-	SHorizontalBox::FSlot* RightSlot = nullptr;
+	SHorizontalBox::FSlot* LeftSlotPtr = nullptr;
+	SHorizontalBox::FSlot* RightSlotPtr = nullptr;
 
 	TSharedRef<SHorizontalBox> NewMain = SNew(SHorizontalBox)
 		+ SHorizontalBox::Slot()
-		.Expose(LeftSlot)
+		.Expose(LeftSlotPtr)
 		.AutoWidth()
 		[
 			SNullWidget::NullWidget
 		]
 
 		+ SHorizontalBox::Slot()
-		.Expose(RightSlot)
+		.Expose(RightSlotPtr)
 		.FillWidth(1.0f)
 		[
 			SNullWidget::NullWidget
 		];
 
-	Left = TDMWidgetSlot<SWidget>(LeftSlot, CreateSlot_Left());
-	Right = TDMWidgetSlot<SWidget>(RightSlot, CreateSlot_Right());
+	LeftSlot = TDMWidgetSlot<SWidget>(LeftSlotPtr, CreateSlot_Left());
+	RightSlot = TDMWidgetSlot<SWidget>(RightSlotPtr, CreateSlot_Right());
 
 	return NewMain;
 }
 
 TSharedRef<SWidget> SDMMaterialEditor::CreateSlot_Left()
 {
-	SVerticalBox::FSlot* PreviewSlot = nullptr;
-	SVerticalBox::FSlot* PropertySelectorSlot = nullptr;
+	SVerticalBox::FSlot* MaterialPreviewSlotPtr = nullptr;
+	SVerticalBox::FSlot* PropertySelectorSlotPtr = nullptr;
 
 	TSharedRef<SVerticalBox> NewLeft = SNew(SVerticalBox)
 		+ SVerticalBox::Slot()
-		.Expose(PreviewSlot)
+		.Expose(MaterialPreviewSlotPtr)
 		.AutoHeight()
 		.Padding(5.f)
 		[
@@ -659,15 +659,15 @@ TSharedRef<SWidget> SDMMaterialEditor::CreateSlot_Left()
 		]
 
 		+ SVerticalBox::Slot()
-		.Expose(PropertySelectorSlot)
+		.Expose(PropertySelectorSlotPtr)
 		.FillHeight(1.0f)
 		.Padding(5.f, 0.f, 5.f, 5.f)
 		[
 			SNullWidget::NullWidget
 		];
 
-	Preview = TDMWidgetSlot<SDMMaterialPreview>(PreviewSlot, CreateSlot_Preview());
-	PropertySelector = TDMWidgetSlot<SDMMaterialPropertySelector>(PropertySelectorSlot, CreateSlot_PropertySelector());
+	MaterialPreviewSlot = TDMWidgetSlot<SDMMaterialPreview>(MaterialPreviewSlotPtr, CreateSlot_Preview());
+	PropertySelectorSlot = TDMWidgetSlot<SDMMaterialPropertySelector>(PropertySelectorSlotPtr, CreateSlot_PropertySelector());
 
 	return NewLeft;
 }
@@ -692,17 +692,17 @@ TSharedRef<SWidget> SDMMaterialEditor::CreateSlot_Right()
 
 TSharedRef<SWidget> SDMMaterialEditor::CreateSlot_Right_GlobalSettings()
 {
-	SVerticalBox::FSlot* GlobalSettingsSlot = nullptr;
+	SVerticalBox::FSlot* GlobalSettingsSlotPtr = nullptr;
 
 	TSharedRef<SVerticalBox> NewRight = SNew(SVerticalBox)
 		+ SVerticalBox::Slot()
-		.Expose(GlobalSettingsSlot)
+		.Expose(GlobalSettingsSlotPtr)
 		.AutoHeight()
 		[
 			SNullWidget::NullWidget
 		];
 
-	GlobalSettingsEditor = TDMWidgetSlot<SDMMaterialGlobalSettingsEditor>(GlobalSettingsSlot, CreateSlot_GlobalSettingsEditor());
+	GlobalSettingsEditorSlot = TDMWidgetSlot<SDMMaterialGlobalSettingsEditor>(GlobalSettingsSlotPtr, CreateSlot_GlobalSettingsEditor());
 
 	return NewRight;
 }
@@ -721,8 +721,8 @@ TSharedRef<SWidget> SDMMaterialEditor::CreateSlot_Right_Slot()
 		SplitterValue = Settings->SplitterLocation;
 	}
 
-	SSplitter::FSlot* SlotEditorSlot = nullptr;
-	SSplitter::FSlot* ComponentEditorSlot = nullptr;
+	SSplitter::FSlot* SlotEditorSlotPtr = nullptr;
+	SSplitter::FSlot* ComponentEditorSlotPtr = nullptr;
 
 	TSharedRef<SSplitter> NewRight = SNew(SSplitter)
 		.Style(FAppStyle::Get(), "DetailsView.Splitter")
@@ -733,7 +733,7 @@ TSharedRef<SWidget> SDMMaterialEditor::CreateSlot_Right_Slot()
 		.OnSplitterFinishedResizing(this, &SDMMaterialEditor::OnRightSlotSplitterResized)
 
 		+ SSplitter::Slot()
-		.Expose(SlotEditorSlot)
+		.Expose(SlotEditorSlotPtr)
 		.Resizable(true)
 		.SizeRule(SSplitter::ESizeRule::FractionOfParent)
 		.MinSize(50)
@@ -743,7 +743,7 @@ TSharedRef<SWidget> SDMMaterialEditor::CreateSlot_Right_Slot()
 		]
 
 		+ SSplitter::Slot()
-		.Expose(ComponentEditorSlot)
+		.Expose(ComponentEditorSlotPtr)
 		.Resizable(true)
 		.SizeRule(SSplitter::ESizeRule::FractionOfParent)
 		.MinSize(50)
@@ -752,8 +752,8 @@ TSharedRef<SWidget> SDMMaterialEditor::CreateSlot_Right_Slot()
 			SNullWidget::NullWidget
 		];
 
-	SlotEditor = TDMWidgetSlot<SDMMaterialSlotEditor>(SlotEditorSlot, CreateSlot_SlotEditor());
-	ComponentEditor = TDMWidgetSlot<SDMMaterialComponentEditor>(ComponentEditorSlot, CreateSlot_ComponentEditor());
+	SlotEditorSlot = TDMWidgetSlot<SDMMaterialSlotEditor>(SlotEditorSlotPtr, CreateSlot_SlotEditor());
+	ComponentEditorSlot = TDMWidgetSlot<SDMMaterialComponentEditor>(ComponentEditorSlotPtr, CreateSlot_ComponentEditor());
 
 	return NewRight;
 }
@@ -841,7 +841,7 @@ void SDMMaterialEditor::OnUndo()
 		{
 			if (PropertyPair.Value->IsEnabled())
 			{
-				PropertySelector->SetSelectedProperty(PropertyPair.Key);
+				PropertySelectorSlot->SetSelectedProperty(PropertyPair.Key);
 			}
 		}
 	}
@@ -849,16 +849,16 @@ void SDMMaterialEditor::OnUndo()
 
 void SDMMaterialEditor::OnEnginePreExit()
 {
-	Preview.ClearWidget();
+	MaterialPreviewSlot.ClearWidget();
 }
 
 void SDMMaterialEditor::OnRightSlotSplitterResized()
 {
 	UDynamicMaterialEditorSettings* Settings = UDynamicMaterialEditorSettings::Get();
 
-	if (SlotEditor.GetSlot())
+	if (SlotEditorSlot.GetSlot())
 	{
-		Settings->SplitterLocation = static_cast<SSplitter::FSlot*>(SlotEditor.GetSlot())->GetSizeValue();
+		Settings->SplitterLocation = static_cast<SSplitter::FSlot*>(SlotEditorSlot.GetSlot())->GetSizeValue();
 		Settings->SaveConfig();
 	}
 }
diff --git a/Engine/Plugins/Experimental/DynamicMaterial/Source/DynamicMaterialEditor/Private/UI/Widgets/SDMMaterialEditor.h b/Engine/Plugins/Experimental/DynamicMaterial/Source/DynamicMaterialEditor/Private/UI/Widgets/SDMMaterialEditor.h
--- a/Engine/Plugins/Experimental/DynamicMaterial/Source/DynamicMaterialEditor/Private/UI/Widgets/SDMMaterialEditor.h
+++ b/Engine/Plugins/Experimental/DynamicMaterial/Source/DynamicMaterialEditor/Private/UI/Widgets/SDMMaterialEditor.h
@@ -79,6 +79,8 @@ public:
 
 	void EditGlobalSettings(bool bInForceRefresh = false);
 
+	void Validate();
+
 	DECLARE_MULTICAST_DELEGATE_TwoParams(FOnEditedSlotChanged, const TSharedRef<SDMMaterialSlotEditor>&, UDMMaterialSlot*);
 	FOnEditedSlotChanged::RegistrationType& GetOnEditedSlotChanged();
 
@@ -87,7 +89,6 @@ public:
 
 	//~ Begin SWidget
 	virtual bool SupportsKeyboardFocus() const override;
-	virtual void Tick(const FGeometry& InAllottedGeometry, const double InCurrentTime, const float InDeltaTime) override;
 	virtual FReply OnKeyDown(const FGeometry& InMyGeometry, const FKeyEvent& InKeyEvent) override;
 	//~ End SWidget
 
@@ -99,17 +100,17 @@ public:
 protected:
 	TWeakPtr<SDMMaterialDesigner> DesignerWidgetWeak;
 
-	TDMWidgetSlot<SWidget> Container;
-	TDMWidgetSlot<SDMToolBar> ToolBar;
-	TDMWidgetSlot<SWidget> Main;
-	TDMWidgetSlot<SWidget> Left;
-	TDMWidgetSlot<SWidget> Right;
-	TDMWidgetSlot<SDMMaterialPreview> Preview;
-	TDMWidgetSlot<SDMMaterialPropertySelector> PropertySelector;
-	TDMWidgetSlot<SDMMaterialGlobalSettingsEditor> GlobalSettingsEditor;
-	TDMWidgetSlot<SDMMaterialSlotEditor> SlotEditor;
-	TDMWidgetSlot<SDMMaterialComponentEditor> ComponentEditor;
-	TDMWidgetSlot<SDMStatusBar> StatusBar;
+	TDMWidgetSlot<SWidget> ContentSlot;
+	TDMWidgetSlot<SDMToolBar> ToolBarSlot;
+	TDMWidgetSlot<SWidget> MainSlot;
+	TDMWidgetSlot<SWidget> LeftSlot;
+	TDMWidgetSlot<SWidget> RightSlot;
+	TDMWidgetSlot<SDMMaterialPreview> MaterialPreviewSlot;
+	TDMWidgetSlot<SDMMaterialPropertySelector> PropertySelectorSlot;
+	TDMWidgetSlot<SDMMaterialGlobalSettingsEditor> GlobalSettingsEditorSlot;
+	TDMWidgetSlot<SDMMaterialSlotEditor> SlotEditorSlot;
+	TDMWidgetSlot<SDMMaterialComponentEditor> ComponentEditorSlot;
+	TDMWidgetSlot<SDMStatusBar> StatusBarSlot;
 
 	TWeakObjectPtr<UDynamicMaterialModelBase> MaterialModelBaseWeak;
 	TOptional<FDMObjectMaterialProperty> ObjectMaterialPropertyOpt;
diff --git a/Engine/Plugins/Experimental/DynamicMaterial/Source/DynamicMaterialEditor/Private/UI/Widgets/SDMMaterialSelectPrompt.cpp b/Engine/Plugins/Experimental/DynamicMaterial/Source/DynamicMaterialEditor/Private/UI/Widgets/SDMMaterialSelectPrompt.cpp
--- a/Engine/Plugins/Experimental/DynamicMaterial/Source/DynamicMaterialEditor/Private/UI/Widgets/SDMMaterialSelectPrompt.cpp
+++ b/Engine/Plugins/Experimental/DynamicMaterial/Source/DynamicMaterialEditor/Private/UI/Widgets/SDMMaterialSelectPrompt.cpp
@@ -15,6 +15,8 @@ void SDMMaterialSelectPrompt::PrivateRegisterAttributes(FSlateAttributeDescripto
 
 void SDMMaterialSelectPrompt::Construct(const FArguments& InArgs)
 {
+	SetCanTick(false);
+
 	ChildSlot
 	[
 		SNew(SBox)
diff --git a/Engine/Plugins/Experimental/DynamicMaterial/Source/DynamicMaterialEditor/Private/UI/Widgets/SDMMaterialWizard.cpp b/Engine/Plugins/Experimental/DynamicMaterial/Source/DynamicMaterialEditor/Private/UI/Widgets/SDMMaterialWizard.cpp
--- a/Engine/Plugins/Experimental/DynamicMaterial/Source/DynamicMaterialEditor/Private/UI/Widgets/SDMMaterialWizard.cpp
+++ b/Engine/Plugins/Experimental/DynamicMaterial/Source/DynamicMaterialEditor/Private/UI/Widgets/SDMMaterialWizard.cpp
@@ -60,6 +60,8 @@ void SDMMaterialWizard::Construct(const FArguments& InArgs, const TSharedRef<SDM
 	MaterialModelWeak = InArgs._MaterialModel;
 	MaterialObjectProperty = InArgs._MaterialProperty;
 
+	SetCanTick(false);
+
 	if (MaterialObjectProperty.IsSet())
 	{
 		if (UDynamicMaterialModelBase* MaterialModelBase = MaterialObjectProperty.GetValue().GetMaterialModelBase())
diff --git a/Engine/Plugins/Experimental/DynamicMaterial/Source/DynamicMaterialEditor/Private/UI/Widgets/Visualizers/SDMTextureUVVisualizerPopout.cpp b/Engine/Plugins/Experimental/DynamicMaterial/Source/DynamicMaterialEditor/Private/UI/Widgets/Visualizers/SDMTextureUVVisualizerPopout.cpp
--- a/Engine/Plugins/Experimental/DynamicMaterial/Source/DynamicMaterialEditor/Private/UI/Widgets/Visualizers/SDMTextureUVVisualizerPopout.cpp
+++ b/Engine/Plugins/Experimental/DynamicMaterial/Source/DynamicMaterialEditor/Private/UI/Widgets/Visualizers/SDMTextureUVVisualizerPopout.cpp
@@ -101,6 +101,8 @@ void SDMTextureUVVisualizerPopout::Construct(const FArguments& InArgs, const TSh
 	check(InMaterialStage);
 	check(InArgs._TextureUV || InArgs._TextureUVDynamic);
 
+	SetCanTick(false);
+
 	ChildSlot
 	[
 		SNew(SHorizontalBox)
diff --git a/Engine/Plugins/Experimental/DynamicMaterial/Source/DynamicMaterialEditor/Private/UI/Widgets/Visualizers/SDMTextureUVVisualizerProperty.cpp b/Engine/Plugins/Experimental/DynamicMaterial/Source/DynamicMaterialEditor/Private/UI/Widgets/Visualizers/SDMTextureUVVisualizerProperty.cpp
--- a/Engine/Plugins/Experimental/DynamicMaterial/Source/DynamicMaterialEditor/Private/UI/Widgets/Visualizers/SDMTextureUVVisualizerProperty.cpp
+++ b/Engine/Plugins/Experimental/DynamicMaterial/Source/DynamicMaterialEditor/Private/UI/Widgets/Visualizers/SDMTextureUVVisualizerProperty.cpp
@@ -22,6 +22,8 @@ void SDMTextureUVVisualizerProperty::Construct(const FArguments& InArgs, const T
 	check(InMaterialStage);
 	check(InArgs._TextureUV || InArgs._TextureUVDynamic);
 
+	SetCanTick(false);
+
 	ChildSlot
 	[
 		SNew(SVerticalBox)
