Material Designer
- When creating a MD asset with the "Create Material" button on the slot selector page, the actor properties are now displayed propertly.
- Each element in the actor element selector is now only displaced once.

diff --git a/Engine/Plugins/Experimental/DynamicMaterial/Source/DynamicMaterialEditor/Private/UI/Widgets/Editor/SDMToolBar.cpp b/Engine/Plugins/Experimental/DynamicMaterial/Source/DynamicMaterialEditor/Private/UI/Widgets/Editor/SDMToolBar.cpp
--- a/Engine/Plugins/Experimental/DynamicMaterial/Source/DynamicMaterialEditor/Private/UI/Widgets/Editor/SDMToolBar.cpp
+++ b/Engine/Plugins/Experimental/DynamicMaterial/Source/DynamicMaterialEditor/Private/UI/Widgets/Editor/SDMToolBar.cpp
@@ -49,17 +49,6 @@ void SDMToolBar::Construct(const FArguments& InArgs, const TSharedRef<SDMMateria
 	SelectedMaterialElementIndex = INDEX_NONE;
 
 	SetCanTick(false);
-
-	if (IsValid(InActor))
-	{
-		TArray<FDMObjectMaterialProperty> ActorProperties = UDMMaterialModelFunctionLibrary::GetActorMaterialProperties(InActor);
-		ActorMaterialProperties.Reserve(ActorProperties.Num());
-
-		for (const FDMObjectMaterialProperty& ActorProperty : ActorProperties)
-		{
-			ActorMaterialProperties.Add(MakeShared<FDMObjectMaterialProperty>(ActorProperty));
-		}
-	}
 	
 	ChildSlot
 		.HAlign(HAlign_Fill)
@@ -84,8 +73,7 @@ TSharedRef<SWidget> SDMToolBar::CreateToolBarEntries()
 {
 	using namespace UE::DynamicMaterialEditor::Private;
 
-	return 
-		SNew(SHorizontalBox)
+	return SNew(SHorizontalBox)
 		
 		+ SHorizontalBox::Slot()
 		.FillWidth(1.0f)
@@ -309,7 +297,10 @@ void SDMToolBar::SetActorPropertySelected(AActor* InActor)
 		TArray<FDMObjectMaterialProperty> ActorProperties = UDMMaterialModelFunctionLibrary::GetActorMaterialProperties(InActor);
 		UDynamicMaterialModelBase* MaterialModelBase = GetMaterialModelBase();
 
-		for (int32 MaterialPropertyIdx = 0; MaterialPropertyIdx < ActorProperties.Num(); ++MaterialPropertyIdx)
+		const int32 ActorPropertyCount = ActorProperties.Num();
+		ActorMaterialProperties.Empty(ActorPropertyCount);
+
+		for (int32 MaterialPropertyIdx = 0; MaterialPropertyIdx < ActorPropertyCount; ++MaterialPropertyIdx)
 		{
 			const FDMObjectMaterialProperty& MaterialProperty = ActorProperties[MaterialPropertyIdx];
 
@@ -323,6 +314,7 @@ void SDMToolBar::SetActorPropertySelected(AActor* InActor)
 	}
 	else
 	{
+		ActorMaterialProperties.Empty(0);
 		ActorNameWidget->SetText(FText::GetEmpty());
 		ActorRowWidget->SetVisibility(EVisibility::Collapsed);
 	}
@@ -389,8 +381,7 @@ TSharedRef<SWidget> SDMToolBar::CreateToolBarButton(TAttribute<const FSlateBrush
 {
 	using namespace UE::DynamicMaterialEditor::Private;
 
-	return
-		SNew(SButton)
+	return SNew(SButton)
 		.ContentPadding(DefaultToolBarButtonContentPadding)
 		.ButtonStyle(FDynamicMaterialEditorStyle::Get(), "HoverHintOnly")
 		.ToolTipText(InTooltipText)
@@ -414,8 +405,7 @@ TSharedRef<SWidget> SDMToolBar::CreateSlotsComboBoxWidget()
 	const TSharedPtr<FDMObjectMaterialProperty> InitiallySelectedItem =
 		ActorMaterialProperties.IsValidIndex(SelectedMaterialElementIndex) ? ActorMaterialProperties[SelectedMaterialElementIndex] : nullptr;
 
-	return 
-		SNew(SComboBox<TSharedPtr<FDMObjectMaterialProperty>>)
+	return SNew(SComboBox<TSharedPtr<FDMObjectMaterialProperty>>)
 		.IsEnabled(ActorMaterialProperties.Num() > 1)
 		.InitiallySelectedItem(InitiallySelectedItem)
 		.OptionsSource(&ActorMaterialProperties)
diff --git a/Engine/Plugins/Experimental/DynamicMaterial/Source/DynamicMaterialEditor/Private/UI/Widgets/SDMActorMaterialSelector.cpp b/Engine/Plugins/Experimental/DynamicMaterial/Source/DynamicMaterialEditor/Private/UI/Widgets/SDMActorMaterialSelector.cpp
--- a/Engine/Plugins/Experimental/DynamicMaterial/Source/DynamicMaterialEditor/Private/UI/Widgets/SDMActorMaterialSelector.cpp
+++ b/Engine/Plugins/Experimental/DynamicMaterial/Source/DynamicMaterialEditor/Private/UI/Widgets/SDMActorMaterialSelector.cpp
@@ -212,11 +212,10 @@ FReply SDMActorMaterialSelector::OnCreateMaterialButtonClicked(int32 InActorProp
 
 	if (NewMaterialModel)
 	{
-		DesignerWidget->OnMaterialModelBaseSelected(NewMaterialModel);
+		DesignerWidget->OnObjectMaterialPropertySelected(ActorProperties[InActorPropertyIndex]);
 	}
 
 	return FReply::Handled();
-
 }
 
 #undef LOCTEXT_NAMESPACE
