Material Designer
- All of the global settings now appear as expected in the general settings.

diff --git a/Engine/Plugins/Experimental/DynamicMaterial/Source/DynamicMaterialEditor/Private/UI/PropertyGenerators/DMMaterialModelPropertyRowGenerator.cpp b/Engine/Plugins/Experimental/DynamicMaterial/Source/DynamicMaterialEditor/Private/UI/PropertyGenerators/DMMaterialModelPropertyRowGenerator.cpp
--- a/Engine/Plugins/Experimental/DynamicMaterial/Source/DynamicMaterialEditor/Private/UI/PropertyGenerators/DMMaterialModelPropertyRowGenerator.cpp
+++ b/Engine/Plugins/Experimental/DynamicMaterial/Source/DynamicMaterialEditor/Private/UI/PropertyGenerators/DMMaterialModelPropertyRowGenerator.cpp
@@ -88,11 +88,11 @@ EDMIterationResult FDMMaterialModelPropertyRowGenerator::ForEachProperty(EDMMate
 	const TSharedRef<SDMMaterialGlobalSettingsEditor>& InGlobalSettingEditorWidget, UDynamicMaterialModelBase* InMaterialModelBase, 
 	TArray<FDMPropertyHandle>& InOutPropertyRows, UDynamicMaterialModelEditorOnlyData* InEditorOnlyData)
 {
-	if (!InEditorOnlyData->GetSlotForMaterialProperty(InProperty))
+	if (InEditorOnlyData->GetSlotForMaterialProperty(InProperty))
 	{
 		if (UDMMaterialProperty* MaterialProperty = InEditorOnlyData->GetMaterialProperty(InProperty))
 		{
-			if (MaterialProperty->IsValidForModel(*InEditorOnlyData))
+			if (MaterialProperty->IsEnabled() && MaterialProperty->IsValidForModel(*InEditorOnlyData))
 			{
 				if (UDMMaterialValueFloat1* AlphaValue = Cast<UDMMaterialValueFloat1>(MaterialProperty->GetComponent(UDynamicMaterialModelEditorOnlyData::AlphaValueName)))
 				{
@@ -130,8 +130,11 @@ void FDMMaterialModelPropertyRowGenerator::AddGlobalValue(const TSharedRef<SDMMa
 		}
 	}
 
-	FDMPropertyHandle& ComponentHandle = InOutPropertyRows.Add_GetRef(FDMWidgetStatics::Get().GetPropertyHandle(&*InGlobalSettingEditorWidget,
-	                                                                                                            InComponent, UDMMaterialValue::ValueName));
+	FDMPropertyHandle& ComponentHandle = InOutPropertyRows.Add_GetRef(FDMWidgetStatics::Get().GetPropertyHandle(
+		&*InGlobalSettingEditorWidget,
+		InComponent,
+		UDMMaterialValue::ValueName
+	));
 
 	ComponentHandle.CategoryOverrideName = TEXT("Material Settings");
 	ComponentHandle.NameOverride = InNameOverride;
