Material Designer
- You can now drag multiple textures into the material designer layer list and it will again try to add them as texture sets.

diff --git a/Engine/Plugins/Experimental/DynamicMaterial/Source/DynamicMaterialEditor/Private/UI/Widgets/Editor/SDMMaterialSlotEditor.cpp b/Engine/Plugins/Experimental/DynamicMaterial/Source/DynamicMaterialEditor/Private/UI/Widgets/Editor/SDMMaterialSlotEditor.cpp
--- a/Engine/Plugins/Experimental/DynamicMaterial/Source/DynamicMaterialEditor/Private/UI/Widgets/Editor/SDMMaterialSlotEditor.cpp
+++ b/Engine/Plugins/Experimental/DynamicMaterial/Source/DynamicMaterialEditor/Private/UI/Widgets/Editor/SDMMaterialSlotEditor.cpp
@@ -18,6 +18,7 @@
 #include "CustomDetailsViewModule.h"
 #include "DetailLayoutBuilder.h"
 #include "DMTextureSet.h"
+#include "DMTextureSetBlueprintFunctionLibrary.h"
 #include "DMWorldSubsystem.h"
 #include "DynamicMaterialEditorStyle.h"
 #include "DynamicMaterialModule.h"
@@ -425,8 +426,9 @@ TSharedRef<SWidget> SDMMaterialSlotEditor::CreateSlot_Container()
 			.VAlign(VAlign_Fill)
 			[
 				 SAssignNew(DropTarget, SAssetDropTarget)
-				.OnAreAssetsAcceptableForDrop(this, &SDMMaterialSlotEditor::OnAssetDraggedOver)
+				.OnAreAssetsAcceptableForDrop(this, &SDMMaterialSlotEditor::OnAreAssetsAcceptableForDrop)
 				.OnAssetsDropped(this, &SDMMaterialSlotEditor::OnAssetsDropped)
+				.bSupportsMultiDrop(true)
 				//.bPlaceDropTargetOnTop(false)
 				[
 					SNew(SBox)
@@ -928,7 +930,7 @@ FReply SDMMaterialSlotEditor::OnLayerRowButtonsRemoveClicked()
 	return FReply::Handled();
 }
 
-bool SDMMaterialSlotEditor::OnAssetDraggedOver(TArrayView<FAssetData> InAssets)
+bool SDMMaterialSlotEditor::OnAreAssetsAcceptableForDrop(TArrayView<FAssetData> InAssets)
 {
 	TSharedPtr<SDMMaterialEditor> EditorWidget = GetEditorWidget();
 
@@ -1001,6 +1003,8 @@ void SDMMaterialSlotEditor::OnAssetsDropped(const FDragDropEvent& InDragDropEven
 		return;
 	}
 
+	TArray<FAssetData> DroppedTextures;
+
 	for (const FAssetData& Asset : InAssets)
 	{
 		UClass* AssetClass = Asset.GetClass(EResolveClass::Yes);
@@ -1012,22 +1016,31 @@ void SDMMaterialSlotEditor::OnAssetsDropped(const FDragDropEvent& InDragDropEven
 
 		if (AssetClass->IsChildOf(UTexture::StaticClass()))
 		{
-			HandleDrop_Texture(Cast<UTexture>(Asset.GetAsset()));
-			break;
+			DroppedTextures.Add(Asset);
+			continue;
 		}
 
 		if (AssetClass->IsChildOf(UDMTextureSet::StaticClass()))
 		{
 			HandleDrop_TextureSet(Cast<UDMTextureSet>(Asset.GetAsset()));
-			break;
+			return;
 		}
 
 		if (AssetClass->IsChildOf(UMaterialFunctionInterface::StaticClass()))
 		{
 			HandleDrop_MaterialFunction(Cast<UMaterialFunctionInterface>(Asset.GetAsset()));
-			break;
+			return;
 		}
 	}
+
+	if (DroppedTextures.Num() == 1)
+	{
+		HandleDrop_Texture(Cast<UTexture>(DroppedTextures[0].GetAsset()));
+	}
+	else if (DroppedTextures.Num() > 1)
+	{
+		HandleDrop_CreateTextureSet(DroppedTextures);
+	}
 }
 
 void SDMMaterialSlotEditor::HandleDrop_Texture(UTexture* InTexture)
@@ -1078,6 +1091,28 @@ void SDMMaterialSlotEditor::HandleDrop_Texture(UTexture* InTexture)
 	}
 }
 
+void SDMMaterialSlotEditor::HandleDrop_CreateTextureSet(const TArray<FAssetData>& InTextureAssets)
+{
+	if (InTextureAssets.Num() < 2)
+	{
+		return;
+	}
+
+	UDMTextureSetBlueprintFunctionLibrary::CreateTextureSetFromAssetsInteractive(
+		InTextureAssets,
+		FDMTextureSetBuilderOnComplete::CreateSPLambda(
+			this,
+			[this](UDMTextureSet* InTextureSet, bool bInWasAccepted)
+			{
+				if (bInWasAccepted)
+				{
+					HandleDrop_TextureSet(InTextureSet);
+				}
+			}
+		)
+	);
+}
+
 void SDMMaterialSlotEditor::HandleDrop_TextureSet(UDMTextureSet* InTextureSet)
 {
 	TSharedPtr<SDMMaterialEditor> EditorWidget = GetEditorWidget();
diff --git a/Engine/Plugins/Experimental/DynamicMaterial/Source/DynamicMaterialEditor/Private/UI/Widgets/Editor/SDMMaterialSlotEditor.h b/Engine/Plugins/Experimental/DynamicMaterial/Source/DynamicMaterialEditor/Private/UI/Widgets/Editor/SDMMaterialSlotEditor.h
--- a/Engine/Plugins/Experimental/DynamicMaterial/Source/DynamicMaterialEditor/Private/UI/Widgets/Editor/SDMMaterialSlotEditor.h
+++ b/Engine/Plugins/Experimental/DynamicMaterial/Source/DynamicMaterialEditor/Private/UI/Widgets/Editor/SDMMaterialSlotEditor.h
@@ -113,12 +113,14 @@ protected:
 	FReply OnLayerRowButtonsRemoveClicked();
 
 	/** Drag and drop. */
-	bool OnAssetDraggedOver(TArrayView<FAssetData> InAssets);
+	bool OnAreAssetsAcceptableForDrop(TArrayView<FAssetData> InAssets);
 
 	void OnAssetsDropped(const FDragDropEvent& InDragDropEvent, TArrayView<FAssetData> InAssets);
 
 	void HandleDrop_Texture(UTexture* InTexture);
 
+	void HandleDrop_CreateTextureSet(const TArray<FAssetData>& InTextureAssets);
+
 	void HandleDrop_TextureSet(UDMTextureSet* InTextureSet);
 
 	void HandleDrop_MaterialFunction(UMaterialFunctionInterface* InMaterialFunction);
