Material Designer
- Temporarily disabled material statistics.

diff --git a/Engine/Plugins/Experimental/DynamicMaterial/Source/DynamicMaterialEditor/Private/Model/DynamicMaterialModelEditorOnlyData.cpp b/Engine/Plugins/Experimental/DynamicMaterial/Source/DynamicMaterialEditor/Private/Model/DynamicMaterialModelEditorOnlyData.cpp
--- a/Engine/Plugins/Experimental/DynamicMaterial/Source/DynamicMaterialEditor/Private/Model/DynamicMaterialModelEditorOnlyData.cpp
+++ b/Engine/Plugins/Experimental/DynamicMaterial/Source/DynamicMaterialEditor/Private/Model/DynamicMaterialModelEditorOnlyData.cpp
@@ -587,7 +587,7 @@ void UDynamicMaterialModelEditorOnlyData::BuildMaterial(bool bInDirtyAssets)
 		Pair.Value->AddOutputProcessor(BuildState);
 	}
 
-	MaterialStats = UMaterialEditingLibrary::GetStatistics(MaterialModel->DynamicMaterial);
+	//MaterialStats = UMaterialEditingLibrary::GetStatistics(MaterialModel->DynamicMaterial);
 
 	State = EDMState::Idle;
 
