Material Designer
- Added a darker background to the left panel.
- Added a slight border around the different sections on the right.
- Added a scroll box to the global settings

diff --git a/Engine/Plugins/Experimental/DynamicMaterial/Source/DynamicMaterialEditor/Private/UI/Widgets/SDMMaterialEditor.cpp b/Engine/Plugins/Experimental/DynamicMaterial/Source/DynamicMaterialEditor/Private/UI/Widgets/SDMMaterialEditor.cpp
--- a/Engine/Plugins/Experimental/DynamicMaterial/Source/DynamicMaterialEditor/Private/UI/Widgets/SDMMaterialEditor.cpp
+++ b/Engine/Plugins/Experimental/DynamicMaterial/Source/DynamicMaterialEditor/Private/UI/Widgets/SDMMaterialEditor.cpp
@@ -25,6 +25,7 @@
 #include "UI/Widgets/Editor/SDMToolBar.h"
 #include "UI/Widgets/SDMMaterialDesigner.h"
 #include "Utils/DMMaterialModelFunctionLibrary.h"
+#include "Widgets/Layout/SScrollBox.h"
 #include "Widgets/Layout/SSplitter.h"
 #include "Widgets/SBoxPanel.h"
 #include "Widgets/SNullWidget.h"
@@ -57,6 +58,9 @@ namespace UE::DynamicMaterialEditor::Private
 			constexpr int32 ComponentEditor = 1;
 		}
 	}
+
+	const TCHAR* EditorDarkBackground = TEXT("Brushes.Title");
+	const TCHAR* EditorLightBackground = TEXT("Brushes.Header");
 }
 
 void SDMMaterialEditor::PrivateRegisterAttributes(FSlateAttributeDescriptor::FInitializer&)
@@ -238,6 +242,7 @@ void SDMMaterialEditor::EditSlot(UDMMaterialSlot* InSlot, bool bInForceRefresh)
 	RightSlot.Invalidate();
 
 	SlotEditorSlot.Invalidate();
+	SplitterSlot = nullptr;
 	SlotToEdit = InSlot;
 
 	ComponentEditorSlot.Invalidate();
@@ -269,6 +274,7 @@ void SDMMaterialEditor::EditComponent(UDMMaterialComponent* InComponent, bool bI
 	{
 		RightSlot.Invalidate();
 		SlotEditorSlot.Invalidate();
+		SplitterSlot = nullptr;
 	}
 
 	bGlobalSettingsMode = false;
@@ -288,6 +294,7 @@ void SDMMaterialEditor::EditGlobalSettings(bool bInForceRefresh)
 	{
 		RightSlot.Invalidate();
 		SlotEditorSlot.Invalidate();
+		SplitterSlot = nullptr;
 		ComponentEditorSlot.Invalidate();
 	}
 
@@ -563,6 +570,7 @@ void SDMMaterialEditor::ClearSlots()
 	MaterialPreviewSlot.ClearWidget();
 	PropertySelectorSlot.ClearWidget();
 	SlotEditorSlot.ClearWidget();
+	SplitterSlot = nullptr;
 	ComponentEditorSlot.ClearWidget();
 	StatusBarSlot.ClearWidget();
 }
@@ -646,25 +654,32 @@ TSharedRef<SWidget> SDMMaterialEditor::CreateSlot_Main()
 
 TSharedRef<SWidget> SDMMaterialEditor::CreateSlot_Left()
 {
+	using namespace UE::DynamicMaterialEditor::Private;
+
 	SVerticalBox::FSlot* MaterialPreviewSlotPtr = nullptr;
 	SVerticalBox::FSlot* PropertySelectorSlotPtr = nullptr;
 
-	TSharedRef<SVerticalBox> NewLeft = SNew(SVerticalBox)
-		+ SVerticalBox::Slot()
-		.Expose(MaterialPreviewSlotPtr)
-		.AutoHeight()
+	TSharedRef<SWidget> NewLeft = SNew(SBorder)
+		.BorderImage(FAppStyle::GetBrush(EditorDarkBackground))
 		.Padding(5.f)
 		[
-			SNullWidget::NullWidget
-		]
-
-		+ SVerticalBox::Slot()
-		.Expose(PropertySelectorSlotPtr)
-		.FillHeight(1.0f)
-		.Padding(5.f, 0.f, 5.f, 5.f)
-		[
-			SNullWidget::NullWidget
-		];
+			SNew(SVerticalBox)
+			+ SVerticalBox::Slot()
+			.Expose(MaterialPreviewSlotPtr)
+			.AutoHeight()
+			.Padding(0.f)
+			[
+				SNullWidget::NullWidget
+			]
+
+			+ SVerticalBox::Slot()
+			.Expose(PropertySelectorSlotPtr)
+			.FillHeight(1.0f)
+			.Padding(0.f, 5.f, 0.f, 0.f)
+			[
+				SNullWidget::NullWidget
+			]
+		];		
 
 	MaterialPreviewSlot = TDMWidgetSlot<SDMMaterialPreview>(MaterialPreviewSlotPtr, CreateSlot_Preview());
 	PropertySelectorSlot = TDMWidgetSlot<SDMMaterialPropertySelector>(PropertySelectorSlotPtr, CreateSlot_PropertySelector());
@@ -674,6 +689,8 @@ TSharedRef<SWidget> SDMMaterialEditor::CreateSlot_Left()
 
 TSharedRef<SWidget> SDMMaterialEditor::CreateSlot_Right()
 {
+	using namespace UE::DynamicMaterialEditor::Private;
+
 	const bool bHasSlotToEdit = SlotToEdit.IsValid();
 
 	if (!bGlobalSettingsMode && !bHasSlotToEdit)
@@ -685,21 +702,33 @@ TSharedRef<SWidget> SDMMaterialEditor::CreateSlot_Right()
 		bGlobalSettingsMode = false;
 	}
 
-	return bGlobalSettingsMode
-		? CreateSlot_Right_GlobalSettings()
-		: CreateSlot_Right_Slot();
+	return SNew(SBorder)
+		.BorderImage(FAppStyle::GetBrush(EditorDarkBackground))
+		.Padding(FMargin(0.f, 5.f))
+		[
+			bGlobalSettingsMode
+				? CreateSlot_Right_GlobalSettings()
+				: CreateSlot_Right_Slot()
+		];
 }
 
 TSharedRef<SWidget> SDMMaterialEditor::CreateSlot_Right_GlobalSettings()
 {
-	SVerticalBox::FSlot* GlobalSettingsSlotPtr = nullptr;
+	using namespace UE::DynamicMaterialEditor::Private;
 
-	TSharedRef<SVerticalBox> NewRight = SNew(SVerticalBox)
-		+ SVerticalBox::Slot()
-		.Expose(GlobalSettingsSlotPtr)
-		.AutoHeight()
+	SScrollBox::FSlot* GlobalSettingsSlotPtr = nullptr;
+
+	TSharedRef<SBorder> NewRight = SNew(SBorder)
+		.BorderImage(FAppStyle::GetBrush(EditorLightBackground))
+		.Padding(0.f)
 		[
-			SNullWidget::NullWidget
+			SNew(SScrollBox)
+			+ SScrollBox::Slot()
+			.Expose(GlobalSettingsSlotPtr)
+			.VAlign(EVerticalAlignment::VAlign_Fill)
+			[
+				SNullWidget::NullWidget
+			]
 		];
 
 	GlobalSettingsEditorSlot = TDMWidgetSlot<SDMMaterialGlobalSettingsEditor>(GlobalSettingsSlotPtr, CreateSlot_GlobalSettingsEditor());
@@ -714,6 +743,8 @@ TSharedRef<SDMMaterialGlobalSettingsEditor> SDMMaterialEditor::CreateSlot_Global
 
 TSharedRef<SWidget> SDMMaterialEditor::CreateSlot_Right_Slot()
 {
+	using namespace UE::DynamicMaterialEditor::Private;
+
 	float SplitterValue = 0.5;
 
 	if (UDynamicMaterialEditorSettings* Settings = UDynamicMaterialEditorSettings::Get())
@@ -721,39 +752,49 @@ TSharedRef<SWidget> SDMMaterialEditor::CreateSlot_Right_Slot()
 		SplitterValue = Settings->SplitterLocation;
 	}
 
-	SSplitter::FSlot* SlotEditorSlotPtr = nullptr;
-	SSplitter::FSlot* ComponentEditorSlotPtr = nullptr;
+	TSharedPtr<SBorder> TopBox;
+	TSharedPtr<SBorder> BottomBox;
+
+	SSplitter::FSlot* ExplosedSlot = nullptr;
 
 	TSharedRef<SSplitter> NewRight = SNew(SSplitter)
 		.Style(FAppStyle::Get(), "DetailsView.Splitter")
 		.Orientation(Orient_Vertical)
 		.ResizeMode(ESplitterResizeMode::Fill)
-		.PhysicalSplitterHandleSize(3.0f)
-		.HitDetectionSplitterHandleSize(4.0f)
+		.PhysicalSplitterHandleSize(5.0f)
+		.HitDetectionSplitterHandleSize(5.0f)
 		.OnSplitterFinishedResizing(this, &SDMMaterialEditor::OnRightSlotSplitterResized)
 
 		+ SSplitter::Slot()
-		.Expose(SlotEditorSlotPtr)
+		.Expose(ExplosedSlot)
 		.Resizable(true)
 		.SizeRule(SSplitter::ESizeRule::FractionOfParent)
-		.MinSize(50)
+		.MinSize(165)
 		.Value(SplitterValue)
 		[
-			SNullWidget::NullWidget
+			SAssignNew(TopBox, SBorder)
+			.BorderImage(FAppStyle::GetBrush(EditorLightBackground))
+			[
+				SNullWidget::NullWidget
+			]
 		]
 
 		+ SSplitter::Slot()
-		.Expose(ComponentEditorSlotPtr)
 		.Resizable(true)
 		.SizeRule(SSplitter::ESizeRule::FractionOfParent)
-		.MinSize(50)
+		.MinSize(60)
 		.Value(1.f - SplitterValue)
 		[
-			SNullWidget::NullWidget
+			SAssignNew(BottomBox, SBorder)
+			.BorderImage(FAppStyle::GetBrush(EditorLightBackground))
+			[
+				SNullWidget::NullWidget
+			]
 		];
 
-	SlotEditorSlot = TDMWidgetSlot<SDMMaterialSlotEditor>(SlotEditorSlotPtr, CreateSlot_SlotEditor());
-	ComponentEditorSlot = TDMWidgetSlot<SDMMaterialComponentEditor>(ComponentEditorSlotPtr, CreateSlot_ComponentEditor());
+	SplitterSlot = ExplosedSlot;
+	SlotEditorSlot = TDMWidgetSlot<SDMMaterialSlotEditor>(TopBox.ToSharedRef(), 0, CreateSlot_SlotEditor());
+	ComponentEditorSlot = TDMWidgetSlot<SDMMaterialComponentEditor>(BottomBox.ToSharedRef(), 0, CreateSlot_ComponentEditor());
 
 	return NewRight;
 }
@@ -856,9 +897,10 @@ void SDMMaterialEditor::OnRightSlotSplitterResized()
 {
 	UDynamicMaterialEditorSettings* Settings = UDynamicMaterialEditorSettings::Get();
 
-	if (SlotEditorSlot.GetSlot())
+	if (SplitterSlot)
 	{
-		Settings->SplitterLocation = static_cast<SSplitter::FSlot*>(SlotEditorSlot.GetSlot())->GetSizeValue();
+		const float SplitterLocation = static_cast<SSplitter::FSlot*>(SplitterSlot)->GetSizeValue();
+		Settings->SplitterLocation = SplitterLocation;
 		Settings->SaveConfig();
 	}
 }
diff --git a/Engine/Plugins/Experimental/DynamicMaterial/Source/DynamicMaterialEditor/Private/UI/Widgets/SDMMaterialEditor.h b/Engine/Plugins/Experimental/DynamicMaterial/Source/DynamicMaterialEditor/Private/UI/Widgets/SDMMaterialEditor.h
--- a/Engine/Plugins/Experimental/DynamicMaterial/Source/DynamicMaterialEditor/Private/UI/Widgets/SDMMaterialEditor.h
+++ b/Engine/Plugins/Experimental/DynamicMaterial/Source/DynamicMaterialEditor/Private/UI/Widgets/SDMMaterialEditor.h
@@ -10,6 +10,7 @@
 #include "UI/Utils/DMWidgetSlot.h"
 
 class FDMPreviewMaterialManager;
+class FSlotBase;
 class FUICommandList;
 class SDMMaterialComponentEditor;
 class SDMMaterialDesigner;
@@ -108,6 +109,7 @@ protected:
 	TDMWidgetSlot<SDMMaterialPreview> MaterialPreviewSlot;
 	TDMWidgetSlot<SDMMaterialPropertySelector> PropertySelectorSlot;
 	TDMWidgetSlot<SDMMaterialGlobalSettingsEditor> GlobalSettingsEditorSlot;
+	FSlotBase* SplitterSlot;
 	TDMWidgetSlot<SDMMaterialSlotEditor> SlotEditorSlot;
 	TDMWidgetSlot<SDMMaterialComponentEditor> ComponentEditorSlot;
 	TDMWidgetSlot<SDMStatusBar> StatusBarSlot;
