Material Designer
- Converted the Ctrl+Click behaviour for toggling Material Properties back to a checkbox.

diff --git a/Engine/Plugins/Experimental/DynamicMaterial/Source/DynamicMaterialEditor/Private/UI/Widgets/Editor/SDMMaterialPropertySelector.cpp b/Engine/Plugins/Experimental/DynamicMaterial/Source/DynamicMaterialEditor/Private/UI/Widgets/Editor/SDMMaterialPropertySelector.cpp
--- a/Engine/Plugins/Experimental/DynamicMaterial/Source/DynamicMaterialEditor/Private/UI/Widgets/Editor/SDMMaterialPropertySelector.cpp
+++ b/Engine/Plugins/Experimental/DynamicMaterial/Source/DynamicMaterialEditor/Private/UI/Widgets/Editor/SDMMaterialPropertySelector.cpp
@@ -14,6 +14,7 @@
 #include "Widgets/Images/SImage.h"
 #include "Widgets/Input/SCheckBox.h"
 #include "Widgets/Layout/SBox.h"
+#include "Widgets/Layout/SGridPanel.h"
 #include "Widgets/Layout/SScrollBox.h"
 #include "Widgets/SBoxPanel.h"
 #include "Widgets/SNullWidget.h"
@@ -21,6 +22,15 @@
 
 #define LOCTEXT_NAMESPACE "SDMMaterialPropertySelector"
 
+namespace UE::DynamicMaterialEditor::Private
+{
+	namespace PropertySelectorColumns
+	{
+		constexpr int32 Enable = 0;
+		constexpr int32 Select = 1;
+	}
+}
+
 void SDMMaterialPropertySelector::PrivateRegisterAttributes(FSlateAttributeDescriptor::FInitializer&)
 {
 }
@@ -114,7 +124,8 @@ TSharedRef<SWidget> SDMMaterialPropertySelector::CreateSlot_PropertyList()
 {
 	using namespace UE::DynamicMaterialEditor::Private;
 
-	TSharedRef<SVerticalBox> NewSlotList = SNew(SVerticalBox);
+	TSharedRef<SGridPanel> NewSlotList = SNew(SGridPanel)
+		.FillColumn(PropertySelectorColumns::Select, 1.f);
 
 	UDynamicMaterialModelEditorOnlyData* EditorOnlyData = GetEditorOnlyData();
 
@@ -123,12 +134,15 @@ TSharedRef<SWidget> SDMMaterialPropertySelector::CreateSlot_PropertyList()
 		return NewSlotList;
 	}
 
-	NewSlotList->AddSlot()
-		.AutoHeight()
+	int32 Row = 0;
+
+	NewSlotList->AddSlot(PropertySelectorColumns::Select, Row)
 		[
 			CreateSlot_SelectButton(EDMMaterialPropertyType::None)
 		];
 
+	++Row;
+
 	for (const TPair<EDMMaterialPropertyType, UDMMaterialProperty*>& PropertyPair : EditorOnlyData->GetMaterialProperties())
 	{
 		if (IsCustomMaterialProperty(PropertyPair.Key))
@@ -136,70 +150,82 @@ TSharedRef<SWidget> SDMMaterialPropertySelector::CreateSlot_PropertyList()
 			continue;
 		}
 
-		NewSlotList->AddSlot()
-			.AutoHeight()
+		NewSlotList->AddSlot(PropertySelectorColumns::Enable, Row)
+			[
+				CreateSlot_EnabledButton(PropertyPair.Key)
+			];
+
+		NewSlotList->AddSlot(PropertySelectorColumns::Select, Row)
 			[
 				CreateSlot_SelectButton(PropertyPair.Key)
 			];
+
+		++Row;
 	}
 
 	return NewSlotList;
 }
 
+TSharedRef<SWidget> SDMMaterialPropertySelector::CreateSlot_EnabledButton(EDMMaterialPropertyType InMaterialProperty)
+{
+	const FText Format = LOCTEXT("PropertyEnableFormat", "Toggle the {0} property.\n\nProperty must be valid for the Material Type.");
+	UEnum* MaterialPropertyEnum = StaticEnum<EDMMaterialPropertyType>();
+
+	const FText ToolTip = FText::Format(Format, MaterialPropertyEnum->GetDisplayNameTextByValue(static_cast<int64>(InMaterialProperty)));
+
+	return SNew(SCheckBox)
+		.IsEnabled(this, &SDMMaterialPropertySelector::GetPropertyEnabledEnabled, InMaterialProperty)
+		.IsChecked(this, &SDMMaterialPropertySelector::GetPropertyEnabledState, InMaterialProperty)
+		.OnCheckStateChanged(this, &SDMMaterialPropertySelector::OnPropertyEnabledStateChanged, InMaterialProperty)
+		.ToolTipText(ToolTip);
+}
+
 TSharedRef<SWidget> SDMMaterialPropertySelector::CreateSlot_SelectButton(EDMMaterialPropertyType InMaterialProperty)
 {
 	const FText ButtonText = InMaterialProperty == EDMMaterialPropertyType::None
 		? LOCTEXT("GlobalSettings", "Global Settings")
 		: StaticEnum<EDMMaterialPropertyType>()->GetDisplayNameTextByValue(static_cast<int64>(InMaterialProperty));
 
-	const FText Format = LOCTEXT("PropertyFormat", "Edit the {0} property.\n\n- Control+Left click: Toggle the property\n\nProperty may not be toggleable if it is not valid for the material type..");
+	const FText Format = LOCTEXT("PropertySelectFormat", "Edit the {0} property.");
 	UEnum* MaterialPropertyEnum = StaticEnum<EDMMaterialPropertyType>();
 
 	const FText ToolTip = (InMaterialProperty == EDMMaterialPropertyType::None)
 		? LOCTEXT("GeneralSettingsToolTip", "Edit the material global settings.")
 		: FText::Format(Format, MaterialPropertyEnum->GetDisplayNameTextByValue(static_cast<int64>(InMaterialProperty)));
 
-	TSharedRef<SBox> Outer = SNew(SBox)
-		[		
-			SNew(SCheckBox)
-			.Style(FAppStyle::Get(), "DetailsView.SectionButton")
-			.HAlign(EHorizontalAlignment::HAlign_Center)
-			.Padding(0.f)
-			.IsEnabled(this, &SDMMaterialPropertySelector::GetPropertySelectEnabled, InMaterialProperty)
-			.IsChecked(this, &SDMMaterialPropertySelector::GetPropertySelectState, InMaterialProperty)
-			.OnCheckStateChanged(this, &SDMMaterialPropertySelector::OnPropertySelectStateChanged, InMaterialProperty)
-			.ToolTipText(ToolTip)
-			.Content()
+	return SNew(SCheckBox)
+		.Style(FAppStyle::Get(), "DetailsView.SectionButton")
+		.HAlign(EHorizontalAlignment::HAlign_Center)
+		.Padding(0.f)
+		.IsEnabled(this, &SDMMaterialPropertySelector::GetPropertySelectEnabled, InMaterialProperty)
+		.IsChecked(this, &SDMMaterialPropertySelector::GetPropertySelectState, InMaterialProperty)
+		.OnCheckStateChanged(this, &SDMMaterialPropertySelector::OnPropertySelectStateChanged, InMaterialProperty)
+		.ToolTipText(ToolTip)
+		.Content()
+		[
+			SNew(SBox)
+			.WidthOverride(135.f)
 			[
-				SNew(SBox)
-				.WidthOverride(135.f)
+				SNew(SHorizontalBox)
+				+SHorizontalBox::Slot()
+				.VAlign(VAlign_Center)
+				.AutoWidth()
+				[
+					SNew(SImage)
+					.Image(FAppStyle::Get().GetBrush("FilterBar.FilterImage"))
+					.ColorAndOpacity(this, &SDMMaterialPropertySelector::GetPropertySelectButtonChipColor, InMaterialProperty)
+				]
+				+SHorizontalBox::Slot()
+				.Padding(10.f, 6.f)
+				.VAlign(VAlign_Center)
+				.FillWidth(1.f)
 				[
-					SNew(SHorizontalBox)
-					+SHorizontalBox::Slot()
-					.VAlign(VAlign_Center)
-					.AutoWidth()
-					[
-						SNew(SImage)
-						.Image(FAppStyle::Get().GetBrush("FilterBar.FilterImage"))
-						.ColorAndOpacity(this, &SDMMaterialPropertySelector::GetPropertySelectButtonChipColor, InMaterialProperty)
-					]
-					+SHorizontalBox::Slot()
-					.Padding(10.f, 6.f)
-					.VAlign(VAlign_Center)
-					.FillWidth(1.f)
-					[
-						SNew(STextBlock)
-						.Font(IDetailLayoutBuilder::GetDetailFont())
-						.Text(ButtonText)
-					]
+					SNew(STextBlock)
+					.Font(IDetailLayoutBuilder::GetDetailFont())
+					.Text(ButtonText)
 				]
 			]
 		];
-
-	// This is triggered when the inner checkbox is disabled.
-	Outer->SetOnMouseButtonDown(FPointerEventHandler::CreateSP(this, &SDMMaterialPropertySelector::OnPropertySelectMouseDown, InMaterialProperty));
-
-	return Outer;
 }
 
 bool SDMMaterialPropertySelector::IsPropertyEnabled(EDMMaterialPropertyType InMaterialProperty) const
@@ -275,6 +301,45 @@ bool SDMMaterialPropertySelector::SetPropertyEnabled(EDMMaterialPropertyType InM
 	return !!EditorOnlyData->AddSlotForMaterialProperty(InMaterialProperty);
 }
 
+bool SDMMaterialPropertySelector::GetPropertyEnabledEnabled(EDMMaterialPropertyType InMaterialProperty) const
+{
+	UDynamicMaterialModelEditorOnlyData* EditorOnlyData = GetEditorOnlyData();
+
+	if (!EditorOnlyData)
+	{
+		return false;
+	}
+
+	UDMMaterialProperty* MaterialProperty = EditorOnlyData->GetMaterialProperty(InMaterialProperty);
+
+	if (!MaterialProperty)
+	{
+		return false;
+	}
+
+	return MaterialProperty->IsValidForModel(*EditorOnlyData);
+}
+
+ECheckBoxState SDMMaterialPropertySelector::GetPropertyEnabledState(EDMMaterialPropertyType InMaterialProperty) const
+{
+	return DoesPropertySlotExist(InMaterialProperty)
+		? ECheckBoxState::Checked
+		: ECheckBoxState::Unchecked;
+}
+
+void SDMMaterialPropertySelector::OnPropertyEnabledStateChanged(ECheckBoxState InState, EDMMaterialPropertyType InMaterialProperty)
+{
+	const bool bSetEnabled = InState == ECheckBoxState::Checked;
+
+	if (SetPropertyEnabled(InMaterialProperty, /* Enabled */ bSetEnabled))
+	{
+		if (bSetEnabled)
+		{
+			SetSelectedProperty(InMaterialProperty);
+		}
+	}
+}
+
 bool SDMMaterialPropertySelector::GetPropertySelectEnabled(EDMMaterialPropertyType InMaterialProperty) const
 {
 	if (InMaterialProperty == EDMMaterialPropertyType::None)
@@ -300,15 +365,6 @@ void SDMMaterialPropertySelector::OnPropertySelectStateChanged(ECheckBoxState In
 		return;
 	}
 
-
-	const FSlateApplication& SlateApplication = FSlateApplication::Get();
-
-	if (SlateApplication.GetModifierKeys().IsControlDown() || SlateApplication.GetModifierKeys().IsCommandDown())
-	{
-		SetPropertyEnabled(InMaterialProperty, /* Enabled */ false);
-		return;
-	}
-
 	UDynamicMaterialModelEditorOnlyData* EditorOnlyData = GetEditorOnlyData();
 
 	if (!EditorOnlyData)
@@ -346,22 +402,6 @@ FSlateColor SDMMaterialPropertySelector::GetPropertySelectButtonChipColor(EDMMat
 	return FStyleColors::Panel;
 }
 
-FReply SDMMaterialPropertySelector::OnPropertySelectMouseDown(const FGeometry& InGeometry, const FPointerEvent& InEvent, 
-	EDMMaterialPropertyType InMaterialProperty)
-{
-	if (!InEvent.GetModifierKeys().IsControlDown() && !InEvent.GetModifierKeys().IsCommandDown())
-	{
-		return FReply::Unhandled();
-	}
-
-	if (SetPropertyEnabled(InMaterialProperty, /* Enabled */ true))
-	{
-		SetSelectedProperty(InMaterialProperty);
-	}
-
-	return FReply::Handled();
-}
-
 void SDMMaterialPropertySelector::OnSelectedPropertyChanged()
 {
 	TSharedPtr<SDMMaterialEditor> EditorWidget = GetEditorWidget();
diff --git a/Engine/Plugins/Experimental/DynamicMaterial/Source/DynamicMaterialEditor/Private/UI/Widgets/Editor/SDMMaterialPropertySelector.h b/Engine/Plugins/Experimental/DynamicMaterial/Source/DynamicMaterialEditor/Private/UI/Widgets/Editor/SDMMaterialPropertySelector.h
--- a/Engine/Plugins/Experimental/DynamicMaterial/Source/DynamicMaterialEditor/Private/UI/Widgets/Editor/SDMMaterialPropertySelector.h
+++ b/Engine/Plugins/Experimental/DynamicMaterial/Source/DynamicMaterialEditor/Private/UI/Widgets/Editor/SDMMaterialPropertySelector.h
@@ -46,6 +46,8 @@ protected:
 
 	TSharedRef<SWidget> CreateSlot_PropertyList();
 
+	TSharedRef<SWidget> CreateSlot_EnabledButton(EDMMaterialPropertyType InMaterialProperty);
+
 	TSharedRef<SWidget> CreateSlot_SelectButton(EDMMaterialPropertyType InMaterialProperty);
 
 	bool IsPropertyEnabled(EDMMaterialPropertyType InMaterialProperty) const;
@@ -54,12 +56,14 @@ protected:
 
 	bool DoesPropertySlotExist(EDMMaterialPropertyType InMaterialProperty) const;
 
+	bool GetPropertyEnabledEnabled(EDMMaterialPropertyType InMaterialProperty) const;
+	ECheckBoxState GetPropertyEnabledState(EDMMaterialPropertyType InMaterialProperty) const;
+	void OnPropertyEnabledStateChanged(ECheckBoxState InState, EDMMaterialPropertyType InMaterialProperty);
+
 	bool GetPropertySelectEnabled(EDMMaterialPropertyType InMaterialProperty) const;
 	ECheckBoxState GetPropertySelectState(EDMMaterialPropertyType InMaterialProperty) const;
 	void OnPropertySelectStateChanged(ECheckBoxState InState, EDMMaterialPropertyType InMaterialProperty);
 	FSlateColor GetPropertySelectButtonChipColor(EDMMaterialPropertyType InMaterialProperty) const;
 
-	FReply OnPropertySelectMouseDown(const FGeometry& InGeometry, const FPointerEvent& InEvent, EDMMaterialPropertyType InMaterialProperty);
-
 	void OnSelectedPropertyChanged();
 };
