Material Designer
- Added back in the thumbnail size option.

diff --git a/Engine/Plugins/Experimental/DynamicMaterial/Source/DynamicMaterialEditor/Private/DynamicMaterialEditorSettings.cpp b/Engine/Plugins/Experimental/DynamicMaterial/Source/DynamicMaterialEditor/Private/DynamicMaterialEditorSettings.cpp
--- a/Engine/Plugins/Experimental/DynamicMaterial/Source/DynamicMaterialEditor/Private/DynamicMaterialEditorSettings.cpp
+++ b/Engine/Plugins/Experimental/DynamicMaterial/Source/DynamicMaterialEditor/Private/DynamicMaterialEditorSettings.cpp
@@ -282,6 +282,7 @@ void UDynamicMaterialEditorSettings::OpenEditorSettingsWindow() const
 void UDynamicMaterialEditorSettings::ResetAllLayoutSettings()
 {
 	SplitterLocation = 240;
+	ThumbnailSize = 256.f;
 	PreviewMesh = EDMMaterialPreviewMesh::Plane;
 	bShowPreviewBackground = true;
 	bPreviewImagesUseTextureUVs = true;
diff --git a/Engine/Plugins/Experimental/DynamicMaterial/Source/DynamicMaterialEditor/Private/DynamicMaterialEditorSettings.h b/Engine/Plugins/Experimental/DynamicMaterial/Source/DynamicMaterialEditor/Private/DynamicMaterialEditorSettings.h
--- a/Engine/Plugins/Experimental/DynamicMaterial/Source/DynamicMaterialEditor/Private/DynamicMaterialEditorSettings.h
+++ b/Engine/Plugins/Experimental/DynamicMaterial/Source/DynamicMaterialEditor/Private/DynamicMaterialEditorSettings.h
@@ -6,7 +6,6 @@
 #include "DMDefs.h"
 #include "Engine/DeveloperSettings.h"
 #include "Engine/Texture.h"
-#include "MaterialDomain.h"
 #include "PropertyEditorDelegates.h"
 #include "UObject/SoftObjectPtr.h"
 #include "DynamicMaterialEditorSettings.generated.h"
@@ -155,6 +154,9 @@ public:
 	UPROPERTY(Config, EditAnywhere, BlueprintReadWrite, Category = "Layout")
 	bool bUVVisualizerVisible;
 
+	UPROPERTY(Config, EditAnywhere, BlueprintReadWrite, Category = "Thumnnails")
+	double ThumbnailSize;
+
 	UPROPERTY(Config, EditAnywhere, BlueprintReadWrite, Category = "Preview")
 	EDMMaterialPreviewMesh PreviewMesh;
 
diff --git a/Engine/Plugins/Experimental/DynamicMaterial/Source/DynamicMaterialEditor/Private/UI/Widgets/Editor/SlotEditor/SDMMaterialStage.cpp b/Engine/Plugins/Experimental/DynamicMaterial/Source/DynamicMaterialEditor/Private/UI/Widgets/Editor/SlotEditor/SDMMaterialStage.cpp
--- a/Engine/Plugins/Experimental/DynamicMaterial/Source/DynamicMaterialEditor/Private/UI/Widgets/Editor/SlotEditor/SDMMaterialStage.cpp
+++ b/Engine/Plugins/Experimental/DynamicMaterial/Source/DynamicMaterialEditor/Private/UI/Widgets/Editor/SlotEditor/SDMMaterialStage.cpp
@@ -11,6 +11,7 @@
 #include "Components/MaterialStageInputs/DMMSIExpression.h"
 #include "Components/MaterialStageInputs/DMMSIValue.h"
 #include "Components/MaterialValues/DMMaterialValueTexture.h"
+#include "DynamicMaterialEditorSettings.h"
 #include "DynamicMaterialEditorStyle.h"
 #include "Engine/Texture.h"
 #include "SAssetDropTarget.h"
@@ -104,13 +105,21 @@ void SDMMaterialStage::Construct(const FArguments& InArgs, const TSharedRef<SDMM
 		]
 	];
 
+	UDynamicMaterialEditorSettings* Settings = UDynamicMaterialEditorSettings::Get();
+
 	SetToolTip(
 		SNew(SToolTip)
 		.IsInteractive(false)
 		.BorderImage(FCoreStyle::Get().GetBrush("ToolTip.Background"))
 		[
 			SNew(SDMMaterialComponentPreview, EditorWidget.ToSharedRef(), InStage)
-			.PreviewSize(FVector2D(256.f))
+			.PreviewSize(TAttribute<TOptional<FVector2D>>::CreateWeakLambda(
+				Settings,
+				[Settings]()
+				{
+					return FVector2D(Settings->ThumbnailSize, Settings->ThumbnailSize);
+				}
+			))
 		]
 	);
 }
diff --git a/Engine/Plugins/Experimental/DynamicMaterial/Source/DynamicMaterialEditor/Private/UI/Widgets/Visualizers/SDMMaterialComponentPreview.cpp b/Engine/Plugins/Experimental/DynamicMaterial/Source/DynamicMaterialEditor/Private/UI/Widgets/Visualizers/SDMMaterialComponentPreview.cpp
--- a/Engine/Plugins/Experimental/DynamicMaterial/Source/DynamicMaterialEditor/Private/UI/Widgets/Visualizers/SDMMaterialComponentPreview.cpp
+++ b/Engine/Plugins/Experimental/DynamicMaterial/Source/DynamicMaterialEditor/Private/UI/Widgets/Visualizers/SDMMaterialComponentPreview.cpp
@@ -79,12 +79,8 @@ void SDMMaterialComponentPreview::Construct(const FArguments& InArgs, const TSha
 		return;
 	}
 
-	PreviewSize = InArgs._PreviewSize;
-
 	SetCanTick(true);
 
-	Brush.SetImageSize(PreviewSize.Get());
-
 	if (ensure(IsValid(InComponent)))
 	{
 		InComponent->GetOnUpdate().AddSP(this, &SDMMaterialComponentPreview::OnComponentUpdated);
@@ -95,7 +91,7 @@ void SDMMaterialComponentPreview::Construct(const FArguments& InArgs, const TSha
 	[
 		SNew(SImage)
 		.Image(&Brush)
-		.DesiredSizeOverride(this, &SDMMaterialComponentPreview::GetPreviewSize)
+		.DesiredSizeOverride(InArgs._PreviewSize)
 	];
 }
 
@@ -149,11 +145,6 @@ void SDMMaterialComponentPreview::OnComponentUpdated(UDMMaterialComponent* InCom
 	}
 }
 
-TOptional<FVector2D> SDMMaterialComponentPreview::GetPreviewSize() const
-{
-	return PreviewSize.Get();
-}
-
 void SDMMaterialComponentPreview::OnValueUpdated(UDynamicMaterialModel* InMaterialModel, UDMMaterialValue* InValue)
 {
 	if (UMaterialInstanceDynamic* PreviewMaterialDynamic = PreviewMaterialDynamicWeak.Get())
diff --git a/Engine/Plugins/Experimental/DynamicMaterial/Source/DynamicMaterialEditor/Private/UI/Widgets/Visualizers/SDMMaterialComponentPreview.h b/Engine/Plugins/Experimental/DynamicMaterial/Source/DynamicMaterialEditor/Private/UI/Widgets/Visualizers/SDMMaterialComponentPreview.h
--- a/Engine/Plugins/Experimental/DynamicMaterial/Source/DynamicMaterialEditor/Private/UI/Widgets/Visualizers/SDMMaterialComponentPreview.h
+++ b/Engine/Plugins/Experimental/DynamicMaterial/Source/DynamicMaterialEditor/Private/UI/Widgets/Visualizers/SDMMaterialComponentPreview.h
@@ -23,7 +23,7 @@ public:
 	SLATE_BEGIN_ARGS(SDMMaterialComponentPreview)
 		: _PreviewSize(FVector2D(48.f, 48.f))
 		{}
-		SLATE_ATTRIBUTE(FVector2D, PreviewSize)
+		SLATE_ATTRIBUTE(TOptional<FVector2D>, PreviewSize)
 	SLATE_END_ARGS()
 
 	SDMMaterialComponentPreview();
@@ -45,12 +45,9 @@ protected:
 	TWeakObjectPtr<UMaterial> PreviewMaterialBaseWeak;
 	TWeakObjectPtr<UMaterialInstanceDynamic> PreviewMaterialDynamicWeak;
 	FSlateMaterialBrush Brush;
-	TAttribute<FVector2D> PreviewSize;
 
 	void OnComponentUpdated(UDMMaterialComponent* InComponent, EDMUpdateType InUpdateType);
 
-	TOptional<FVector2D> GetPreviewSize() const;
-
 	void OnValueUpdated(UDynamicMaterialModel* InMaterialModel, UDMMaterialValue* InValue);
 
 	void OnTextureUVUpdated(UDynamicMaterialModel* InMaterialModel, UDMTextureUV* InTextureUV);
