Material Designer
- Added the Displacement, Subsurface Color and Surface Thickness properties.

diff --git a/Engine/Plugins/Experimental/DynamicMaterial/Source/DynamicMaterial/Private/DynamicMaterialModule.cpp b/Engine/Plugins/Experimental/DynamicMaterial/Source/DynamicMaterial/Private/DynamicMaterialModule.cpp
--- a/Engine/Plugins/Experimental/DynamicMaterial/Source/DynamicMaterial/Private/DynamicMaterialModule.cpp
+++ b/Engine/Plugins/Experimental/DynamicMaterial/Source/DynamicMaterial/Private/DynamicMaterialModule.cpp
@@ -10,10 +10,6 @@
 
 DEFINE_LOG_CATEGORY(LogDynamicMaterial);
 
-#if WITH_EDITOR
-FDMCreateEditorOnlyDataDelegate FDynamicMaterialModule::CreateEditorOnlyDataDelegate;
-#endif
-
 static TAutoConsoleVariable<bool> CVarExportMaterials(
 	TEXT("DM.ExportMaterials"),
 	false,
@@ -35,16 +31,4 @@ FDynamicMaterialModule& FDynamicMaterialModule::Get()
 	return FModuleManager::LoadModuleChecked<FDynamicMaterialModule>("DynamicMaterial");
 }
 
-#if WITH_EDITOR
-TScriptInterface<IDynamicMaterialModelEditorOnlyDataInterface> FDynamicMaterialModule::CreateEditorOnlyData(UDynamicMaterialModel* InMaterialModel)
-{
-	if (CreateEditorOnlyDataDelegate.IsBound())
-	{
-		return CreateEditorOnlyDataDelegate.Execute(InMaterialModel);
-	}
-
-	return nullptr;
-}
-#endif
-
 IMPLEMENT_MODULE(FDynamicMaterialModule, DynamicMaterial)
diff --git a/Engine/Plugins/Experimental/DynamicMaterial/Source/DynamicMaterial/Private/Model/DynamicMaterialModel.cpp b/Engine/Plugins/Experimental/DynamicMaterial/Source/DynamicMaterial/Private/Model/DynamicMaterialModel.cpp
--- a/Engine/Plugins/Experimental/DynamicMaterial/Source/DynamicMaterial/Private/Model/DynamicMaterialModel.cpp
+++ b/Engine/Plugins/Experimental/DynamicMaterial/Source/DynamicMaterial/Private/Model/DynamicMaterialModel.cpp
@@ -39,6 +39,12 @@ const FLazyName UDynamicMaterialModel::GlobalRefractionValueName = FLazyName(TEX
 const FLazyName UDynamicMaterialModel::GlobalRefractionParameterName = FLazyName(TEXT("GlobalRefraction"));
 const FLazyName UDynamicMaterialModel::GlobalPixelDepthOffsetValueName = FLazyName(TEXT("GlobalPixelDepthOffsetValue"));
 const FLazyName UDynamicMaterialModel::GlobalPixelDepthOffsetParameterName = FLazyName(TEXT("GlobalPixelDepthOffset"));
+const FLazyName UDynamicMaterialModel::GlobalDisplacementValueName = FLazyName(TEXT("GlobalDisplacementValue"));
+const FLazyName UDynamicMaterialModel::GlobalDisplacementParameterName = FLazyName(TEXT("GlobalDisplacement"));
+const FLazyName UDynamicMaterialModel::GlobalSubsurfaceColorValueName = FLazyName(TEXT("GlobalSubsurfaceColorValue"));
+const FLazyName UDynamicMaterialModel::GlobalSubsurfaceColorParameterName = FLazyName(TEXT("GlobalSubsurfaceColor"));
+const FLazyName UDynamicMaterialModel::GlobalSurfaceThicknessValueName = FLazyName(TEXT("GlobalSurfaceThicknessValue"));
+const FLazyName UDynamicMaterialModel::GlobalSurfaceThicknessParameterName = FLazyName(TEXT("GlobalSurfaceThickness"));
 const FLazyName UDynamicMaterialModel::GlobalOffsetValueName = FLazyName(TEXT("GlobalOffsetValue"));
 const FLazyName UDynamicMaterialModel::GlobalOffsetParameterName = FLazyName(TEXT("GlobalOffset"));
 const FLazyName UDynamicMaterialModel::GlobalTilingValueName = FLazyName(TEXT("GlobalTilingValue"));
@@ -101,6 +107,7 @@ UDynamicMaterialModel::UDynamicMaterialModel()
 	AddFloatParameter(GlobalAmbientOcclusionValueName, GlobalAmbientOcclusionParameterName);
 	AddFloatParameter(GlobalRefractionValueName, GlobalRefractionParameterName);
 	AddFloatParameter(GlobalPixelDepthOffsetValueName, GlobalPixelDepthOffsetParameterName);
+	AddFloatParameter(GlobalDisplacementValueName, GlobalDisplacementParameterName);
 
 	AddVector2Parameter(GlobalOffsetValueName, GlobalOffsetParameterName, FVector2D::ZeroVector);
 	AddVector2Parameter(GlobalTilingValueName, GlobalTilingParameterName, FVector2D::UnitVector);
@@ -525,14 +532,6 @@ void UDynamicMaterialModel::PostLoad()
 
 	SetFlags(RF_Transactional);
 
-	IDynamicMaterialModelEditorOnlyDataInterface* ModelEditorOnlyData = GetEditorOnlyData();
-
-	if (!ModelEditorOnlyData)
-	{
-		EditorOnlyDataSI = FDynamicMaterialModule::CreateEditorOnlyData(this);
-		EditorOnlyDataSI->LoadDeprecatedModelData(this);
-	}
-
 	ReinitComponents();
 #endif
 }
@@ -728,6 +727,7 @@ void UDynamicMaterialModel::FixGlobalVars()
 	FixGlobalVar(GlobalAmbientOcclusionValueName, GlobalAmbientOcclusionParameterName);
 	FixGlobalVar(GlobalRefractionValueName, GlobalRefractionParameterName);
 	FixGlobalVar(GlobalPixelDepthOffsetValueName, GlobalPixelDepthOffsetParameterName);
+	FixGlobalVar(GlobalDisplacementValueName, GlobalDisplacementParameterName);
 	FixGlobalVar(GlobalOffsetValueName, GlobalOffsetParameterName);
 	FixGlobalVar(GlobalTilingValueName, GlobalTilingParameterName);
 	FixGlobalVar(GlobalRotationValueName, GlobalRotationParameterName);
diff --git a/Engine/Plugins/Experimental/DynamicMaterial/Source/DynamicMaterial/Private/Utils/DMUtils.cpp b/Engine/Plugins/Experimental/DynamicMaterial/Source/DynamicMaterial/Private/Utils/DMUtils.cpp
--- a/Engine/Plugins/Experimental/DynamicMaterial/Source/DynamicMaterial/Private/Utils/DMUtils.cpp
+++ b/Engine/Plugins/Experimental/DynamicMaterial/Source/DynamicMaterial/Private/Utils/DMUtils.cpp
@@ -73,6 +73,15 @@ EDMTextureSetMaterialProperty FDMUtils::MaterialPropertyTypeToTextureSetMaterial
 		case EDMMaterialPropertyType::Tangent:
 			return EDMTextureSetMaterialProperty::Tangent;
 
+		case EDMMaterialPropertyType::Displacement:
+			return EDMTextureSetMaterialProperty::Displacement;
+
+		case EDMMaterialPropertyType::SubsurfaceColor:
+			return EDMTextureSetMaterialProperty::SubsurfaceColor;
+
+		case EDMMaterialPropertyType::SurfaceThickness:
+			return EDMTextureSetMaterialProperty::SurfaceThickness;
+
 		default:
 			return EDMTextureSetMaterialProperty::None;
 	}
@@ -124,6 +133,15 @@ EDMMaterialPropertyType FDMUtils::TextureSetMaterialPropertyToMaterialPropertyTy
 		case EDMTextureSetMaterialProperty::Tangent:
 			return EDMMaterialPropertyType::Tangent;
 
+		case EDMTextureSetMaterialProperty::Displacement:
+			return EDMMaterialPropertyType::Displacement;
+
+		case EDMTextureSetMaterialProperty::SubsurfaceColor:
+			return EDMMaterialPropertyType::SubsurfaceColor;
+
+		case EDMTextureSetMaterialProperty::SurfaceThickness:
+			return EDMMaterialPropertyType::SurfaceThickness;
+
 		default:
 			return EDMMaterialPropertyType::None;
 	}
@@ -175,6 +193,15 @@ EMaterialProperty FDMUtils::MaterialPropertyTypeToMaterialProperty(EDMMaterialPr
 		case EDMMaterialPropertyType::Tangent:
 			return EMaterialProperty::MP_Tangent;
 
+		case EDMMaterialPropertyType::Displacement:
+			return EMaterialProperty::MP_Displacement;
+
+		case EDMMaterialPropertyType::SubsurfaceColor:
+			return EMaterialProperty::MP_SubsurfaceColor;
+
+		case EDMMaterialPropertyType::SurfaceThickness:
+			return EMaterialProperty::MP_SurfaceThickness;
+
 		default:
 			return EMaterialProperty::MP_MAX;
 	}
@@ -226,6 +253,15 @@ EDMMaterialPropertyType FDMUtils::MaterialPropertyToMaterialPropertyType(EMateri
 		case EMaterialProperty::MP_Tangent:
 			return EDMMaterialPropertyType::Tangent;
 
+		case EMaterialProperty::MP_Displacement:
+			return EDMMaterialPropertyType::Displacement;
+
+		case EMaterialProperty::MP_SubsurfaceColor:
+			return EDMMaterialPropertyType::SubsurfaceColor;
+
+		case EMaterialProperty::MP_SurfaceThickness:
+			return EDMMaterialPropertyType::SurfaceThickness;
+
 		default:
 			return EDMMaterialPropertyType::None;
 	}
diff --git a/Engine/Plugins/Experimental/DynamicMaterial/Source/DynamicMaterial/Public/DMDefs.h b/Engine/Plugins/Experimental/DynamicMaterial/Source/DynamicMaterial/Public/DMDefs.h
--- a/Engine/Plugins/Experimental/DynamicMaterial/Source/DynamicMaterial/Public/DMDefs.h
+++ b/Engine/Plugins/Experimental/DynamicMaterial/Source/DynamicMaterial/Public/DMDefs.h
@@ -39,6 +39,9 @@ enum class EDMMaterialPropertyType : uint8
 	Custom2,
 	Custom3,
 	Custom4,
+	Displacement,
+	SubsurfaceColor,
+	SurfaceThickness,
 	Any
 };
 
diff --git a/Engine/Plugins/Experimental/DynamicMaterial/Source/DynamicMaterial/Public/DynamicMaterialModule.h b/Engine/Plugins/Experimental/DynamicMaterial/Source/DynamicMaterial/Public/DynamicMaterialModule.h
--- a/Engine/Plugins/Experimental/DynamicMaterial/Source/DynamicMaterial/Public/DynamicMaterialModule.h
+++ b/Engine/Plugins/Experimental/DynamicMaterial/Source/DynamicMaterial/Public/DynamicMaterialModule.h
@@ -35,15 +35,4 @@ public:
 
 	/** Returns true if the material export flag has been enabled. @See DM.ExportMaterials */
 	DYNAMICMATERIAL_API static bool IsMaterialExportEnabled();
-
-#if WITH_EDITOR
-	DYNAMICMATERIAL_API static TScriptInterface<IDynamicMaterialModelEditorOnlyDataInterface> CreateEditorOnlyData(UDynamicMaterialModel* InMaterialModel);
-	static FDMCreateEditorOnlyDataDelegate::RegistrationType& GetCreateEditorOnlyDataDelegate() { return CreateEditorOnlyDataDelegate; }
-#endif
-
-protected:
-#if WITH_EDITOR
-	DYNAMICMATERIAL_API static FDMCreateEditorOnlyDataDelegate CreateEditorOnlyDataDelegate;
-#endif
-
 };
diff --git a/Engine/Plugins/Experimental/DynamicMaterial/Source/DynamicMaterial/Public/Model/DynamicMaterialModel.h b/Engine/Plugins/Experimental/DynamicMaterial/Source/DynamicMaterial/Public/Model/DynamicMaterialModel.h
--- a/Engine/Plugins/Experimental/DynamicMaterial/Source/DynamicMaterial/Public/Model/DynamicMaterialModel.h
+++ b/Engine/Plugins/Experimental/DynamicMaterial/Source/DynamicMaterial/Public/Model/DynamicMaterialModel.h
@@ -65,6 +65,12 @@ public:
 	DYNAMICMATERIAL_API static const FLazyName GlobalRefractionParameterName;
 	DYNAMICMATERIAL_API static const FLazyName GlobalPixelDepthOffsetValueName;
 	DYNAMICMATERIAL_API static const FLazyName GlobalPixelDepthOffsetParameterName;
+	DYNAMICMATERIAL_API static const FLazyName GlobalDisplacementValueName;
+	DYNAMICMATERIAL_API static const FLazyName GlobalDisplacementParameterName;
+	DYNAMICMATERIAL_API static const FLazyName GlobalSubsurfaceColorValueName;
+	DYNAMICMATERIAL_API static const FLazyName GlobalSubsurfaceColorParameterName;
+	DYNAMICMATERIAL_API static const FLazyName GlobalSurfaceThicknessValueName;
+	DYNAMICMATERIAL_API static const FLazyName GlobalSurfaceThicknessParameterName;
 	DYNAMICMATERIAL_API static const FLazyName GlobalOffsetValueName;
 	DYNAMICMATERIAL_API static const FLazyName GlobalOffsetParameterName;
 	DYNAMICMATERIAL_API static const FLazyName GlobalTilingValueName;
diff --git a/Engine/Plugins/Experimental/DynamicMaterial/Source/DynamicMaterial/Public/Model/IDynamicMaterialModelEditorOnlyDataInterface.h b/Engine/Plugins/Experimental/DynamicMaterial/Source/DynamicMaterial/Public/Model/IDynamicMaterialModelEditorOnlyDataInterface.h
--- a/Engine/Plugins/Experimental/DynamicMaterial/Source/DynamicMaterial/Public/Model/IDynamicMaterialModelEditorOnlyDataInterface.h
+++ b/Engine/Plugins/Experimental/DynamicMaterial/Source/DynamicMaterial/Public/Model/IDynamicMaterialModelEditorOnlyDataInterface.h
@@ -47,9 +47,6 @@ public:
 	/** Called when a texture uv is updated. */
 	virtual void OnTextureUVUpdated(UDMTextureUV* InTextureUV) = 0;
 
-	/** Called to load deprecated model data from 5.2 to 5.3 materials. */
-	virtual void LoadDeprecatedModelData(UDynamicMaterialModel* InMaterialModel) = 0;
-
 #if WITH_EDITOR
 	/** Called when the model needs to have the material rebuild. */
 	UFUNCTION(BlueprintCallable, Category = "Material Designer")
diff --git a/Engine/Plugins/Experimental/DynamicMaterial/Source/DynamicMaterialEditor/Private/Components/DMMaterialProperty.cpp b/Engine/Plugins/Experimental/DynamicMaterial/Source/DynamicMaterialEditor/Private/Components/DMMaterialProperty.cpp
--- a/Engine/Plugins/Experimental/DynamicMaterial/Source/DynamicMaterialEditor/Private/Components/DMMaterialProperty.cpp
+++ b/Engine/Plugins/Experimental/DynamicMaterial/Source/DynamicMaterialEditor/Private/Components/DMMaterialProperty.cpp
@@ -94,7 +94,19 @@ FText UDMMaterialProperty::GetDescription() const
 
 bool UDMMaterialProperty::IsMaterialPin() const
 {
-	return MaterialProperty > EDMMaterialPropertyType::None && MaterialProperty < EDMMaterialPropertyType::Custom1;
+	switch (MaterialProperty)
+	{
+		case EDMMaterialPropertyType::None:
+		case EDMMaterialPropertyType::Any:
+		case EDMMaterialPropertyType::Custom1:
+		case EDMMaterialPropertyType::Custom2:
+		case EDMMaterialPropertyType::Custom3:
+		case EDMMaterialPropertyType::Custom4:
+				return false;
+
+		default:
+			return true;
+	}
 }
 
 void UDMMaterialProperty::ResetInputConnectionMap()
@@ -672,13 +684,13 @@ bool UDMMaterialProperty::IsValidForModel(UDynamicMaterialModelEditorOnlyData& I
 		InMaterialModel.GetBlendMode(),
 		InMaterialModel.GetShadingModel() == EDMMaterialShadingModel::DefaultLit ? EMaterialShadingModel::MSM_DefaultLit : EMaterialShadingModel::MSM_Unlit,
 		TLM_Surface,
-		/* Tesselation enabled */ false,
+		/* Tesselation enabled */ InMaterialModel.IsNaniteTessellationEnabled(),
 		/* BlendableOutputAlpha (Post Process Alpha) */ false,
 		/* Uses Distortion */ false,
 		/* Shading model from master material */ false,
 		/* Outputting translucency velocity */ InMaterialModel.GetBlendMode() != BLEND_Opaque,
 		/* Thin surface */ false,
-		/* Is supported (substate check) */ true
+		/* Is supported (substrate check) */ true
 	});
 }
 
diff --git a/Engine/Plugins/Experimental/DynamicMaterial/Source/DynamicMaterialEditor/Private/Components/MaterialProperties/DMMPDisplacement.cpp b/Engine/Plugins/Experimental/DynamicMaterial/Source/DynamicMaterialEditor/Private/Components/MaterialProperties/DMMPDisplacement.cpp
new file mode 100644
--- /dev/null
+++ b/Engine/Plugins/Experimental/DynamicMaterial/Source/DynamicMaterialEditor/Private/Components/MaterialProperties/DMMPDisplacement.cpp
@@ -0,0 +1,17 @@
+// Copyright Epic Games, Inc. All Rights Reserved.
+
+#include "Components/MaterialProperties/DMMPDisplacement.h"
+#include "Model/DynamicMaterialModelEditorOnlyData.h"
+
+UDMMaterialPropertyDisplacement::UDMMaterialPropertyDisplacement()
+	: UDMMaterialProperty(
+		EDMMaterialPropertyType(EDMMaterialPropertyType::Displacement),
+		EDMValueType::VT_Float3_XYZ)
+{
+}
+
+UMaterialExpression* UDMMaterialPropertyDisplacement::GetDefaultInput(
+	const TSharedRef<FDMMaterialBuildState>& InBuildState) const
+{
+	return CreateConstant(InBuildState, FVector::ZeroVector);
+}
diff --git a/Engine/Plugins/Experimental/DynamicMaterial/Source/DynamicMaterialEditor/Private/Components/MaterialProperties/DMMPDisplacement.h b/Engine/Plugins/Experimental/DynamicMaterial/Source/DynamicMaterialEditor/Private/Components/MaterialProperties/DMMPDisplacement.h
new file mode 100644
--- /dev/null
+++ b/Engine/Plugins/Experimental/DynamicMaterial/Source/DynamicMaterialEditor/Private/Components/MaterialProperties/DMMPDisplacement.h
@@ -0,0 +1,21 @@
+// Copyright Epic Games, Inc. All Rights Reserved.
+
+#pragma once
+
+#include "Components/DMMaterialProperty.h"
+#include "DMMPDisplacement.generated.h"
+
+class UDynamicMaterialModelEditorOnlyData;
+
+UCLASS(BlueprintType, ClassGroup = "Material Designer")
+class UDMMaterialPropertyDisplacement : public UDMMaterialProperty
+{
+	GENERATED_BODY()
+
+public:
+	UDMMaterialPropertyDisplacement();
+
+	//~ Begin UDMMaterialProperty
+	virtual UMaterialExpression* GetDefaultInput(const TSharedRef<FDMMaterialBuildState>& InBuildState) const override;
+	//~ End UDMMaterialProperty
+};
diff --git a/Engine/Plugins/Experimental/DynamicMaterial/Source/DynamicMaterialEditor/Private/Components/MaterialProperties/DMMPSubsurfaceColor.cpp b/Engine/Plugins/Experimental/DynamicMaterial/Source/DynamicMaterialEditor/Private/Components/MaterialProperties/DMMPSubsurfaceColor.cpp
new file mode 100644
--- /dev/null
+++ b/Engine/Plugins/Experimental/DynamicMaterial/Source/DynamicMaterialEditor/Private/Components/MaterialProperties/DMMPSubsurfaceColor.cpp
@@ -0,0 +1,22 @@
+// Copyright Epic Games, Inc. All Rights Reserved.
+
+#include "Components/MaterialProperties/DMMPSubsurfaceColor.h"
+#include "Model/DynamicMaterialModelEditorOnlyData.h"
+
+UDMMaterialPropertySubsurfaceColor::UDMMaterialPropertySubsurfaceColor()
+	: UDMMaterialProperty(
+		EDMMaterialPropertyType(EDMMaterialPropertyType::SubsurfaceColor),
+		EDMValueType::VT_Float3_RGB)
+{
+}
+
+UMaterialExpression* UDMMaterialPropertySubsurfaceColor::GetDefaultInput(
+	const TSharedRef<FDMMaterialBuildState>& InBuildState) const
+{
+	return CreateConstant(InBuildState, FVector::OneVector);
+}
+
+TEnumAsByte<EMaterialSamplerType> UDMMaterialPropertySubsurfaceColor::GetTextureSamplerType() const
+{
+	return EMaterialSamplerType::SAMPLERTYPE_Color;
+}
diff --git a/Engine/Plugins/Experimental/DynamicMaterial/Source/DynamicMaterialEditor/Private/Components/MaterialProperties/DMMPSubsurfaceColor.h b/Engine/Plugins/Experimental/DynamicMaterial/Source/DynamicMaterialEditor/Private/Components/MaterialProperties/DMMPSubsurfaceColor.h
new file mode 100644
--- /dev/null
+++ b/Engine/Plugins/Experimental/DynamicMaterial/Source/DynamicMaterialEditor/Private/Components/MaterialProperties/DMMPSubsurfaceColor.h
@@ -0,0 +1,22 @@
+// Copyright Epic Games, Inc. All Rights Reserved.
+
+#pragma once
+
+#include "Components/DMMaterialProperty.h"
+#include "DMMPSubsurfaceColor.generated.h"
+
+class UDynamicMaterialModelEditorOnlyData;
+
+UCLASS(BlueprintType, ClassGroup = "Material Designer")
+class UDMMaterialPropertySubsurfaceColor : public UDMMaterialProperty
+{
+	GENERATED_BODY()
+
+public:
+	UDMMaterialPropertySubsurfaceColor();
+
+	//~ Begin UDMMaterialProperty
+	virtual UMaterialExpression* GetDefaultInput(const TSharedRef<FDMMaterialBuildState>& InBuildState) const override;
+	virtual TEnumAsByte<EMaterialSamplerType> GetTextureSamplerType() const override;
+	//~ End UDMMaterialProperty
+};
diff --git a/Engine/Plugins/Experimental/DynamicMaterial/Source/DynamicMaterialEditor/Private/Components/MaterialProperties/DMMPSurfaceThickness.cpp b/Engine/Plugins/Experimental/DynamicMaterial/Source/DynamicMaterialEditor/Private/Components/MaterialProperties/DMMPSurfaceThickness.cpp
new file mode 100644
--- /dev/null
+++ b/Engine/Plugins/Experimental/DynamicMaterial/Source/DynamicMaterialEditor/Private/Components/MaterialProperties/DMMPSurfaceThickness.cpp
@@ -0,0 +1,16 @@
+// Copyright Epic Games, Inc. All Rights Reserved.
+
+#include "Components/MaterialProperties/DMMPSurfaceThickness.h"
+#include "Model/DynamicMaterialModelEditorOnlyData.h"
+
+UDMMaterialPropertySurfaceThickness::UDMMaterialPropertySurfaceThickness()
+	: UDMMaterialProperty(
+		EDMMaterialPropertyType(EDMMaterialPropertyType::SurfaceThickness),
+		EDMValueType::VT_Float1)
+{
+}
+
+UMaterialExpression* UDMMaterialPropertySurfaceThickness::GetDefaultInput(const TSharedRef<FDMMaterialBuildState>& InBuildState) const
+{
+	return CreateConstant(InBuildState, 0.f);
+}
diff --git a/Engine/Plugins/Experimental/DynamicMaterial/Source/DynamicMaterialEditor/Private/Components/MaterialProperties/DMMPSurfaceThickness.h b/Engine/Plugins/Experimental/DynamicMaterial/Source/DynamicMaterialEditor/Private/Components/MaterialProperties/DMMPSurfaceThickness.h
new file mode 100644
--- /dev/null
+++ b/Engine/Plugins/Experimental/DynamicMaterial/Source/DynamicMaterialEditor/Private/Components/MaterialProperties/DMMPSurfaceThickness.h
@@ -0,0 +1,21 @@
+// Copyright Epic Games, Inc. All Rights Reserved.
+
+#pragma once
+
+#include "Components/DMMaterialProperty.h"
+#include "DMMPSurfaceThickness.generated.h"
+
+class UDynamicMaterialModelEditorOnlyData;
+
+UCLASS(BlueprintType, ClassGroup = "Material Designer")
+class UDMMaterialPropertySurfaceThickness : public UDMMaterialProperty
+{
+	GENERATED_BODY()
+
+public:
+	UDMMaterialPropertySurfaceThickness();
+
+	//~ Begin UDMMaterialProperty
+	virtual UMaterialExpression* GetDefaultInput(const TSharedRef<FDMMaterialBuildState>& InBuildState) const override;
+	//~ End UDMMaterialProperty
+};
diff --git a/Engine/Plugins/Experimental/DynamicMaterial/Source/DynamicMaterialEditor/Private/DynamicMaterialEditorModule.cpp b/Engine/Plugins/Experimental/DynamicMaterial/Source/DynamicMaterialEditor/Private/DynamicMaterialEditorModule.cpp
--- a/Engine/Plugins/Experimental/DynamicMaterial/Source/DynamicMaterialEditor/Private/DynamicMaterialEditorModule.cpp
+++ b/Engine/Plugins/Experimental/DynamicMaterial/Source/DynamicMaterialEditor/Private/DynamicMaterialEditorModule.cpp
@@ -243,11 +243,6 @@ void FDynamicMaterialEditorModule::StartupModule()
 	RegisterComponentPropertyRowGeneratorDelegate<UDMMaterialEffectFunction,       FDMMaterialEffectFunctionPropertyRowGenerator>();
 	RegisterComponentPropertyRowGeneratorDelegate<UDMMaterialStageInputFunction,   FDMMaterialStageFunctionPropertyRowGenerator>();
 
-	FDynamicMaterialModule::GetCreateEditorOnlyDataDelegate().BindLambda([](UDynamicMaterialModel* InMaterialModel) -> TScriptInterface<IDynamicMaterialModelEditorOnlyDataInterface>
-		{
-			return TScriptInterface<IDynamicMaterialModelEditorOnlyDataInterface>(NewObject<UDynamicMaterialModelEditorOnlyData>(InMaterialModel));
-		});
-
 	BuildRequestList.Empty();
 
 	UDMMaterialValueTexture::GetDefaultRGBTexture.BindLambda([]()
diff --git a/Engine/Plugins/Experimental/DynamicMaterial/Source/DynamicMaterialEditor/Private/DynamicMaterialEditorSettings.cpp b/Engine/Plugins/Experimental/DynamicMaterial/Source/DynamicMaterialEditor/Private/DynamicMaterialEditorSettings.cpp
--- a/Engine/Plugins/Experimental/DynamicMaterial/Source/DynamicMaterialEditor/Private/DynamicMaterialEditorSettings.cpp
+++ b/Engine/Plugins/Experimental/DynamicMaterial/Source/DynamicMaterialEditor/Private/DynamicMaterialEditorSettings.cpp
@@ -33,6 +33,9 @@ namespace UE::DynamicMaterialEditor::Private
 				FLinearColor(0.f, 0.f, 1.f, 1.f)
 			}},
 		{EDMMaterialPropertyType::AmbientOcclusion,    TSoftObjectPtr<UTexture>(FSoftObjectPath(TEXT("/Script/Engine.Texture2D'/DynamicMaterial/Textures/SlotDefaults/T_MD_AmbientOcclusion.T_MD_AmbientOcclusion'")))},
+		{EDMMaterialPropertyType::Displacement,        TSoftObjectPtr<UTexture>(FSoftObjectPath(TEXT("/Script/Engine.Texture2D'/DynamicMaterial/Textures/SlotDefaults/T_MD_Displacement.T_MD_Displacement'")))},
+		{EDMMaterialPropertyType::SubsurfaceColor,     TSoftObjectPtr<UTexture>(FSoftObjectPath(TEXT("/Script/Engine.Texture2D'/Engine/EngineResources/T_MD_SubsurfaceColor.T_MD_SubsurfaceColor'")))},
+		{EDMMaterialPropertyType::SurfaceThickness,    FLinearColor(0.0f, 0.0f, 0.0f, 1.f)},
 	};
 }
 
@@ -107,6 +110,15 @@ bool FDMMaterialChannelListPreset::IsPropertyEnabled(EDMMaterialPropertyType InP
 		case EDMMaterialPropertyType::Tangent:
 			return bTangent;
 
+		case EDMMaterialPropertyType::Displacement:
+			return bDisplacement;
+
+		case EDMMaterialPropertyType::SubsurfaceColor:
+			return bSubsurfaceColor;
+
+		case EDMMaterialPropertyType::SurfaceThickness:
+			return bSurfaceThickness;
+
 		default:
 			return false;
 	}
@@ -160,6 +172,7 @@ UDynamicMaterialEditorSettings::UDynamicMaterialEditorSettings()
 	PBR.bRoughness = true;
 	PBR.bNormal = true;
 	PBR.bAmbientOcclusion = true;
+	PBR.bDisplacement = true;
 	PBR.DefaultBlendMode = BLEND_Opaque;
 	PBR.DefaultShadingModel = EDMMaterialShadingModel::DefaultLit;
 	PBR.bDefaultAnimated = false;
@@ -180,6 +193,9 @@ UDynamicMaterialEditorSettings::UDynamicMaterialEditorSettings()
 	All.bRefraction = true;
 	All.bTangent = true;
 	All.bWorldPositionOffset = true;
+	All.bDisplacement = true;
+	All.bSubsurfaceColor = true;
+	All.bSurfaceThickness = true;
 	All.DefaultBlendMode = BLEND_Opaque;
 	All.DefaultShadingModel = EDMMaterialShadingModel::DefaultLit;
 	All.bDefaultAnimated = false;
diff --git a/Engine/Plugins/Experimental/DynamicMaterial/Source/DynamicMaterialEditor/Private/DynamicMaterialEditorSettings.h b/Engine/Plugins/Experimental/DynamicMaterial/Source/DynamicMaterialEditor/Private/DynamicMaterialEditorSettings.h
--- a/Engine/Plugins/Experimental/DynamicMaterial/Source/DynamicMaterialEditor/Private/DynamicMaterialEditorSettings.h
+++ b/Engine/Plugins/Experimental/DynamicMaterial/Source/DynamicMaterialEditor/Private/DynamicMaterialEditorSettings.h
@@ -103,6 +103,15 @@ struct FDMMaterialChannelListPreset
 	UPROPERTY(Config, EditAnywhere, BlueprintReadWrite, Category = "Channel")
 	bool bTangent = false;
 
+	UPROPERTY(Config, EditAnywhere, BlueprintReadWrite, Category = "Channel")
+	bool bDisplacement = false;
+
+	UPROPERTY(Config, EditAnywhere, BlueprintReadWrite, Category = "Channel")
+	bool bSubsurfaceColor = false;
+
+	UPROPERTY(Config, EditAnywhere, BlueprintReadWrite, Category = "Channel")
+	bool bSurfaceThickness = false;
+
 	UPROPERTY(Config, EditAnywhere, BlueprintReadWrite, Category = "Channel")
 	TEnumAsByte<EBlendMode> DefaultBlendMode = BLEND_Opaque;
 
diff --git a/Engine/Plugins/Experimental/DynamicMaterial/Source/DynamicMaterialEditor/Private/Model/DMMaterialBuildState.cpp b/Engine/Plugins/Experimental/DynamicMaterial/Source/DynamicMaterialEditor/Private/Model/DMMaterialBuildState.cpp
--- a/Engine/Plugins/Experimental/DynamicMaterial/Source/DynamicMaterialEditor/Private/Model/DMMaterialBuildState.cpp
+++ b/Engine/Plugins/Experimental/DynamicMaterial/Source/DynamicMaterialEditor/Private/Model/DMMaterialBuildState.cpp
@@ -93,49 +93,65 @@ IDMMaterialBuildUtilsInterface& FDMMaterialBuildState::GetBuildUtils() const
 
 FExpressionInput* FDMMaterialBuildState::GetMaterialProperty(EDMMaterialPropertyType InProperty) const
 {
+	UMaterialEditorOnlyData* EditorOnlyData = DynamicMaterial->GetEditorOnlyData();
+
+	if (!EditorOnlyData)
+	{
+		return nullptr;
+	}
+
 	switch (InProperty)
 	{
 		case EDMMaterialPropertyType::AmbientOcclusion:
-			return &(DynamicMaterial->GetEditorOnlyData()->AmbientOcclusion);
+			return &(EditorOnlyData->AmbientOcclusion);
 
 		case EDMMaterialPropertyType::Anisotropy:
-			return &(DynamicMaterial->GetEditorOnlyData()->Anisotropy);
+			return &(EditorOnlyData->Anisotropy);
 
 		case EDMMaterialPropertyType::BaseColor:
-			return &(DynamicMaterial->GetEditorOnlyData()->BaseColor);
+			return &(EditorOnlyData->BaseColor);
 
 		case EDMMaterialPropertyType::EmissiveColor:
-			return &(DynamicMaterial->GetEditorOnlyData()->EmissiveColor);
+			return &(EditorOnlyData->EmissiveColor);
 
 		case EDMMaterialPropertyType::Metallic:
-			return &(DynamicMaterial->GetEditorOnlyData()->Metallic);
+			return &(EditorOnlyData->Metallic);
 
 		case EDMMaterialPropertyType::Normal:
-			return &(DynamicMaterial->GetEditorOnlyData()->Normal);
+			return &(EditorOnlyData->Normal);
 
 		case EDMMaterialPropertyType::Opacity:
-			return &(DynamicMaterial->GetEditorOnlyData()->Opacity);
+			return &(EditorOnlyData->Opacity);
 
 		case EDMMaterialPropertyType::OpacityMask:
-			return &(DynamicMaterial->GetEditorOnlyData()->OpacityMask);
+			return &(EditorOnlyData->OpacityMask);
 
 		case EDMMaterialPropertyType::PixelDepthOffset:
-			return &(DynamicMaterial->GetEditorOnlyData()->PixelDepthOffset);
+			return &(EditorOnlyData->PixelDepthOffset);
 
 		case EDMMaterialPropertyType::Refraction:
-			return &(DynamicMaterial->GetEditorOnlyData()->Refraction);
+			return &(EditorOnlyData->Refraction);
 
 		case EDMMaterialPropertyType::Roughness:
-			return &(DynamicMaterial->GetEditorOnlyData()->Roughness);
+			return &(EditorOnlyData->Roughness);
 
 		case EDMMaterialPropertyType::Specular:
-			return &(DynamicMaterial->GetEditorOnlyData()->Specular);
+			return &(EditorOnlyData->Specular);
 
 		case EDMMaterialPropertyType::Tangent:
-			return &(DynamicMaterial->GetEditorOnlyData()->Tangent);
+			return &(EditorOnlyData->Tangent);
 
 		case EDMMaterialPropertyType::WorldPositionOffset:
-			return &(DynamicMaterial->GetEditorOnlyData()->WorldPositionOffset);
+			return &(EditorOnlyData->WorldPositionOffset);
+
+		case EDMMaterialPropertyType::Displacement:
+			return &(EditorOnlyData->Displacement);
+
+		case EDMMaterialPropertyType::SubsurfaceColor:
+			return &(EditorOnlyData->SubsurfaceColor);
+
+		case EDMMaterialPropertyType::SurfaceThickness:
+			return &(EditorOnlyData->SurfaceThickness);
 
 		default:
 			return nullptr;
diff --git a/Engine/Plugins/Experimental/DynamicMaterial/Source/DynamicMaterialEditor/Private/Model/DMMaterialNodeArranger.cpp b/Engine/Plugins/Experimental/DynamicMaterial/Source/DynamicMaterialEditor/Private/Model/DMMaterialNodeArranger.cpp
--- a/Engine/Plugins/Experimental/DynamicMaterial/Source/DynamicMaterialEditor/Private/Model/DMMaterialNodeArranger.cpp
+++ b/Engine/Plugins/Experimental/DynamicMaterial/Source/DynamicMaterialEditor/Private/Model/DMMaterialNodeArranger.cpp
@@ -18,20 +18,30 @@ FDMMaterialNodeArranger::FDMMaterialNodeArranger(UMaterial* InDynamicMaterial)
 
 void FDMMaterialNodeArranger::ArrangeNodes()
 {
-	ArrangeMaterialInputNodes(DynamicMaterial->GetEditorOnlyData()->BaseColor.Expression);
-	ArrangeMaterialInputNodes(DynamicMaterial->GetEditorOnlyData()->EmissiveColor.Expression);
-	ArrangeMaterialInputNodes(DynamicMaterial->GetEditorOnlyData()->Opacity.Expression);
-	ArrangeMaterialInputNodes(DynamicMaterial->GetEditorOnlyData()->OpacityMask.Expression);
-	ArrangeMaterialInputNodes(DynamicMaterial->GetEditorOnlyData()->Metallic.Expression);
-	ArrangeMaterialInputNodes(DynamicMaterial->GetEditorOnlyData()->Specular.Expression);
-	ArrangeMaterialInputNodes(DynamicMaterial->GetEditorOnlyData()->Roughness.Expression);
-	ArrangeMaterialInputNodes(DynamicMaterial->GetEditorOnlyData()->Anisotropy.Expression);
-	ArrangeMaterialInputNodes(DynamicMaterial->GetEditorOnlyData()->Normal.Expression);
-	ArrangeMaterialInputNodes(DynamicMaterial->GetEditorOnlyData()->Tangent.Expression);
-	ArrangeMaterialInputNodes(DynamicMaterial->GetEditorOnlyData()->WorldPositionOffset.Expression);
-	ArrangeMaterialInputNodes(DynamicMaterial->GetEditorOnlyData()->Refraction.Expression);
-	ArrangeMaterialInputNodes(DynamicMaterial->GetEditorOnlyData()->AmbientOcclusion.Expression);
-	ArrangeMaterialInputNodes(DynamicMaterial->GetEditorOnlyData()->PixelDepthOffset.Expression);
+	UMaterialEditorOnlyData* EditorOnlyData = DynamicMaterial->GetEditorOnlyData();
+
+	if (!EditorOnlyData)
+	{
+		return;
+	}
+
+	ArrangeMaterialInputNodes(EditorOnlyData->BaseColor.Expression);
+	ArrangeMaterialInputNodes(EditorOnlyData->EmissiveColor.Expression);
+	ArrangeMaterialInputNodes(EditorOnlyData->Opacity.Expression);
+	ArrangeMaterialInputNodes(EditorOnlyData->OpacityMask.Expression);
+	ArrangeMaterialInputNodes(EditorOnlyData->Metallic.Expression);
+	ArrangeMaterialInputNodes(EditorOnlyData->Specular.Expression);
+	ArrangeMaterialInputNodes(EditorOnlyData->Roughness.Expression);
+	ArrangeMaterialInputNodes(EditorOnlyData->Anisotropy.Expression);
+	ArrangeMaterialInputNodes(EditorOnlyData->Normal.Expression);
+	ArrangeMaterialInputNodes(EditorOnlyData->Tangent.Expression);
+	ArrangeMaterialInputNodes(EditorOnlyData->WorldPositionOffset.Expression);
+	ArrangeMaterialInputNodes(EditorOnlyData->Refraction.Expression);
+	ArrangeMaterialInputNodes(EditorOnlyData->AmbientOcclusion.Expression);
+	ArrangeMaterialInputNodes(EditorOnlyData->PixelDepthOffset.Expression);
+	ArrangeMaterialInputNodes(EditorOnlyData->Displacement.Expression);
+	ArrangeMaterialInputNodes(EditorOnlyData->SubsurfaceColor.Expression);
+	ArrangeMaterialInputNodes(EditorOnlyData->SurfaceThickness.Expression);
 
 	FInt32Interval Vertical = {0, 0};
 
diff --git a/Engine/Plugins/Experimental/DynamicMaterial/Source/DynamicMaterialEditor/Private/Model/DynamicMaterialModelEditorOnlyData.cpp b/Engine/Plugins/Experimental/DynamicMaterial/Source/DynamicMaterialEditor/Private/Model/DynamicMaterialModelEditorOnlyData.cpp
--- a/Engine/Plugins/Experimental/DynamicMaterial/Source/DynamicMaterialEditor/Private/Model/DynamicMaterialModelEditorOnlyData.cpp
+++ b/Engine/Plugins/Experimental/DynamicMaterial/Source/DynamicMaterialEditor/Private/Model/DynamicMaterialModelEditorOnlyData.cpp
@@ -14,6 +14,7 @@
 #include "Components/MaterialProperties/DMMPAmbientOcclusion.h"
 #include "Components/MaterialProperties/DMMPAnisotropy.h"
 #include "Components/MaterialProperties/DMMPBaseColor.h"
+#include "Components/MaterialProperties/DMMPDisplacement.h"
 #include "Components/MaterialProperties/DMMPEmissiveColor.h"
 #include "Components/MaterialProperties/DMMPMetallic.h"
 #include "Components/MaterialProperties/DMMPNormal.h"
@@ -23,6 +24,8 @@
 #include "Components/MaterialProperties/DMMPRefraction.h"
 #include "Components/MaterialProperties/DMMPRoughness.h"
 #include "Components/MaterialProperties/DMMPSpecular.h"
+#include "Components/MaterialProperties/DMMPSubsurfaceColor.h"
+#include "Components/MaterialProperties/DMMPSurfaceThickness.h"
 #include "Components/MaterialProperties/DMMPTangent.h"
 #include "Components/MaterialProperties/DMMPWorldPositionOffset.h"
 #include "Components/MaterialStageExpressions/DMMSETextureSample.h"
@@ -66,6 +69,9 @@ const FString UDynamicMaterialModelEditorOnlyData::AmbientOcclusionPathToken
 const FString UDynamicMaterialModelEditorOnlyData::AnisotropyPathToken          = FString(TEXT("Anisotropy"));
 const FString UDynamicMaterialModelEditorOnlyData::RefractionPathToken          = FString(TEXT("Refraction"));
 const FString UDynamicMaterialModelEditorOnlyData::TangentPathToken             = FString(TEXT("Tangent"));
+const FString UDynamicMaterialModelEditorOnlyData::DisplacementPathToken        = FString(TEXT("Displacement"));
+const FString UDynamicMaterialModelEditorOnlyData::SubsurfaceColorPathToken     = FString(TEXT("SubsurfaceColor"));
+const FString UDynamicMaterialModelEditorOnlyData::SurfaceThicknessPathToken    = FString(TEXT("SurfaceThickness"));
 const FString UDynamicMaterialModelEditorOnlyData::Custom1PathToken             = FString(TEXT("Custom1"));
 const FString UDynamicMaterialModelEditorOnlyData::Custom2PathToken             = FString(TEXT("Custom2"));
 const FString UDynamicMaterialModelEditorOnlyData::Custom3PathToken             = FString(TEXT("Custom3"));
@@ -88,6 +94,9 @@ namespace UE::DynamicMaterialEditor::Private
 		{UDynamicMaterialModelEditorOnlyData::AnisotropyPathToken,          EDMMaterialPropertyType::Anisotropy},
 		{UDynamicMaterialModelEditorOnlyData::RefractionPathToken,          EDMMaterialPropertyType::Refraction},
 		{UDynamicMaterialModelEditorOnlyData::TangentPathToken,             EDMMaterialPropertyType::Tangent},
+		{UDynamicMaterialModelEditorOnlyData::DisplacementPathToken,        EDMMaterialPropertyType::Displacement},
+		{UDynamicMaterialModelEditorOnlyData::SubsurfaceColorPathToken,     EDMMaterialPropertyType::SubsurfaceColor},
+		{UDynamicMaterialModelEditorOnlyData::SurfaceThicknessPathToken,    EDMMaterialPropertyType::SurfaceThickness},
 		{UDynamicMaterialModelEditorOnlyData::Custom1PathToken,             EDMMaterialPropertyType::Custom1},
 		{UDynamicMaterialModelEditorOnlyData::Custom2PathToken,             EDMMaterialPropertyType::Custom2},
 		{UDynamicMaterialModelEditorOnlyData::Custom3PathToken,             EDMMaterialPropertyType::Custom3},
@@ -175,45 +184,56 @@ UDynamicMaterialModelEditorOnlyData::UDynamicMaterialModelEditorOnlyData()
 	, bResponsiveAAEnabled(true)
 	, ChannelListPreset(NAME_None)
 {
-	Properties.Emplace(EDMMaterialPropertyType::BaseColor,           CreateDefaultSubobject<UDMMaterialPropertyBaseColor>(          "MaterialProperty_BaseColor"));
-	Properties.Emplace(EDMMaterialPropertyType::EmissiveColor,       CreateDefaultSubobject<UDMMaterialPropertyEmissiveColor>(      "MaterialProperty_EmissiveColor"));
-	Properties.Emplace(EDMMaterialPropertyType::Opacity,             CreateDefaultSubobject<UDMMaterialPropertyOpacity>(            "MaterialProperty_Opacity"));
-	Properties.Emplace(EDMMaterialPropertyType::OpacityMask,         CreateDefaultSubobject<UDMMaterialPropertyOpacityMask>(        "MaterialProperty_OpacityMask"));
-	Properties.Emplace(EDMMaterialPropertyType::Metallic,            CreateDefaultSubobject<UDMMaterialPropertyMetallic>(           "MaterialProperty_Metallic"));
-	Properties.Emplace(EDMMaterialPropertyType::Specular,            CreateDefaultSubobject<UDMMaterialPropertySpecular>(           "MaterialProperty_Specular"));
-	Properties.Emplace(EDMMaterialPropertyType::Roughness,           CreateDefaultSubobject<UDMMaterialPropertyRoughness>(          "MaterialProperty_Roughness"));
-	Properties.Emplace(EDMMaterialPropertyType::Anisotropy,          CreateDefaultSubobject<UDMMaterialPropertyAnisotropy>(         "MaterialProperty_Anisotropy"));
-	Properties.Emplace(EDMMaterialPropertyType::Normal,              CreateDefaultSubobject<UDMMaterialPropertyNormal>(             "MaterialProperty_Normal"));
-	Properties.Emplace(EDMMaterialPropertyType::Tangent,             CreateDefaultSubobject<UDMMaterialPropertyTangent>(            "MaterialProperty_Tangent"));
-	Properties.Emplace(EDMMaterialPropertyType::WorldPositionOffset, CreateDefaultSubobject<UDMMaterialPropertyWorldPositionOffset>("MaterialProperty_WorldPositionOffset"));
-	Properties.Emplace(EDMMaterialPropertyType::AmbientOcclusion,    CreateDefaultSubobject<UDMMaterialPropertyAmbientOcclusion>(   "MaterialProperty_AmbientOcclusion"));
-	Properties.Emplace(EDMMaterialPropertyType::Refraction,          CreateDefaultSubobject<UDMMaterialPropertyRefraction>(         "MaterialProperty_Refraction"));
-	Properties.Emplace(EDMMaterialPropertyType::PixelDepthOffset,    CreateDefaultSubobject<UDMMaterialPropertyPixelDepthOffset>(   "MaterialProperty_PixelDepthOffset"));
+	BaseColor =           CreateDefaultSubobject<UDMMaterialPropertyBaseColor>(          "MaterialProperty_BaseColor");
+	EmissiveColor =       CreateDefaultSubobject<UDMMaterialPropertyEmissiveColor>(      "MaterialProperty_EmissiveColor");
+	Opacity =             CreateDefaultSubobject<UDMMaterialPropertyOpacity>(            "MaterialProperty_Opacity");
+	OpacityMask =         CreateDefaultSubobject<UDMMaterialPropertyOpacityMask>(        "MaterialProperty_OpacityMask");
+	Roughness =           CreateDefaultSubobject<UDMMaterialPropertyRoughness>(          "MaterialProperty_Roughness");
+	Specular =            CreateDefaultSubobject<UDMMaterialPropertySpecular>(           "MaterialProperty_Specular");
+	Metallic =            CreateDefaultSubobject<UDMMaterialPropertyMetallic>(           "MaterialProperty_Metallic");
+	Normal =              CreateDefaultSubobject<UDMMaterialPropertyNormal>(             "MaterialProperty_Normal");
+	PixelDepthOffset =    CreateDefaultSubobject<UDMMaterialPropertyPixelDepthOffset>(   "MaterialProperty_PixelDepthOffset");
+	WorldPositionOffset = CreateDefaultSubobject<UDMMaterialPropertyWorldPositionOffset>("MaterialProperty_WorldPositionOffset");
+	AmbientOcclusion =    CreateDefaultSubobject<UDMMaterialPropertyAmbientOcclusion>(   "MaterialProperty_AmbientOcclusion");
+	Anisotropy =          CreateDefaultSubobject<UDMMaterialPropertyAnisotropy>(         "MaterialProperty_Anisotropy");
+	Refraction =          CreateDefaultSubobject<UDMMaterialPropertyRefraction>(         "MaterialProperty_Refraction");
+	Tangent =             CreateDefaultSubobject<UDMMaterialPropertyTangent>(            "MaterialProperty_Tangent");
+	Displacement =        CreateDefaultSubobject<UDMMaterialPropertyDisplacement>(       "MaterialProperty_Displacement");
+	SubsurfaceColor =     CreateDefaultSubobject<UDMMaterialPropertySubsurfaceColor>(    "MaterialProperty_SubsurfaceColor");
+	SurfaceThickness =    CreateDefaultSubobject<UDMMaterialPropertySurfaceThickness>(   "MaterialProperty_SurfaceThickness");
+
+	Custom1 = UDMMaterialProperty::CreateCustomMaterialPropertyDefaultSubobject(this, EDMMaterialPropertyType::Custom1, "MaterialProperty_Custom1");
+	Custom2 = UDMMaterialProperty::CreateCustomMaterialPropertyDefaultSubobject(this, EDMMaterialPropertyType::Custom2, "MaterialProperty_Custom2");
+	Custom3 = UDMMaterialProperty::CreateCustomMaterialPropertyDefaultSubobject(this, EDMMaterialPropertyType::Custom3, "MaterialProperty_Custom3");
+	Custom4 = UDMMaterialProperty::CreateCustomMaterialPropertyDefaultSubobject(this, EDMMaterialPropertyType::Custom4, "MaterialProperty_Custom4");
 
-	Properties.Emplace(EDMMaterialPropertyType::Custom1, UDMMaterialProperty::CreateCustomMaterialPropertyDefaultSubobject(this, EDMMaterialPropertyType::Custom1, "MaterialProperty_Custom1"));
-	Properties.Emplace(EDMMaterialPropertyType::Custom2, UDMMaterialProperty::CreateCustomMaterialPropertyDefaultSubobject(this, EDMMaterialPropertyType::Custom2, "MaterialProperty_Custom2"));
-	Properties.Emplace(EDMMaterialPropertyType::Custom3, UDMMaterialProperty::CreateCustomMaterialPropertyDefaultSubobject(this, EDMMaterialPropertyType::Custom3, "MaterialProperty_Custom3"));
-	Properties.Emplace(EDMMaterialPropertyType::Custom4, UDMMaterialProperty::CreateCustomMaterialPropertyDefaultSubobject(this, EDMMaterialPropertyType::Custom4, "MaterialProperty_Custom4"));
+	ForEachMaterialPropertyType(
+		[this](EDMMaterialPropertyType InType)
+		{
+			if (UDMMaterialProperty* Property = GetMaterialProperty(InType))
+			{
+				Property->SetComponentState(EDMComponentLifetimeState::Added);
+			}
 
-	for (const TPair<EDMMaterialPropertyType, TObjectPtr<UDMMaterialProperty>>& Pair : Properties)
-	{
-		Pair.Value->SetComponentState(EDMComponentLifetimeState::Added);
-	}
+			return EDMIterationResult::Continue;
+		}
+	);
 }
 
 void UDynamicMaterialModelEditorOnlyData::AssignPropertyAlphaValues()
 {
-	Properties[EDMMaterialPropertyType::Opacity            ]->AddComponent(AlphaValueName, MaterialModel->GetGlobalParameterValue(UDynamicMaterialModel::GlobalOpacityValueName));
-	Properties[EDMMaterialPropertyType::OpacityMask        ]->AddComponent(AlphaValueName, MaterialModel->GetGlobalParameterValue(UDynamicMaterialModel::GlobalOpacityValueName));
-	Properties[EDMMaterialPropertyType::Metallic           ]->AddComponent(AlphaValueName, MaterialModel->GetGlobalParameterValue(UDynamicMaterialModel::GlobalMetallicValueName));
-	Properties[EDMMaterialPropertyType::Specular           ]->AddComponent(AlphaValueName, MaterialModel->GetGlobalParameterValue(UDynamicMaterialModel::GlobalSpecularValueName));
-	Properties[EDMMaterialPropertyType::Roughness          ]->AddComponent(AlphaValueName, MaterialModel->GetGlobalParameterValue(UDynamicMaterialModel::GlobalRoughnessValueName));
-	Properties[EDMMaterialPropertyType::Normal             ]->AddComponent(AlphaValueName, MaterialModel->GetGlobalParameterValue(UDynamicMaterialModel::GlobalNormalValueName));
-	Properties[EDMMaterialPropertyType::Anisotropy         ]->AddComponent(AlphaValueName, MaterialModel->GetGlobalParameterValue(UDynamicMaterialModel::GlobalAnisotropyValueName));
-	Properties[EDMMaterialPropertyType::WorldPositionOffset]->AddComponent(AlphaValueName, MaterialModel->GetGlobalParameterValue(UDynamicMaterialModel::GlobalWorldPositionOffsetValueName));
-	Properties[EDMMaterialPropertyType::AmbientOcclusion   ]->AddComponent(AlphaValueName, MaterialModel->GetGlobalParameterValue(UDynamicMaterialModel::GlobalAmbientOcclusionValueName));
-	Properties[EDMMaterialPropertyType::Refraction         ]->AddComponent(AlphaValueName, MaterialModel->GetGlobalParameterValue(UDynamicMaterialModel::GlobalRefractionValueName));
-	Properties[EDMMaterialPropertyType::PixelDepthOffset   ]->AddComponent(AlphaValueName, MaterialModel->GetGlobalParameterValue(UDynamicMaterialModel::GlobalPixelDepthOffsetValueName));
+	Opacity            ->AddComponent(AlphaValueName, MaterialModel->GetGlobalParameterValue(UDynamicMaterialModel::GlobalOpacityValueName));
+	OpacityMask        ->AddComponent(AlphaValueName, MaterialModel->GetGlobalParameterValue(UDynamicMaterialModel::GlobalOpacityValueName));
+	Metallic           ->AddComponent(AlphaValueName, MaterialModel->GetGlobalParameterValue(UDynamicMaterialModel::GlobalMetallicValueName));
+	Specular           ->AddComponent(AlphaValueName, MaterialModel->GetGlobalParameterValue(UDynamicMaterialModel::GlobalSpecularValueName));
+	Roughness          ->AddComponent(AlphaValueName, MaterialModel->GetGlobalParameterValue(UDynamicMaterialModel::GlobalRoughnessValueName));
+	Normal             ->AddComponent(AlphaValueName, MaterialModel->GetGlobalParameterValue(UDynamicMaterialModel::GlobalNormalValueName));
+	Anisotropy         ->AddComponent(AlphaValueName, MaterialModel->GetGlobalParameterValue(UDynamicMaterialModel::GlobalAnisotropyValueName));
+	WorldPositionOffset->AddComponent(AlphaValueName, MaterialModel->GetGlobalParameterValue(UDynamicMaterialModel::GlobalWorldPositionOffsetValueName));
+	AmbientOcclusion   ->AddComponent(AlphaValueName, MaterialModel->GetGlobalParameterValue(UDynamicMaterialModel::GlobalAmbientOcclusionValueName));
+	Refraction         ->AddComponent(AlphaValueName, MaterialModel->GetGlobalParameterValue(UDynamicMaterialModel::GlobalRefractionValueName));
+	PixelDepthOffset   ->AddComponent(AlphaValueName, MaterialModel->GetGlobalParameterValue(UDynamicMaterialModel::GlobalPixelDepthOffsetValueName));
+	Displacement       ->AddComponent(AlphaValueName, MaterialModel->GetGlobalParameterValue(UDynamicMaterialModel::GlobalDisplacementValueName));
 }
 
 TArray<FName> UDynamicMaterialModelEditorOnlyData::GetPresetOptions() const
@@ -472,21 +492,26 @@ void UDynamicMaterialModelEditorOnlyData::BuildMaterial(bool bInDirtyAssets)
 	/**
 	 * Process slots to build base material inputs.
 	 */
-	for (const TPair<EDMMaterialPropertyType, TObjectPtr<UDMMaterialProperty>>& Pair : Properties)
-	{
-		if (!Pair.Value->IsMaterialPin())
+	ForEachMaterialPropertyType(
+		[this, &BuildState](EDMMaterialPropertyType InType)
 		{
-			continue;
-		}
+			if (UDMMaterialProperty* Property = GetMaterialProperty(InType))
+			{
+				if (Property->IsMaterialPin())
+				{
+					Property->GenerateExpressions(BuildState);
 
-		Pair.Value->GenerateExpressions(BuildState);
+					// Global opacity is handled at later
+					if (InType != EDMMaterialPropertyType::Opacity && InType != EDMMaterialPropertyType::OpacityMask)
+					{
+						Property->AddAlphaMultiplier(BuildState);
+					}
+				}
+			}
 
-		// Global opacity is handled at later
-		if (Pair.Key != EDMMaterialPropertyType::Opacity && Pair.Key != EDMMaterialPropertyType::OpacityMask)
-		{
-			Pair.Value->AddAlphaMultiplier(BuildState);
+			return EDMIterationResult::Continue;
 		}
-	}
+	);
 
 	if (Domain != EMaterialDomain::MD_PostProcess)
 	{
@@ -561,11 +586,11 @@ void UDynamicMaterialModelEditorOnlyData::BuildMaterial(bool bInDirtyAssets)
 
 		if (BlendMode == BLEND_Masked)
 		{
-			OpacityProperty = Properties.FindRef(EDMMaterialPropertyType::OpacityMask);
+			OpacityProperty = GetMaterialProperty(EDMMaterialPropertyType::OpacityMask);
 		}
 		else if (BlendMode != BLEND_Opaque)
 		{
-			OpacityProperty = Properties.FindRef(EDMMaterialPropertyType::Opacity);
+			OpacityProperty = GetMaterialProperty(EDMMaterialPropertyType::Opacity);
 		}
 
 		if (OpacityProperty != nullptr)
@@ -577,15 +602,20 @@ void UDynamicMaterialModelEditorOnlyData::BuildMaterial(bool bInDirtyAssets)
 	/**
 	 * Apply output processors.
 	 */
-	for (const TPair<EDMMaterialPropertyType, TObjectPtr<UDMMaterialProperty>>& Pair : Properties)
-	{
-		if (!Pair.Value->IsMaterialPin())
+	ForEachMaterialPropertyType(
+		[this, &BuildState](EDMMaterialPropertyType InType)
 		{
-			continue;
-		}
+			if (UDMMaterialProperty* Property = GetMaterialProperty(InType))
+			{
+				if (Property->IsMaterialPin())
+				{
+					Property->AddOutputProcessor(BuildState);
+				}
+			}
 
-		Pair.Value->AddOutputProcessor(BuildState);
-	}
+			return EDMIterationResult::Continue;
+		}
+	);
 
 	//MaterialStats = UMaterialEditingLibrary::GetStatistics(MaterialModel->DynamicMaterial);
 
@@ -612,123 +642,6 @@ void UDynamicMaterialModelEditorOnlyData::OnValueListUpdate()
 	OnValueListUpdateDelegate.Broadcast(MaterialModel);
 }
 
-void UDynamicMaterialModelEditorOnlyData::LoadDeprecatedModelData_Base(bool bInCreateMaterialPackage, EBlendMode InBlendMode, 
-	EDMMaterialShadingModel InShadingModel)
-{
-	bCreateMaterialPackage = bInCreateMaterialPackage;
-
-	if (IsValid(MaterialModel->DynamicMaterial))
-	{
-		BlendMode = MaterialModel->DynamicMaterial->BlendMode;
-
-		ShadingModel = MaterialModel->DynamicMaterial->GetShadingModels().HasShadingModel(EMaterialShadingModel::MSM_DefaultLit)
-			? EDMMaterialShadingModel::DefaultLit
-			: EDMMaterialShadingModel::Unlit;
-	}
-	else if (IsValid(MaterialModel->DynamicMaterialInstance))
-	{
-		BlendMode = MaterialModel->DynamicMaterialInstance->BlendMode;
-
-		ShadingModel = MaterialModel->DynamicMaterialInstance->GetShadingModels().HasShadingModel(EMaterialShadingModel::MSM_DefaultLit)
-			? EDMMaterialShadingModel::DefaultLit
-			: EDMMaterialShadingModel::Unlit;
-	}
-	else
-	{
-		BlendMode = InBlendMode;
-		ShadingModel = InShadingModel;
-	}
-}
-
-void UDynamicMaterialModelEditorOnlyData::LoadDeprecatedModelData_Expressions(TArray<TObjectPtr<UMaterialExpression>>& InExpressions)
-{
-	// Expressions should be parented to the material.
-	Expressions = InExpressions;
-	InExpressions.Empty();
-}
-
-void UDynamicMaterialModelEditorOnlyData::LoadDeprecatedModelData_Properties(TMap<EDMMaterialPropertyType, TObjectPtr<UObject>>& InProperties)
-{
-	for (const TPair<EDMMaterialPropertyType, TObjectPtr<UObject>>& Pair : InProperties)
-	{
-		UDMMaterialProperty* InProperty = Cast<UDMMaterialProperty>(Pair.Value);
-
-		if (!InProperty)
-		{
-			continue;
-		}
-
-		if (const TObjectPtr<UDMMaterialProperty>* ThisPropertyPtr = Properties.Find(Pair.Key))
-		{
-			if (UDMMaterialProperty* ThisProperty = *ThisPropertyPtr)
-			{
-				ThisProperty->LoadDeprecatedModelData(InProperty);
-			}
-		}
-	}
-
-	InProperties.Empty();
-}
-
-void UDynamicMaterialModelEditorOnlyData::LoadDeprecatedModelData_Slots(TArray<TObjectPtr<UObject>>& InSlots)
-{
-	Slots.Empty();
-	Slots.Reserve(InSlots.Num());
-
-	for (TObjectPtr<UObject>& SlotObject : InSlots)
-	{
-		UDMMaterialSlot* Slot = Cast<UDMMaterialSlot>(SlotObject);
-
-		if (!Slot)
-		{
-			continue;
-		}
-
-		Slots.Add(Slot);
-
-		// Reparent from the model to the editor only data
-		Slot->Rename(nullptr, this, UE::DynamicMaterial::RenameFlags);
-	}
-
-	InSlots.Empty();
-}
-
-void UDynamicMaterialModelEditorOnlyData::LoadDeprecatedModelData_PropertySlotMap(TMap<EDMMaterialPropertyType, TObjectPtr<UObject>>& InPropertySlotMap)
-{
-	PropertySlotMap.Empty();
-	PropertySlotMap.Reserve(InPropertySlotMap.Num());
-
-	for (const TPair<EDMMaterialPropertyType, TObjectPtr<UObject>>& Pair : InPropertySlotMap)
-	{
-		UDMMaterialSlot* Slot = Cast<UDMMaterialSlot>(Pair.Value);
-
-		if (!Slot)
-		{
-			continue;
-		}
-
-		PropertySlotMap.Add(Pair.Key, Slot);
-	}
-
-	InPropertySlotMap.Empty();
-}
-
-void UDynamicMaterialModelEditorOnlyData::LoadDeprecatedModelData(UDynamicMaterialModel* InMaterialModel)
-{
-	MaterialModel = InMaterialModel;
-	check(InMaterialModel);
-
-	PRAGMA_DISABLE_DEPRECATION_WARNINGS
-
-	LoadDeprecatedModelData_Base(MaterialModel->bCreateMaterialPackage, MaterialModel->BlendMode, MaterialModel->ShadingModel);
-	LoadDeprecatedModelData_Expressions(MaterialModel->Expressions);
-	LoadDeprecatedModelData_Properties(InMaterialModel->Properties);
-	LoadDeprecatedModelData_Slots(InMaterialModel->Slots);
-	LoadDeprecatedModelData_PropertySlotMap(InMaterialModel->PropertySlotMap);
-
-	PRAGMA_ENABLE_DEPRECATION_WARNINGS
-}
-
 TSharedRef<FDMMaterialBuildState> UDynamicMaterialModelEditorOnlyData::CreateBuildState(UMaterial* InMaterialToBuild, bool bInDirtyAssets) const
 {
 	check(MaterialModel);
@@ -932,6 +845,20 @@ void UDynamicMaterialModelEditorOnlyData::ForEachMaterialPropertyType(TFunctionR
 	}
 }
 
+void UDynamicMaterialModelEditorOnlyData::ForEachMaterialPropertyType(TFunctionRef<EDMIterationResult(EDMMaterialPropertyType InType)> InCallable, 
+	EDMMaterialPropertyType InStart, EDMMaterialPropertyType InEnd) const
+{
+	for (uint8 PropertyIndex = static_cast<uint8>(InStart); PropertyIndex <= static_cast<uint8>(InEnd); ++PropertyIndex)
+	{
+		const EDMMaterialPropertyType Property = static_cast<EDMMaterialPropertyType>(PropertyIndex);
+
+		if (InCallable(Property) == EDMIterationResult::Break)
+		{
+			break;
+		}
+	}
+}
+
 void UDynamicMaterialModelEditorOnlyData::SetChannelListPreset(FName InPresetName)
 {
 	ChannelListPreset = InPresetName;
@@ -975,9 +902,9 @@ UDMMaterialComponent* UDynamicMaterialModelEditorOnlyData::GetSubComponentByPath
 			{
 				EDMMaterialPropertyType EnumValue = static_cast<EDMMaterialPropertyType>(IntValue);
 
-				if (const TObjectPtr<UDMMaterialProperty>* PropertyPtr = Properties.Find(EnumValue))
+				if (UDMMaterialProperty* PropertyPtr = GetMaterialProperty(EnumValue))
 				{
-					return (*PropertyPtr)->GetComponentByPath(InPath);
+					return PropertyPtr->GetComponentByPath(InPath);
 				}
 			}
 		}
@@ -1004,9 +931,9 @@ TSharedRef<IDMMaterialBuildStateInterface> UDynamicMaterialModelEditorOnlyData::
 
 void UDynamicMaterialModelEditorOnlyData::SetPropertyComponent(EDMMaterialPropertyType InPropertyType, FName InComponentName, UDMMaterialComponent* InComponent)
 {
-	if (TObjectPtr<UDMMaterialProperty>* Property = Properties.Find(InPropertyType))
+	if (UDMMaterialProperty* Property = GetMaterialProperty(InPropertyType))
 	{
-		(*Property)->AddComponent(InComponentName, InComponent);
+		Property->AddComponent(InComponentName, InComponent);
 	}
 }
 
@@ -1219,25 +1146,91 @@ TMap<EDMMaterialPropertyType, UDMMaterialProperty*> UDynamicMaterialModelEditorO
 {
 	TMap<EDMMaterialPropertyType, UDMMaterialProperty*> LocalProperties;
 
-	LocalProperties.Reserve(Properties.Num());
-	for (const TPair<EDMMaterialPropertyType, TObjectPtr<UDMMaterialProperty>>& Pair : Properties)
-	{
-		LocalProperties.Add(Pair.Key, Pair.Value.Get());
-	}
+	ForEachMaterialPropertyType(
+		[this, &LocalProperties](EDMMaterialPropertyType InType)
+		{
+			if (UDMMaterialProperty* Property = GetMaterialProperty(InType))
+			{
+				LocalProperties.Emplace(InType, Property);
+			}
+
+			return EDMIterationResult::Continue;
+		}
+	);
 
 	return LocalProperties;
 }
 
-UDMMaterialProperty* UDynamicMaterialModelEditorOnlyData::GetMaterialProperty(EDMMaterialPropertyType MaterialProperty) const
+UDMMaterialProperty* UDynamicMaterialModelEditorOnlyData::GetMaterialProperty(EDMMaterialPropertyType InMaterialProperty) const
 {
-	TObjectPtr<UDMMaterialProperty> const* PropertyObjPtr = Properties.Find(MaterialProperty);
-
-	if (PropertyObjPtr)
+	switch (InMaterialProperty)
 	{
-		return *PropertyObjPtr;
-	}
+		case EDMMaterialPropertyType::BaseColor:
+			return BaseColor;
 
-	return nullptr;
+		case EDMMaterialPropertyType::EmissiveColor:
+			return EmissiveColor;
+
+		case EDMMaterialPropertyType::Opacity:
+			return Opacity;
+
+		case EDMMaterialPropertyType::OpacityMask:
+			return OpacityMask;
+
+		case EDMMaterialPropertyType::Roughness:
+			return Roughness;
+
+		case EDMMaterialPropertyType::Specular:
+			return Specular;
+
+		case EDMMaterialPropertyType::Metallic:
+			return Metallic;
+
+		case EDMMaterialPropertyType::Normal:
+			return Normal;
+
+		case EDMMaterialPropertyType::PixelDepthOffset:
+			return PixelDepthOffset;
+
+		case EDMMaterialPropertyType::WorldPositionOffset:
+			return WorldPositionOffset;
+
+		case EDMMaterialPropertyType::AmbientOcclusion:
+			return AmbientOcclusion;
+
+		case EDMMaterialPropertyType::Anisotropy:
+			return Anisotropy;
+
+		case EDMMaterialPropertyType::Refraction:
+			return Refraction;
+
+		case EDMMaterialPropertyType::Tangent:
+			return Tangent;
+
+		case EDMMaterialPropertyType::Displacement:
+			return Displacement;
+
+		case EDMMaterialPropertyType::SubsurfaceColor:
+			return SubsurfaceColor;
+
+		case EDMMaterialPropertyType::SurfaceThickness:
+			return SurfaceThickness;
+
+		case EDMMaterialPropertyType::Custom1:
+			return Custom1;
+
+		case EDMMaterialPropertyType::Custom2:
+			return Custom2;
+
+		case EDMMaterialPropertyType::Custom3:
+			return Custom3;
+
+		case EDMMaterialPropertyType::Custom4:
+			return Custom4;
+
+		default:
+			return nullptr;
+	}	
 }
 
 UDMMaterialSlot* UDynamicMaterialModelEditorOnlyData::GetSlot(int32 Index) const
@@ -1422,13 +1415,13 @@ UDMMaterialSlot* UDynamicMaterialModelEditorOnlyData::RemoveSlotForMaterialPrope
 	return nullptr;
 }
 
-TArray<EDMMaterialPropertyType> UDynamicMaterialModelEditorOnlyData::GetMaterialPropertiesForSlot(const UDMMaterialSlot* Slot) const
+TArray<EDMMaterialPropertyType> UDynamicMaterialModelEditorOnlyData::GetMaterialPropertiesForSlot(const UDMMaterialSlot* InSlot) const
 {
 	TArray<EDMMaterialPropertyType> OutProperties;
 
 	for (const TPair<EDMMaterialPropertyType, TObjectPtr<UDMMaterialSlot>>& Pair : PropertySlotMap)
 	{
-		if (Pair.Value == Slot)
+		if (Pair.Value == InSlot)
 		{
 			OutProperties.Add(Pair.Key);
 		}
@@ -1437,39 +1430,40 @@ TArray<EDMMaterialPropertyType> UDynamicMaterialModelEditorOnlyData::GetMaterial
 	return OutProperties;
 }
 
-void UDynamicMaterialModelEditorOnlyData::AssignMaterialPropertyToSlot(EDMMaterialPropertyType Property, UDMMaterialSlot* Slot)
+void UDynamicMaterialModelEditorOnlyData::AssignMaterialPropertyToSlot(EDMMaterialPropertyType InProperty, UDMMaterialSlot* InSlot)
 {
-	if (!Slot)
+	if (!InSlot)
 	{
-		UnassignMaterialProperty(Property);
+		UnassignMaterialProperty(InProperty);
 		return;
 	}
 
-	check(Properties.Contains(Property));
+	UDMMaterialProperty* Property = GetMaterialProperty(InProperty);
+	check(Property);
 
-	PropertySlotMap.FindOrAdd(Property) = Slot;
-	Properties[Property]->ResetInputConnectionMap();
-	Slot->OnPropertiesUpdated();
+	PropertySlotMap.FindOrAdd(InProperty) = InSlot;
+	Property->ResetInputConnectionMap();
+	InSlot->OnPropertiesUpdated();
 
 	RequestMaterialBuild();
 }
 
-void UDynamicMaterialModelEditorOnlyData::UnassignMaterialProperty(EDMMaterialPropertyType Property)
+void UDynamicMaterialModelEditorOnlyData::UnassignMaterialProperty(EDMMaterialPropertyType InProperty)
 {
-	TObjectPtr<UDMMaterialSlot>* SlotPtr = PropertySlotMap.Find(Property);
+	TObjectPtr<UDMMaterialSlot>* SlotPtr = PropertySlotMap.Find(InProperty);
 	
 	if (!SlotPtr)
 	{
 		return;
 	}
 
-	PropertySlotMap.Remove(Property);
+	PropertySlotMap.Remove(InProperty);
 	(*SlotPtr)->OnPropertiesUpdated();
 
 	RequestMaterialBuild();
 }
 
-void UDynamicMaterialModelEditorOnlyData::NotifyPostChange(const FPropertyChangedEvent& InPropertyChangedEvent, class FEditPropertyChain* PropertyThatChanged)
+void UDynamicMaterialModelEditorOnlyData::NotifyPostChange(const FPropertyChangedEvent& InPropertyChangedEvent, class FEditPropertyChain* InPropertyThatChanged)
 {
 	RequestMaterialBuild();
 }
@@ -1525,11 +1519,11 @@ void UDynamicMaterialModelEditorOnlyData::PostDuplicate(bool bInDuplicateForPIE)
 	}
 }
 
-void UDynamicMaterialModelEditorOnlyData::PostEditChangeChainProperty(FPropertyChangedChainEvent& PropertyChangedEvent)
+void UDynamicMaterialModelEditorOnlyData::PostEditChangeChainProperty(FPropertyChangedChainEvent& InPropertyChangedEvent)
 {
-	UObject::PostEditChangeChainProperty(PropertyChangedEvent);
+	UObject::PostEditChangeChainProperty(InPropertyChangedEvent);
 
-	const FName Property = PropertyChangedEvent.GetMemberPropertyName();
+	const FName Property = InPropertyChangedEvent.GetMemberPropertyName();
 
 	if (Property == GET_MEMBER_NAME_CHECKED(ThisClass, ChannelListPreset))
 	{
@@ -1557,13 +1551,13 @@ void UDynamicMaterialModelEditorOnlyData::PostEditChangeChainProperty(FPropertyC
 	}
 }
 
-void UDynamicMaterialModelEditorOnlyData::Serialize(FArchive& Ar)
+void UDynamicMaterialModelEditorOnlyData::Serialize(FArchive& InAr)
 {
-	Ar.UsingCustomVersion(FDynamicMaterialModelEditorOnlyDataVersion::GUID);
+	InAr.UsingCustomVersion(FDynamicMaterialModelEditorOnlyDataVersion::GUID);
 
-	Super::Serialize(Ar);
+	Super::Serialize(InAr);
 
-	const int32 Version = Ar.CustomVer(FDynamicMaterialModelEditorOnlyDataVersion::GUID);
+	const int32 Version = InAr.CustomVer(FDynamicMaterialModelEditorOnlyDataVersion::GUID);
 
 	// The default blend mode was changed from translucent and to opaque. If we have a
 	// transparent channel, let's set it back to translucent.
@@ -1650,22 +1644,25 @@ FString UDynamicMaterialModelEditorOnlyData::GetMaterialAssetName() const
 	return GetName() + "_Mat";
 }
 
-FString UDynamicMaterialModelEditorOnlyData::GetMaterialPackageName(const FString& MaterialBaseName) const
+FString UDynamicMaterialModelEditorOnlyData::GetMaterialPackageName(const FString& InMaterialBaseName) const
 {
 	return GetPackage()->GetName() + "_Mat";
 }
 
-void UDynamicMaterialModelEditorOnlyData::OnSlotConnectorsUpdated(UDMMaterialSlot* Slot)
+void UDynamicMaterialModelEditorOnlyData::OnSlotConnectorsUpdated(UDMMaterialSlot* InSlot)
 {
-	check(Slot);
+	check(InSlot);
 
 	RequestMaterialBuild();
 		
-	TArray<EDMMaterialPropertyType> SlotProperties = GetMaterialPropertiesForSlot(Slot);
+	TArray<EDMMaterialPropertyType> SlotProperties = GetMaterialPropertiesForSlot(InSlot);
 
-	for (EDMMaterialPropertyType Property : SlotProperties)
+	for (EDMMaterialPropertyType SlotProperty : SlotProperties)
 	{
-		Properties[Property]->ResetInputConnectionMap();
+		if (UDMMaterialProperty* Property = GetMaterialProperty(SlotProperty))
+		{
+			Property->ResetInputConnectionMap();
+		}		
 	}
 }
 
@@ -1692,18 +1689,22 @@ void UDynamicMaterialModelEditorOnlyData::PostEditorDuplicate()
 		Modify();
 	}
 
-	for (const TPair<EDMMaterialPropertyType, TObjectPtr<UDMMaterialProperty>>& Pair : Properties)
-	{
-		if (UDMMaterialProperty* Property = Pair.Value.Get())
+	ForEachMaterialPropertyType(
+		[this](EDMMaterialPropertyType InType)
 		{
-			if (GUndo)
+			if (UDMMaterialProperty* Property = GetMaterialProperty(InType))
 			{
-				Property->Modify();
+				if (GUndo)
+				{
+					Property->Modify();
+				}
+
+				Property->PostEditorDuplicate(MaterialModel, nullptr);
 			}
 
-			Property->PostEditorDuplicate(MaterialModel, nullptr);
+			return EDMIterationResult::Continue;
 		}
-	}
+	);
 
 	for (UDMMaterialSlot* Slot : Slots)
 	{
diff --git a/Engine/Plugins/Experimental/DynamicMaterial/Source/DynamicMaterialEditor/Public/Model/DynamicMaterialModelEditorOnlyData.h b/Engine/Plugins/Experimental/DynamicMaterial/Source/DynamicMaterialEditor/Public/Model/DynamicMaterialModelEditorOnlyData.h
--- a/Engine/Plugins/Experimental/DynamicMaterial/Source/DynamicMaterialEditor/Public/Model/DynamicMaterialModelEditorOnlyData.h
+++ b/Engine/Plugins/Experimental/DynamicMaterial/Source/DynamicMaterialEditor/Public/Model/DynamicMaterialModelEditorOnlyData.h
@@ -65,6 +65,9 @@ public:
 	DYNAMICMATERIALEDITOR_API static const FString AnisotropyPathToken;
 	DYNAMICMATERIALEDITOR_API static const FString RefractionPathToken;
 	DYNAMICMATERIALEDITOR_API static const FString TangentPathToken;
+	DYNAMICMATERIALEDITOR_API static const FString DisplacementPathToken;
+	DYNAMICMATERIALEDITOR_API static const FString SubsurfaceColorPathToken;
+	DYNAMICMATERIALEDITOR_API static const FString SurfaceThicknessPathToken;
 	DYNAMICMATERIALEDITOR_API static const FString Custom1PathToken;
 	DYNAMICMATERIALEDITOR_API static const FString Custom2PathToken;
 	DYNAMICMATERIALEDITOR_API static const FString Custom3PathToken;
@@ -161,7 +164,7 @@ public:
 	DYNAMICMATERIALEDITOR_API TMap<EDMMaterialPropertyType, UDMMaterialProperty*> GetMaterialProperties() const;
 
 	UFUNCTION(BlueprintPure, Category = "Material Designer")
-	DYNAMICMATERIALEDITOR_API UDMMaterialProperty* GetMaterialProperty(EDMMaterialPropertyType MaterialProperty) const;
+	DYNAMICMATERIALEDITOR_API UDMMaterialProperty* GetMaterialProperty(EDMMaterialPropertyType InMaterialProperty) const;
 
 	UFUNCTION(BlueprintPure, Category = "Material Designer")
 	const TArray<UDMMaterialSlot*>& GetSlots() const { return Slots; }
@@ -188,13 +191,13 @@ public:
 	DYNAMICMATERIALEDITOR_API UDMMaterialSlot* RemoveSlotForMaterialProperty(EDMMaterialPropertyType InType);
 
 	UFUNCTION(BlueprintPure, Category = "Material Designer")
-	DYNAMICMATERIALEDITOR_API TArray<EDMMaterialPropertyType> GetMaterialPropertiesForSlot(const UDMMaterialSlot* Slot) const;
+	DYNAMICMATERIALEDITOR_API TArray<EDMMaterialPropertyType> GetMaterialPropertiesForSlot(const UDMMaterialSlot* InSlot) const;
 
 	UFUNCTION(BlueprintCallable, Category = "Material Designer")
-	DYNAMICMATERIALEDITOR_API void AssignMaterialPropertyToSlot(EDMMaterialPropertyType Property, UDMMaterialSlot* Slot);
+	DYNAMICMATERIALEDITOR_API void AssignMaterialPropertyToSlot(EDMMaterialPropertyType InProperty, UDMMaterialSlot* InSlot);
 
 	UFUNCTION(BlueprintCallable, Category = "Material Designer")
-	DYNAMICMATERIALEDITOR_API void UnassignMaterialProperty(EDMMaterialPropertyType Property);
+	DYNAMICMATERIALEDITOR_API void UnassignMaterialProperty(EDMMaterialPropertyType InProperty);
 	
 	FDMOnMaterialBuilt::RegistrationType& GetOnMaterialBuiltDelegate() { return OnMaterialBuiltDelegate; }
 	FDMOnValueListUpdated::RegistrationType& GetOnValueListUpdateDelegate() { return OnValueListUpdateDelegate; }
@@ -219,10 +222,14 @@ public:
 		EDMMaterialPropertyType InStart = static_cast<EDMMaterialPropertyType>(static_cast<uint8>(EDMMaterialPropertyType::None) + 1),
 		EDMMaterialPropertyType InEnd = static_cast<EDMMaterialPropertyType>(static_cast<uint8>(EDMMaterialPropertyType::Any) - 1));
 
+	DYNAMICMATERIALEDITOR_API void ForEachMaterialPropertyType(TFunctionRef<EDMIterationResult(EDMMaterialPropertyType InType)> InCallable,
+		EDMMaterialPropertyType InStart = static_cast<EDMMaterialPropertyType>(static_cast<uint8>(EDMMaterialPropertyType::None) + 1),
+		EDMMaterialPropertyType InEnd = static_cast<EDMMaterialPropertyType>(static_cast<uint8>(EDMMaterialPropertyType::Any) - 1)) const;
+
 	void SaveEditor();
 
 	//~ Begin FNotifyHook
-	DYNAMICMATERIALEDITOR_API virtual void NotifyPostChange(const FPropertyChangedEvent& InPropertyChangedEvent, class FEditPropertyChain* PropertyThatChanged);
+	DYNAMICMATERIALEDITOR_API virtual void NotifyPostChange(const FPropertyChangedEvent& InPropertyChangedEvent, class FEditPropertyChain* InPropertyThatChanged);
 	//~ End FNotifyHook
 
 	//~ Begin UObject
@@ -230,8 +237,8 @@ public:
 	DYNAMICMATERIALEDITOR_API virtual void PostEditUndo() override;
 	DYNAMICMATERIALEDITOR_API virtual void PostEditImport() override;
 	DYNAMICMATERIALEDITOR_API virtual void PostDuplicate(bool bInDuplicateForPIE) override;
-	DYNAMICMATERIALEDITOR_API virtual void PostEditChangeChainProperty(FPropertyChangedChainEvent& PropertyChangedEvent) override;
-	DYNAMICMATERIALEDITOR_API virtual void Serialize(FArchive& Ar) override;
+	DYNAMICMATERIALEDITOR_API virtual void PostEditChangeChainProperty(FPropertyChangedChainEvent& InPropertyChangedEvent) override;
+	DYNAMICMATERIALEDITOR_API virtual void Serialize(FArchive& InAr) override;
 	//~ End UObject
 
 	//~ Begin IDMBuildable
@@ -244,7 +251,6 @@ public:
 	DYNAMICMATERIALEDITOR_API virtual void OnValueListUpdate() override;
 	DYNAMICMATERIALEDITOR_API virtual void OnValueUpdated(UDMMaterialValue* InValue, EDMUpdateType InUpdateType) override;
 	DYNAMICMATERIALEDITOR_API virtual void OnTextureUVUpdated(UDMTextureUV* InTextureUV) override;
-	DYNAMICMATERIALEDITOR_API virtual void LoadDeprecatedModelData(UDynamicMaterialModel* InMaterialModel) override;
 	DYNAMICMATERIALEDITOR_API virtual TSharedRef<IDMMaterialBuildStateInterface> CreateBuildStateInterface(UMaterial* InMaterialToBuild) const override;
 	DYNAMICMATERIALEDITOR_API virtual void SetPropertyComponent(EDMMaterialPropertyType InPropertyType, FName InComponentName, UDMMaterialComponent* InComponent) override;
 	DYNAMICMATERIALEDITOR_API virtual UDMMaterialComponent* GetSubComponentByPath(FDMComponentPath& InPath) const override;
@@ -300,9 +306,6 @@ protected:
 		meta = (GetOptions = GetPresetOptions, NoResetToDefault, DisplayName = "Material Type Preset"))
 	FName ChannelListPreset;
 
-	UPROPERTY(VisibleInstanceOnly, BlueprintReadOnly, Instanced, Category = "Material Designer")
-	TMap<EDMMaterialPropertyType, TObjectPtr<UDMMaterialProperty>> Properties;
-
 	UPROPERTY(VisibleInstanceOnly, BlueprintReadOnly, TextExportTransient, Category = "Material Designer")
 	TMap<EDMMaterialPropertyType, TObjectPtr<UDMMaterialSlot>> PropertySlotMap;
 
@@ -322,6 +325,69 @@ protected:
 	FDMOnValueListUpdated OnValueListUpdateDelegate;
 	FDMOnSlotListUpdated OnSlotListUpdateDelegate;
 
+	UPROPERTY(VisibleInstanceOnly, BlueprintReadOnly, Instanced, Category = "Material Designer")
+	TObjectPtr<UDMMaterialProperty> BaseColor;
+
+	UPROPERTY(VisibleInstanceOnly, BlueprintReadOnly, Instanced, Category = "Material Designer")
+	TObjectPtr<UDMMaterialProperty> EmissiveColor;
+
+	UPROPERTY(VisibleInstanceOnly, BlueprintReadOnly, Instanced, Category = "Material Designer")
+	TObjectPtr<UDMMaterialProperty> Opacity;
+
+	UPROPERTY(VisibleInstanceOnly, BlueprintReadOnly, Instanced, Category = "Material Designer")
+	TObjectPtr<UDMMaterialProperty> OpacityMask;
+
+	UPROPERTY(VisibleInstanceOnly, BlueprintReadOnly, Instanced, Category = "Material Designer")
+	TObjectPtr<UDMMaterialProperty> Roughness;
+
+	UPROPERTY(VisibleInstanceOnly, BlueprintReadOnly, Instanced, Category = "Material Designer")
+	TObjectPtr<UDMMaterialProperty> Specular;
+
+	UPROPERTY(VisibleInstanceOnly, BlueprintReadOnly, Instanced, Category = "Material Designer")
+	TObjectPtr<UDMMaterialProperty> Metallic;
+
+	UPROPERTY(VisibleInstanceOnly, BlueprintReadOnly, Instanced, Category = "Material Designer")
+	TObjectPtr<UDMMaterialProperty> Normal;
+
+	UPROPERTY(VisibleInstanceOnly, BlueprintReadOnly, Instanced, Category = "Material Designer")
+	TObjectPtr<UDMMaterialProperty> PixelDepthOffset;
+
+	UPROPERTY(VisibleInstanceOnly, BlueprintReadOnly, Instanced, Category = "Material Designer")
+	TObjectPtr<UDMMaterialProperty> WorldPositionOffset;
+
+	UPROPERTY(VisibleInstanceOnly, BlueprintReadOnly, Instanced, Category = "Material Designer")
+	TObjectPtr<UDMMaterialProperty> AmbientOcclusion;
+
+	UPROPERTY(VisibleInstanceOnly, BlueprintReadOnly, Instanced, Category = "Material Designer")
+	TObjectPtr<UDMMaterialProperty> Anisotropy;
+
+	UPROPERTY(VisibleInstanceOnly, BlueprintReadOnly, Instanced, Category = "Material Designer")
+	TObjectPtr<UDMMaterialProperty> Refraction;
+
+	UPROPERTY(VisibleInstanceOnly, BlueprintReadOnly, Instanced, Category = "Material Designer")
+	TObjectPtr<UDMMaterialProperty> Tangent;
+
+	UPROPERTY(VisibleInstanceOnly, BlueprintReadOnly, Instanced, Category = "Material Designer")
+	TObjectPtr<UDMMaterialProperty> Displacement;
+
+	UPROPERTY(VisibleInstanceOnly, BlueprintReadOnly, Instanced, Category = "Material Designer")
+	TObjectPtr<UDMMaterialProperty> SubsurfaceColor;
+
+	UPROPERTY(VisibleInstanceOnly, BlueprintReadOnly, Instanced, Category = "Material Designer")
+	TObjectPtr<UDMMaterialProperty> SurfaceThickness;
+
+	UPROPERTY(VisibleInstanceOnly, BlueprintReadOnly, Instanced, Category = "Material Designer")
+	TObjectPtr<UDMMaterialProperty> Custom1;
+
+	UPROPERTY(VisibleInstanceOnly, BlueprintReadOnly, Instanced, Category = "Material Designer")
+	TObjectPtr<UDMMaterialProperty> Custom2;
+
+	UPROPERTY(VisibleInstanceOnly, BlueprintReadOnly, Instanced, Category = "Material Designer")
+	TObjectPtr<UDMMaterialProperty> Custom3;
+
+	UPROPERTY(VisibleInstanceOnly, BlueprintReadOnly, Instanced, Category = "Material Designer")
+	TObjectPtr<UDMMaterialProperty> Custom4;
+
 	void CreateMaterial();
 
 	UFUNCTION(BlueprintCallable, Category = "Material Designer")
@@ -329,9 +395,9 @@ protected:
 
 	FString GetMaterialAssetPath() const;
 	FString GetMaterialAssetName() const;
-	FString GetMaterialPackageName(const FString& MaterialBaseName) const;
+	FString GetMaterialPackageName(const FString& InMaterialBaseName) const;
 
-	void OnSlotConnectorsUpdated(UDMMaterialSlot* Slot);
+	void OnSlotConnectorsUpdated(UDMMaterialSlot* InSlot);
 
 	/** Swaps the material properties from one slot to another, unless both slots exist and/or are the same. */
 	void SwapSlotMaterialProperty(EDMMaterialPropertyType InPropertyFrom, EDMMaterialPropertyType InPropertyTo);
@@ -357,10 +423,4 @@ protected:
 	//~ Begin IDynamicMaterialModelEditorOnlyDataInterface
 	virtual void ReinitComponents() override;
 	//~ End IDynamicMaterialModelEditorOnlyDataInterface
-
-	void LoadDeprecatedModelData_Base(bool bInCreateMaterialPackage, EBlendMode InBlendMode, EDMMaterialShadingModel InShadingModel);
-	void LoadDeprecatedModelData_Expressions(TArray<TObjectPtr<UMaterialExpression>>& InExpressions);
-	void LoadDeprecatedModelData_Properties(TMap<EDMMaterialPropertyType, TObjectPtr<UObject>>& InProperties);
-	void LoadDeprecatedModelData_Slots(TArray<TObjectPtr<UObject>>& InSlots);
-	void LoadDeprecatedModelData_PropertySlotMap(TMap<EDMMaterialPropertyType, TObjectPtr<UObject>>& InPropertySlotMap);
 };
diff --git a/Engine/Plugins/Experimental/DynamicMaterial/Source/DynamicMaterialTextureSet/Public/DMTextureSetMaterialProperty.h b/Engine/Plugins/Experimental/DynamicMaterial/Source/DynamicMaterialTextureSet/Public/DMTextureSetMaterialProperty.h
--- a/Engine/Plugins/Experimental/DynamicMaterial/Source/DynamicMaterialTextureSet/Public/DMTextureSetMaterialProperty.h
+++ b/Engine/Plugins/Experimental/DynamicMaterial/Source/DynamicMaterialTextureSet/Public/DMTextureSetMaterialProperty.h
@@ -32,5 +32,7 @@ enum class EDMTextureSetMaterialProperty : uint8
 	WorldPositionOffset = EMaterialProperty::MP_WorldPositionOffset,
 	PixelDepthOffset = EMaterialProperty::MP_PixelDepthOffset,
 
+	SurfaceThickness = EMaterialProperty::MP_SurfaceThickness,
+
 	None = EMaterialProperty::MP_MAX UMETA(Hidden)
 };
diff --git a/Engine/Plugins/Experimental/DynamicMaterial/Source/DynamicMaterialTextureSetEditor/Private/DMTextureSetSettings.cpp b/Engine/Plugins/Experimental/DynamicMaterial/Source/DynamicMaterialTextureSetEditor/Private/DMTextureSetSettings.cpp
--- a/Engine/Plugins/Experimental/DynamicMaterial/Source/DynamicMaterialTextureSetEditor/Private/DMTextureSetSettings.cpp
+++ b/Engine/Plugins/Experimental/DynamicMaterial/Source/DynamicMaterialTextureSetEditor/Private/DMTextureSetSettings.cpp
@@ -76,6 +76,11 @@ UDMTextureSetSettings::UDMTextureSetSettings()
 	SubsurfaceColor.MaterialProperties = {{EDMTextureSetMaterialProperty::SubsurfaceColor, EDMTextureChannelMask::RGBA}};
 	++FilterCount;
 
+	FDMTextureSetFilter SurfaceThickness;
+	SubsurfaceColor.FilterStrings = {TEXT("SurfaceThickness"), TEXT("Surface_Thickness"), TEXT("_ST")};
+	SubsurfaceColor.MaterialProperties = {{EDMTextureSetMaterialProperty::SurfaceThickness, EDMTextureChannelMask::Red}};
+	++FilterCount;
+
 	FDMTextureSetFilter Anisotropy;
 	Anisotropy.FilterStrings = {TEXT("Anisotropy"), TEXT("_Ani")};
 	Anisotropy.MaterialProperties = {{EDMTextureSetMaterialProperty::Anisotropy, EDMTextureChannelMask::RGBA}};
@@ -141,6 +146,7 @@ UDMTextureSetSettings::UDMTextureSetSettings()
 	Filters.Add(MoveTemp(WorldPositionOffset));
 	Filters.Add(MoveTemp(PixelDepthOffset));
 	Filters.Add(MoveTemp(SubsurfaceColor));
+	Filters.Add(MoveTemp(SurfaceThickness));
 	Filters.Add(MoveTemp(Anisotropy));
 	Filters.Add(MoveTemp(Refraction));
 	Filters.Add(MoveTemp(Tangent));
