Material Designer
- Fixed the UV link detection for Blends.
- The details section now refresh when the UV link is toggled and the mask stage is selected.

diff --git a/Engine/Plugins/Experimental/DynamicMaterial/Source/DynamicMaterialEditor/Private/Components/DMMaterialStageBlend.cpp b/Engine/Plugins/Experimental/DynamicMaterial/Source/DynamicMaterialEditor/Private/Components/DMMaterialStageBlend.cpp
--- a/Engine/Plugins/Experimental/DynamicMaterial/Source/DynamicMaterialEditor/Private/Components/DMMaterialStageBlend.cpp
+++ b/Engine/Plugins/Experimental/DynamicMaterial/Source/DynamicMaterialEditor/Private/Components/DMMaterialStageBlend.cpp
@@ -575,53 +575,82 @@ FText UDMMaterialStageBlend::GetStageDescription() const
 	return Super::GetStageDescription();
 }
 
+bool UDMMaterialStageBlend::SupportsLayerMaskTextureUVLink() const
+{
+	UDMMaterialStageInput* StageInputB = GetInputB();
+
+	if (!StageInputB)
+	{
+		return false;
+	}
+
+	UDMMaterialStageInputThroughput* InputThroughput = Cast<UDMMaterialStageInputThroughput>(StageInputB);
+
+	if (!InputThroughput)
+	{
+		return false;
+	}
+
+	UDMMaterialStageThroughput* Throughput = InputThroughput->GetMaterialStageThroughput();
+
+	if (!Throughput)
+	{
+		return false;
+	}
+
+	return Throughput->SupportsLayerMaskTextureUVLink();
+}
+
 FDMExpressionInput UDMMaterialStageBlend::GetLayerMaskLinkTextureUVInputExpressions(const TSharedRef<FDMMaterialBuildState>& InBuildState) const
 {
 	FDMExpressionInput ExpressionInput = {};
- 
-	if (!SupportsLayerMaskTextureUVLink())
+
+	UDMMaterialStage* Stage = GetStage();
+
+	if (!Stage)
 	{
 		return ExpressionInput;
 	}
- 
-	UDMMaterialStage* Stage = GetStage();
-	check(Stage);
- 
+
+	UDMMaterialStageInput* StageInputB = GetInputB();
+
+	if (!StageInputB)
+	{
+		return ExpressionInput;
+	}
+
+	UDMMaterialStageInputThroughput* InputThroughput = Cast<UDMMaterialStageInputThroughput>(StageInputB);
+
+	if (!InputThroughput)
+	{
+		return ExpressionInput;
+	}
+
+	UDMMaterialStageThroughput* Throughput = InputThroughput->GetMaterialStageThroughput();
+
+	if (!Throughput)
+	{
+		return ExpressionInput;
+	}
+
 	const TArray<FDMMaterialStageConnection>& InputConnections = Stage->GetInputConnectionMap();
-	const TArray<UDMMaterialStageInput*>& StageInputs = Stage->GetInputs();
- 
-	FDMMaterialStageConnectorChannel Channel;
- 
-	for (int32 InputIdx = InputA; InputIdx <= InputB; ++InputIdx)
+
+	if (!InputConnections.IsValidIndex(InputB) || InputConnections[InputB].Channels.Num() != 1)
 	{
-		if (!InputConnections.IsValidIndex(InputIdx)
-			|| InputConnections[InputIdx].Channels.Num() != 1
-			|| InputConnections[InputIdx].Channels[0].SourceIndex < FDMMaterialStageConnectorChannel::FIRST_STAGE_INPUT)
-		{
-			continue;
-		}
- 
-		int32 StageInputIdx = InputConnections[InputIdx].Channels[0].SourceIndex - FDMMaterialStageConnectorChannel::FIRST_STAGE_INPUT;
- 
-		ExpressionInput.OutputIndex = ResolveLayerMaskTextureUVLinkInputImpl(
-			InBuildState, 
-			StageInputs[StageInputIdx], 
-			Channel, 
-			ExpressionInput.OutputExpressions
-		);
-		
-		ExpressionInput.OutputChannel = Channel.OutputChannel;
-		
-		if (ExpressionInput.IsValid())
-		{
-			break;
-		}
-		else
-		{
-			ExpressionInput = {};
-		}
+		return ExpressionInput;
 	}
- 
+
+	FDMMaterialStageConnectorChannel Channel = InputConnections[InputB].Channels[0];
+
+	ExpressionInput.OutputIndex = ResolveLayerMaskTextureUVLinkInputImpl(
+		InBuildState,
+		StageInputB,
+		Channel,
+		ExpressionInput.OutputExpressions
+	);
+
+	ExpressionInput.OutputChannel = Channel.OutputChannel;
+
 	return ExpressionInput;
 }
  
diff --git a/Engine/Plugins/Experimental/DynamicMaterial/Source/DynamicMaterialEditor/Private/UI/Widgets/Editor/SlotEditor/SDMMaterialSlotLayerItem.cpp b/Engine/Plugins/Experimental/DynamicMaterial/Source/DynamicMaterialEditor/Private/UI/Widgets/Editor/SlotEditor/SDMMaterialSlotLayerItem.cpp
--- a/Engine/Plugins/Experimental/DynamicMaterial/Source/DynamicMaterialEditor/Private/UI/Widgets/Editor/SlotEditor/SDMMaterialSlotLayerItem.cpp
+++ b/Engine/Plugins/Experimental/DynamicMaterial/Source/DynamicMaterialEditor/Private/UI/Widgets/Editor/SlotEditor/SDMMaterialSlotLayerItem.cpp
@@ -17,6 +17,7 @@
 #include "Styling/StyleColors.h"
 #include "UI/DragDrop/DMSlotLayerDragDropOperation.h"
 #include "UI/Utils/DMWidgetStatics.h"
+#include "UI/Widgets/Editor/SDMMaterialComponentEditor.h"
 #include "UI/Widgets/Editor/SDMMaterialSlotEditor.h"
 #include "UI/Widgets/Editor/SlotEditor/SDMMaterialSlotLayerEffectView.h"
 #include "UI/Widgets/Editor/SlotEditor/SDMMaterialSlotLayerView.h"
@@ -738,6 +739,29 @@ FReply SDMMaterialSlotLayerItem::OnLayerLinkToggleButton()
 		FDMScopedUITransaction Transaction(LOCTEXT("UVLayerLinkToggle", "Toggle Layer UV Link"));
 		Layer->Modify();
 		Layer->ToggleTextureUVLinkEnabled();
+
+		if (MaskStageWidget.IsValid())
+		{
+			if (UDMMaterialStage* MaskStage = MaskStageWidget->GetStage())
+			{
+				if (TSharedPtr<SDMMaterialSlotLayerView> LayerView = GetSlotLayerView())
+				{
+					if (TSharedPtr<SDMMaterialSlotEditor> SlotEditor = LayerView->GetSlotEditorWidget())
+					{
+						if (TSharedPtr<SDMMaterialEditor> EditorWidget = SlotEditor->GetEditorWidget())
+						{
+							if (TSharedPtr<SDMMaterialComponentEditor> ComponentEditor = EditorWidget->GetComponentEditorWidget())
+							{
+								if (ComponentEditor->GetObject() == MaskStage)
+								{
+									EditorWidget->EditComponent(MaskStage, /* Force Refresh */ true);
+								}
+							}
+						}
+					}
+				}
+			}
+		}
 	}
 
 	return FReply::Handled();
diff --git a/Engine/Plugins/Experimental/DynamicMaterial/Source/DynamicMaterialEditor/Public/Components/DMMaterialStageBlend.h b/Engine/Plugins/Experimental/DynamicMaterial/Source/DynamicMaterialEditor/Public/Components/DMMaterialStageBlend.h
--- a/Engine/Plugins/Experimental/DynamicMaterial/Source/DynamicMaterialEditor/Public/Components/DMMaterialStageBlend.h
+++ b/Engine/Plugins/Experimental/DynamicMaterial/Source/DynamicMaterialEditor/Public/Components/DMMaterialStageBlend.h
@@ -54,7 +54,7 @@ public:
 
 	//~ Begin UDMMaterialStageSource
 	DYNAMICMATERIALEDITOR_API virtual FText GetStageDescription() const override;
-	virtual bool SupportsLayerMaskTextureUVLink() const override { return true; }
+	DYNAMICMATERIALEDITOR_API virtual bool SupportsLayerMaskTextureUVLink() const override;
 	DYNAMICMATERIALEDITOR_API virtual FDMExpressionInput GetLayerMaskLinkTextureUVInputExpressions(const TSharedRef<FDMMaterialBuildState>& InBuildState) const override;
 	DYNAMICMATERIALEDITOR_API virtual void GetMaskAlphaBlendNode(const TSharedRef<FDMMaterialBuildState>& InBuildState, 
 		UMaterialExpression*& OutExpression, int32& OutOutputIndex, int32& OutOutputChannel) const override;
