Material Designer
- When a material function texture input has 'use preview value as default' and a texture object node connected, it will be used as the default value for that input.

diff --git a/Engine/Plugins/Experimental/DynamicMaterial/Source/DynamicMaterialEditor/Private/Components/DMMaterialEffectFunction.cpp b/Engine/Plugins/Experimental/DynamicMaterial/Source/DynamicMaterialEditor/Private/Components/DMMaterialEffectFunction.cpp
--- a/Engine/Plugins/Experimental/DynamicMaterial/Source/DynamicMaterialEditor/Private/Components/DMMaterialEffectFunction.cpp
+++ b/Engine/Plugins/Experimental/DynamicMaterial/Source/DynamicMaterialEditor/Private/Components/DMMaterialEffectFunction.cpp
@@ -1,6 +1,7 @@
 // Copyright Epic Games, Inc. All Rights Reserved.
 
 #include "Components/DMMaterialEffectFunction.h"
+
 #include "Components/DMMaterialSlot.h"
 #include "Components/DMMaterialValue.h"
 #include "Components/MaterialValues/DMMaterialValueFloat1.h"
@@ -8,10 +9,10 @@
 #include "Components/MaterialValues/DMMaterialValueFloat3RGB.h"
 #include "Components/MaterialValues/DMMaterialValueFloat3XYZ.h"
 #include "Components/MaterialValues/DMMaterialValueFloat4.h"
+#include "Components/MaterialValues/DMMaterialValueTexture.h"
 #include "DMComponentPath.h"
 #include "Dom/JsonObject.h"
 #include "Dom/JsonValue.h"
-#include "DynamicMaterialEditorModule.h"
 #include "DynamicMaterialEditorSettings.h"
 #include "Materials/MaterialExpressionComponentMask.h"
 #include "Materials/MaterialExpressionFunctionInput.h"
@@ -21,7 +22,7 @@
 #include "MaterialValueType.h"
 #include "Model/DMMaterialBuildState.h"
 #include "Model/DynamicMaterialModelEditorOnlyData.h"
-#include "Utils/DMMaterialFunctionLibrary.h"
+#include "Utils/DMMaterialFunctionFunctionLibrary.h"
 #include "Utils/DMPrivate.h"
 #include "Utils/DMUtils.h"
 
@@ -481,42 +482,12 @@ bool UDMMaterialEffectFunction::NeedsFunctionInit() const
 			continue;
 		}
 
-		EDMValueType ValueType = EDMValueType::VT_None;
+		const EDMValueType ValueType = UDMMaterialFunctionFunctionLibrary::GetInputValueType(FunctionInput);
 
-		switch (FunctionInput->InputType)
+		if (ValueType == EDMValueType::VT_None)
 		{
-			case EFunctionInputType::FunctionInput_Scalar:
-				ValueType = EDMValueType::VT_Float1;
-				break;
-
-			case EFunctionInputType::FunctionInput_Vector2:
-				ValueType = EDMValueType::VT_Float2;
-				break;
-
-			case EFunctionInputType::FunctionInput_Vector3:
-				if (UDynamicMaterialEditorSettings::IsUseLinearColorForVectorsEnabled())
-				{
-					ValueType = EDMValueType::VT_Float3_RGB;
-				}
-				else
-				{
-					ValueType = EDMValueType::VT_Float3_XYZ;
-				}
-				break;
-
-			case EFunctionInputType::FunctionInput_Vector4:
-				ValueType = EDMValueType::VT_Float4_RGBA;
-				break;
-
-			case EFunctionInputType::FunctionInput_Texture2D:
-			case EFunctionInputType::FunctionInput_TextureCube:
-			case EFunctionInputType::FunctionInput_VolumeTexture:
-				ValueType = EDMValueType::VT_Texture;
-				break;
-
-			default:
-				UE::DynamicMaterialEditor::Private::LogError(TEXT("Effect Function has invalid input type -  must be a scalar, vector or texture."), true, this);
-				return false;
+			UE::DynamicMaterialEditor::Private::LogError(TEXT("Effect Function has invalid input type -  must be a scalar, vector or texture."), true, this);
+			return false;
 		}
 
 		if (ValueType != InputValues[InputIndex]->GetType())
@@ -623,43 +594,13 @@ void UDMMaterialEffectFunction::InitFunction()
 			continue;
 		}
 
-		EDMValueType ValueType = EDMValueType::VT_None;
+		const EDMValueType ValueType = UDMMaterialFunctionFunctionLibrary::GetInputValueType(FunctionInput);
 
-		switch (FunctionInput->InputType)
+		if (ValueType == EDMValueType::VT_None)
 		{
-			case EFunctionInputType::FunctionInput_Scalar:
-				ValueType = EDMValueType::VT_Float1;
-				break;
-
-			case EFunctionInputType::FunctionInput_Vector2:
-				ValueType = EDMValueType::VT_Float2;
-				break;
-
-			case EFunctionInputType::FunctionInput_Vector3:
-				if (UDynamicMaterialEditorSettings::IsUseLinearColorForVectorsEnabled())
-				{
-					ValueType = EDMValueType::VT_Float3_RGB;
-				}
-				else
-				{
-					ValueType = EDMValueType::VT_Float3_XYZ;
-				}
-				break;
-
-			case EFunctionInputType::FunctionInput_Vector4:
-				ValueType = EDMValueType::VT_Float4_RGBA;
-				break;
-
-			case EFunctionInputType::FunctionInput_Texture2D:
-			case EFunctionInputType::FunctionInput_TextureCube:
-			case EFunctionInputType::FunctionInput_VolumeTexture:
-				ValueType = EDMValueType::VT_Texture;
-				break;
-
-			default:
-				UE::DynamicMaterialEditor::Private::LogError(TEXT("Effect Function has invalid input type -  must be a scalar, vector or texture."), true, this);
-				MaterialFunctionPtr = nullptr;
-				return;
+			UE::DynamicMaterialEditor::Private::LogError(TEXT("Effect Function has invalid input type -  must be a scalar, vector or texture."), true, this);
+			MaterialFunctionPtr = nullptr;
+			return;
 		}
 
 		UDMMaterialValue* Value = UDMMaterialValue::CreateMaterialValue(
@@ -672,52 +613,7 @@ void UDMMaterialEffectFunction::InitFunction()
 
 		InputValues.Add(Value);
 
-		if (FunctionInput->bUsePreviewValueAsDefault)
-		{
-			switch (FunctionInput->InputType)
-			{
-				case EFunctionInputType::FunctionInput_Scalar:
-					if (UDMMaterialValueFloat1* Float1Value = Cast<UDMMaterialValueFloat1>(Value))
-					{
-						Float1Value->SetDefaultValue(FunctionInput->PreviewValue.X);
-						Float1Value->ApplyDefaultValue();
-					}
-					break;
-
-				case EFunctionInputType::FunctionInput_Vector2:
-					if (UDMMaterialValueFloat2* Float2Value = Cast<UDMMaterialValueFloat2>(Value))
-					{
-						Float2Value->SetDefaultValue({FunctionInput->PreviewValue.X, FunctionInput->PreviewValue.Y});
-						Float2Value->ApplyDefaultValue();
-					}
-					break;
-
-				case EFunctionInputType::FunctionInput_Vector3:
-					if (UDMMaterialValueFloat3XYZ* Float3XYZ = Cast<UDMMaterialValueFloat3XYZ>(Value))
-					{
-						Float3XYZ->SetDefaultValue({FunctionInput->PreviewValue.X, FunctionInput->PreviewValue.Y, FunctionInput->PreviewValue.Z});
-						Float3XYZ->ApplyDefaultValue();
-					}
-					else if (UDMMaterialValueFloat3RGB* Float3RGB = Cast<UDMMaterialValueFloat3RGB>(Value))
-					{
-						Float3RGB->SetDefaultValue({FunctionInput->PreviewValue.X, FunctionInput->PreviewValue.Y, FunctionInput->PreviewValue.Z});
-						Float3RGB->ApplyDefaultValue();
-					}
-					break;
-
-				case EFunctionInputType::FunctionInput_Vector4:
-					if (UDMMaterialValueFloat4* Float4Value = Cast<UDMMaterialValueFloat4>(Value))
-					{
-						Float4Value->SetDefaultValue(FunctionInput->PreviewValue);
-						Float4Value->ApplyDefaultValue();
-					}
-					break;
-
-				default:
-					// Not possible
-					break;
-			}
-		}
+		UDMMaterialFunctionFunctionLibrary::SetInputDefault(FunctionInput, Value);
 
 		if (bSetValueAdded)
 		{
diff --git a/Engine/Plugins/Experimental/DynamicMaterial/Source/DynamicMaterialEditor/Private/Components/DMMaterialStageFunction.cpp b/Engine/Plugins/Experimental/DynamicMaterial/Source/DynamicMaterialEditor/Private/Components/DMMaterialStageFunction.cpp
--- a/Engine/Plugins/Experimental/DynamicMaterial/Source/DynamicMaterialEditor/Private/Components/DMMaterialStageFunction.cpp
+++ b/Engine/Plugins/Experimental/DynamicMaterial/Source/DynamicMaterialEditor/Private/Components/DMMaterialStageFunction.cpp
@@ -21,6 +21,7 @@
 #include "Model/DMMaterialBuildState.h"
 #include "Model/DMMaterialBuildUtils.h"
 #include "Model/DynamicMaterialModelEditorOnlyData.h"
+#include "Utils/DMMaterialFunctionFunctionLibrary.h"
 #include "Utils/DMPrivate.h"
 #include "Utils/DMUtils.h"
 
@@ -379,42 +380,12 @@ bool UDMMaterialStageFunction::NeedsFunctionInit() const
 			continue;
 		}
 
-		EDMValueType ValueType = EDMValueType::VT_None;
+		const EDMValueType ValueType = UDMMaterialFunctionFunctionLibrary::GetInputValueType(FunctionInput);
 
-		switch (FunctionInput->InputType)
+		if (ValueType == EDMValueType::VT_None)
 		{
-			case EFunctionInputType::FunctionInput_Scalar:
-				ValueType = EDMValueType::VT_Float1;
-				break;
-
-			case EFunctionInputType::FunctionInput_Vector2:
-				ValueType = EDMValueType::VT_Float2;
-				break;
-
-			case EFunctionInputType::FunctionInput_Vector3:
-				if (UDynamicMaterialEditorSettings::IsUseLinearColorForVectorsEnabled())
-				{
-					ValueType = EDMValueType::VT_Float3_RGB;
-				}
-				else
-				{
-					ValueType = EDMValueType::VT_Float3_XYZ;
-				}
-				break;
-
-			case EFunctionInputType::FunctionInput_Vector4:
-				ValueType = EDMValueType::VT_Float4_RGBA;
-				break;
-
-			case EFunctionInputType::FunctionInput_Texture2D:
-			case EFunctionInputType::FunctionInput_TextureCube:
-			case EFunctionInputType::FunctionInput_VolumeTexture:
-				ValueType = EDMValueType::VT_Texture;
-				break;
-
-			default:
-				UE::DynamicMaterialEditor::Private::LogError(TEXT("Function has invalid input type - must be a scalar, vector or texture."), true, this);
-				return false;
+			UE::DynamicMaterialEditor::Private::LogError(TEXT("Function has invalid input type - must be a scalar, vector or texture."), true, this);
+			return false;
 		}
 
 		if (ValueType != InputValues[InputIndex]->GetType())
@@ -533,44 +504,14 @@ void UDMMaterialStageFunction::InitFunction()
 			continue;
 		}
 
-		EDMValueType ValueType = EDMValueType::VT_None;
+		const EDMValueType ValueType = UDMMaterialFunctionFunctionLibrary::GetInputValueType(FunctionInput);
 
-		switch (FunctionInput->InputType)
+		if (ValueType == EDMValueType::VT_None)
 		{
-			case EFunctionInputType::FunctionInput_Scalar:
-				ValueType = EDMValueType::VT_Float1;
-				break;
-
-			case EFunctionInputType::FunctionInput_Vector2:
-				ValueType = EDMValueType::VT_Float2;
-				break;
-
-			case EFunctionInputType::FunctionInput_Vector3:
-				if (UDynamicMaterialEditorSettings::IsUseLinearColorForVectorsEnabled())
-				{
-					ValueType = EDMValueType::VT_Float3_RGB;
-				}
-				else
-				{
-					ValueType = EDMValueType::VT_Float3_XYZ;
-				}
-				break;
-
-			case EFunctionInputType::FunctionInput_Vector4:
-				ValueType = EDMValueType::VT_Float4_RGBA;
-				break;
-
-			case EFunctionInputType::FunctionInput_Texture2D:
-			case EFunctionInputType::FunctionInput_TextureCube:
-			case EFunctionInputType::FunctionInput_VolumeTexture:
-				ValueType = EDMValueType::VT_Texture;
-				break;
-
-			default:
-				UE::DynamicMaterialEditor::Private::LogError(TEXT("Function has invalid input type - must be a scalar, vector or texture."), true, this);
-				InputConnectors.SetNum(1);
-				MaterialFunction = nullptr;
-				return;
+			UE::DynamicMaterialEditor::Private::LogError(TEXT("Function has invalid input type - must be a scalar, vector or texture."), true, this);
+			InputConnectors.SetNum(1);
+			MaterialFunction = nullptr;
+			return;
 		}
 
 		InputConnectors[InputIndex].Index = InputIndex;
@@ -599,52 +540,7 @@ void UDMMaterialStageFunction::InitFunction()
 		UDMMaterialValue* Value = InputValue->GetValue();
 		check(Value);
 
-		if (FunctionInput->bUsePreviewValueAsDefault)
-		{
-			switch (FunctionInput->InputType)
-			{
-				case EFunctionInputType::FunctionInput_Scalar:
-					if (UDMMaterialValueFloat1* Float1Value = Cast<UDMMaterialValueFloat1>(Value))
-					{
-						Float1Value->SetDefaultValue(FunctionInput->PreviewValue.X);
-						Float1Value->ApplyDefaultValue();
-					}
-					break;
-
-				case EFunctionInputType::FunctionInput_Vector2:
-					if (UDMMaterialValueFloat2* Float2Value = Cast<UDMMaterialValueFloat2>(Value))
-					{
-						Float2Value->SetDefaultValue({FunctionInput->PreviewValue.X, FunctionInput->PreviewValue.Y});
-						Float2Value->ApplyDefaultValue();
-					}
-					break;
-
-				case EFunctionInputType::FunctionInput_Vector3:
-					if (UDMMaterialValueFloat3XYZ* Float3XYZ = Cast<UDMMaterialValueFloat3XYZ>(Value))
-					{
-						Float3XYZ->SetDefaultValue({FunctionInput->PreviewValue.X, FunctionInput->PreviewValue.Y, FunctionInput->PreviewValue.Z});
-						Float3XYZ->ApplyDefaultValue();
-					}
-					else if (UDMMaterialValueFloat3RGB* Float3RGB = Cast<UDMMaterialValueFloat3RGB>(Value))
-					{
-						Float3RGB->SetDefaultValue({FunctionInput->PreviewValue.X, FunctionInput->PreviewValue.Y, FunctionInput->PreviewValue.Z});
-						Float3RGB->ApplyDefaultValue();
-					}
-					break;
-
-				case EFunctionInputType::FunctionInput_Vector4:
-					if (UDMMaterialValueFloat4* Float4Value = Cast<UDMMaterialValueFloat4>(Value))
-					{
-						Float4Value->SetDefaultValue(FunctionInput->PreviewValue);
-						Float4Value->ApplyDefaultValue();
-					}
-					break;
-
-				default:
-					// Not possible
-					break;
-			}
-		}
+		UDMMaterialFunctionFunctionLibrary::SetInputDefault(FunctionInput, Value);
 	}
 }
 
diff --git a/Engine/Plugins/Experimental/DynamicMaterial/Source/DynamicMaterialEditor/Private/UI/PropertyGenerators/DMMaterialEffectFunctionPropertyRowGenerator.cpp b/Engine/Plugins/Experimental/DynamicMaterial/Source/DynamicMaterialEditor/Private/UI/PropertyGenerators/DMMaterialEffectFunctionPropertyRowGenerator.cpp
--- a/Engine/Plugins/Experimental/DynamicMaterial/Source/DynamicMaterialEditor/Private/UI/PropertyGenerators/DMMaterialEffectFunctionPropertyRowGenerator.cpp
+++ b/Engine/Plugins/Experimental/DynamicMaterial/Source/DynamicMaterialEditor/Private/UI/PropertyGenerators/DMMaterialEffectFunctionPropertyRowGenerator.cpp
@@ -9,8 +9,8 @@
 #include "Materials/MaterialExpressionFunctionInput.h"
 #include "Materials/MaterialExpressionMaterialFunctionCall.h"
 #include "Materials/MaterialFunctionInterface.h"
-#include "UI/PropertyGenerators/DMMaterialStageFunctionPropertyRowGenerator.h"
 #include "UI/Widgets/Editor/SDMMaterialComponentEditor.h"
+#include "Utils/DMMaterialFunctionFunctionLibrary.h"
 
 #define LOCTEXT_NAMESPACE "DMMaterialEffectFunctionPropertyRowGenerator"
 
@@ -108,7 +108,7 @@ void FDMMaterialEffectFunctionPropertyRowGenerator::AddComponentProperties(const
 
 			if (ValuePropertyRow.PropertyHandle.IsValid())
 			{
-				FDMMaterialStageFunctionPropertyRowGenerator::ApplyMetaData(
+				UDMMaterialFunctionFunctionLibrary::ApplyMetaData(
 					Inputs[InputIndex],
 					ValuePropertyRow.PropertyHandle.ToSharedRef()
 				);
diff --git a/Engine/Plugins/Experimental/DynamicMaterial/Source/DynamicMaterialEditor/Private/UI/PropertyGenerators/DMMaterialStageFunctionPropertyRowGenerator.cpp b/Engine/Plugins/Experimental/DynamicMaterial/Source/DynamicMaterialEditor/Private/UI/PropertyGenerators/DMMaterialStageFunctionPropertyRowGenerator.cpp
--- a/Engine/Plugins/Experimental/DynamicMaterial/Source/DynamicMaterialEditor/Private/UI/PropertyGenerators/DMMaterialStageFunctionPropertyRowGenerator.cpp
+++ b/Engine/Plugins/Experimental/DynamicMaterial/Source/DynamicMaterialEditor/Private/UI/PropertyGenerators/DMMaterialStageFunctionPropertyRowGenerator.cpp
@@ -11,6 +11,7 @@
 #include "Materials/MaterialExpressionMaterialFunctionCall.h"
 #include "Materials/MaterialFunctionInterface.h"
 #include "UI/Widgets/Editor/SDMMaterialComponentEditor.h"
+#include "Utils/DMMaterialFunctionFunctionLibrary.h"
 
 #define LOCTEXT_NAMESPACE "DMMaterialStageFunctionPropertyRowGenerator"
 
@@ -20,62 +21,6 @@ const TSharedRef<FDMMaterialStageFunctionPropertyRowGenerator>& FDMMaterialStage
 	return Generator;
 }
 
-void FDMMaterialStageFunctionPropertyRowGenerator::ApplyMetaData(const FFunctionExpressionInput& InFunctionInput, 
-	const TSharedRef<IPropertyHandle>& InPropertyHandle)
-{
-	const TArray<FString> MetaNames = {TEXT("UIMin"), TEXT("UIMax"), TEXT("ClampMin"), TEXT("ClampMax")};
-
-	TArray<FString> MetaDatas;
-	InFunctionInput.ExpressionInput->Desc.ParseIntoArray(MetaDatas, TEXT(","));
-
-	for (FString& MetaData : MetaDatas)
-	{
-		const int32 EqualsPosition = MetaData.Find(TEXT("="));
-
-		if (EqualsPosition == INDEX_NONE)
-		{
-			continue;
-		}
-
-		const FString MetaDataName = MetaData.Left(EqualsPosition).TrimStartAndEnd();
-
-		bool bValidName = false;
-
-		for (const FString& Name : MetaNames)
-		{
-			if (MetaDataName.Equals(Name, ESearchCase::IgnoreCase))
-			{
-				bValidName = true;
-				break;
-			}
-		}
-
-		if (!bValidName)
-		{
-			continue;
-		}
-
-		const FString MetaDataValue = MetaData.Mid(EqualsPosition + 1).TrimStartAndEnd();
-		bool bValidValue = true;
-
-		for (int32 Index = 0; Index < MetaDataValue.Len(); ++Index)
-		{
-			if (MetaDataValue[Index] != '.'	&& MetaDataValue[Index] != '-' && (MetaDataValue[Index] < '0' || MetaDataValue[Index] > '9'))
-			{
-				bValidValue = false;
-				break;
-			}			
-		}
-
-		if (!bValidValue)
-		{
-			continue;
-		}
-
-		InPropertyHandle->SetInstanceMetaData(*MetaDataName, MetaDataValue);
-	}
-}
-
 void FDMMaterialStageFunctionPropertyRowGenerator::AddComponentProperties(const TSharedRef<SDMMaterialComponentEditor>& InComponentEditorWidget, UDMMaterialComponent* InComponent,
 	TArray<FDMPropertyHandle>& InOutPropertyRows, TSet<UDMMaterialComponent*>& InOutProcessedObjects)
 {
@@ -177,7 +122,10 @@ void FDMMaterialStageFunctionPropertyRowGenerator::AddComponentProperties(const
 						if (ValuePropertyRow.PropertyHandle.IsValid())
 						{
 							// The first input is the previous layer, so it does not have a value.
-							ApplyMetaData(Inputs[InputIndex + 1], ValuePropertyRow.PropertyHandle.ToSharedRef());
+							UDMMaterialFunctionFunctionLibrary::ApplyMetaData(
+								Inputs[InputIndex + 1],
+								ValuePropertyRow.PropertyHandle.ToSharedRef()
+							);
 						}
 					}
 
diff --git a/Engine/Plugins/Experimental/DynamicMaterial/Source/DynamicMaterialEditor/Private/UI/PropertyGenerators/DMMaterialStageFunctionPropertyRowGenerator.h b/Engine/Plugins/Experimental/DynamicMaterial/Source/DynamicMaterialEditor/Private/UI/PropertyGenerators/DMMaterialStageFunctionPropertyRowGenerator.h
--- a/Engine/Plugins/Experimental/DynamicMaterial/Source/DynamicMaterialEditor/Private/UI/PropertyGenerators/DMMaterialStageFunctionPropertyRowGenerator.h
+++ b/Engine/Plugins/Experimental/DynamicMaterial/Source/DynamicMaterialEditor/Private/UI/PropertyGenerators/DMMaterialStageFunctionPropertyRowGenerator.h
@@ -4,9 +4,6 @@
 
 #include "UI/PropertyGenerators/DMComponentPropertyRowGenerator.h"
 
-class IPropertyHandle;
-struct FFunctionExpressionInput;
-
 class FDMMaterialStageFunctionPropertyRowGenerator : public FDMComponentPropertyRowGenerator
 {
 public:
@@ -14,8 +11,6 @@ public:
 
 	virtual ~FDMMaterialStageFunctionPropertyRowGenerator() override = default;
 
-	static void ApplyMetaData(const FFunctionExpressionInput& InFunctionInput, const TSharedRef<IPropertyHandle>& InPropertyHandle);
-
 	//~ Begin FDMComponentPropertyRowGenerator
 	virtual void AddComponentProperties(const TSharedRef<SDMMaterialComponentEditor>& InComponentEditorWidget, UDMMaterialComponent* InComponent,
 		TArray<FDMPropertyHandle>& InOutPropertyRows, TSet<UDMMaterialComponent*>& InOutProcessedObjects) override;
diff --git a/Engine/Plugins/Experimental/DynamicMaterial/Source/DynamicMaterialEditor/Private/Utils/DMMaterialFunctionFunctionLibrary.cpp b/Engine/Plugins/Experimental/DynamicMaterial/Source/DynamicMaterialEditor/Private/Utils/DMMaterialFunctionFunctionLibrary.cpp
new file mode 100644
--- /dev/null
+++ b/Engine/Plugins/Experimental/DynamicMaterial/Source/DynamicMaterialEditor/Private/Utils/DMMaterialFunctionFunctionLibrary.cpp
@@ -0,0 +1,197 @@
+// Copyright Epic Games, Inc. All Rights Reserved.
+
+#include "Utils/DMMaterialFunctionFunctionLibrary.h"
+
+#include "Components/MaterialValues/DMMaterialValueFloat1.h"
+#include "Components/MaterialValues/DMMaterialValueFloat2.h"
+#include "Components/MaterialValues/DMMaterialValueFloat3RGB.h"
+#include "Components/MaterialValues/DMMaterialValueFloat3XYZ.h"
+#include "Components/MaterialValues/DMMaterialValueFloat4.h"
+#include "Components/MaterialValues/DMMaterialValueTexture.h"
+#include "Containers/Array.h"
+#include "Containers/UnrealString.h"
+#include "DMDefs.h"
+#include "DynamicMaterialEditorSettings.h"
+#include "Materials/MaterialExpressionFunctionInput.h"
+#include "Materials/MaterialExpressionMaterialFunctionCall.h"
+#include "Materials/MaterialExpressionTextureObject.h"
+#include "PropertyHandle.h"
+#include "Templates/SharedPointer.h"
+
+#define LOCTEXT_NAMESPACE "DMMaterialFunctionFunctionLibrary"
+
+void UDMMaterialFunctionFunctionLibrary::ApplyMetaData(const FFunctionExpressionInput& InFunctionInput,
+	const TSharedRef<IPropertyHandle>& InPropertyHandle)
+{
+	const TArray<FString> MetaNames = {TEXT("UIMin"), TEXT("UIMax"), TEXT("ClampMin"), TEXT("ClampMax")};
+
+	TArray<FString> MetaDatas;
+	InFunctionInput.ExpressionInput->Desc.ParseIntoArray(MetaDatas, TEXT(","));
+
+	for (FString& MetaData : MetaDatas)
+	{
+		const int32 EqualsPosition = MetaData.Find(TEXT("="));
+
+		if (EqualsPosition == INDEX_NONE)
+		{
+			continue;
+		}
+
+		const FString MetaDataName = MetaData.Left(EqualsPosition).TrimStartAndEnd();
+
+		bool bValidName = false;
+
+		for (const FString& Name : MetaNames)
+		{
+			if (MetaDataName.Equals(Name, ESearchCase::IgnoreCase))
+			{
+				bValidName = true;
+				break;
+			}
+		}
+
+		if (!bValidName)
+		{
+			continue;
+		}
+
+		const FString MetaDataValue = MetaData.Mid(EqualsPosition + 1).TrimStartAndEnd();
+		bool bValidValue = true;
+
+		for (int32 Index = 0; Index < MetaDataValue.Len(); ++Index)
+		{
+			if (MetaDataValue[Index] != '.' && MetaDataValue[Index] != '-' && (MetaDataValue[Index] < '0' || MetaDataValue[Index] > '9'))
+			{
+				bValidValue = false;
+				break;
+			}
+		}
+
+		if (!bValidValue)
+		{
+			continue;
+		}
+
+		InPropertyHandle->SetInstanceMetaData(*MetaDataName, MetaDataValue);
+	}
+}
+
+EDMValueType UDMMaterialFunctionFunctionLibrary::GetInputValueType(UMaterialExpressionFunctionInput* InFunctionInput)
+{
+	if (!InFunctionInput)
+	{
+		return EDMValueType::VT_None;
+	}
+
+	switch (InFunctionInput->InputType)
+	{
+		case EFunctionInputType::FunctionInput_Scalar:
+			return EDMValueType::VT_Float1;
+
+		case EFunctionInputType::FunctionInput_Vector2:
+			return EDMValueType::VT_Float2;
+
+		case EFunctionInputType::FunctionInput_Vector3:
+			if (UDynamicMaterialEditorSettings::IsUseLinearColorForVectorsEnabled())
+			{
+				return EDMValueType::VT_Float3_RGB;
+			}
+
+			return EDMValueType::VT_Float3_XYZ;
+
+		case EFunctionInputType::FunctionInput_Vector4:
+			return EDMValueType::VT_Float4_RGBA;
+
+		case EFunctionInputType::FunctionInput_Texture2D:
+		case EFunctionInputType::FunctionInput_TextureCube:
+		case EFunctionInputType::FunctionInput_VolumeTexture:
+			return EDMValueType::VT_Texture;
+
+		default:
+			return EDMValueType::VT_None;
+	}
+}
+
+void UDMMaterialFunctionFunctionLibrary::SetInputDefault(UMaterialExpressionFunctionInput* InFunctionInput, UDMMaterialValue* InValue)
+{
+	if (!InFunctionInput || !InFunctionInput->bUsePreviewValueAsDefault)
+	{
+		return;
+	}
+
+	const EDMValueType ValueType = GetInputValueType(InFunctionInput);
+
+	switch (ValueType)
+	{
+		case EDMValueType::VT_Float1:
+			if (UDMMaterialValueFloat1* Float1Value = Cast<UDMMaterialValueFloat1>(InValue))
+			{
+				Float1Value->SetDefaultValue(InFunctionInput->PreviewValue.X);
+				Float1Value->ApplyDefaultValue();
+			}
+			break;
+
+		case EDMValueType::VT_Float2:
+			if (UDMMaterialValueFloat2* Float2Value = Cast<UDMMaterialValueFloat2>(InValue))
+			{
+				Float2Value->SetDefaultValue({InFunctionInput->PreviewValue.X, InFunctionInput->PreviewValue.Y});
+				Float2Value->ApplyDefaultValue();
+			}
+			break;
+
+		case EDMValueType::VT_Float3_XYZ:
+			if (UDMMaterialValueFloat3XYZ* Float3XYZ = Cast<UDMMaterialValueFloat3XYZ>(InValue))
+			{
+				Float3XYZ->SetDefaultValue({InFunctionInput->PreviewValue.X, InFunctionInput->PreviewValue.Y, InFunctionInput->PreviewValue.Z});
+				Float3XYZ->ApplyDefaultValue();
+			}
+			break;
+
+		case EDMValueType::VT_Float3_RGB:
+			if (UDMMaterialValueFloat3RGB* Float3RGB = Cast<UDMMaterialValueFloat3RGB>(InValue))
+			{
+				Float3RGB->SetDefaultValue({InFunctionInput->PreviewValue.X, InFunctionInput->PreviewValue.Y, InFunctionInput->PreviewValue.Z});
+				Float3RGB->ApplyDefaultValue();
+			}
+			break;
+
+		case EDMValueType::VT_Float4_RGBA:
+			if (UDMMaterialValueFloat4* Float4Value = Cast<UDMMaterialValueFloat4>(InValue))
+			{
+				Float4Value->SetDefaultValue(InFunctionInput->PreviewValue);
+				Float4Value->ApplyDefaultValue();
+			}
+			break;
+
+		case EDMValueType::VT_Texture:
+			if (UDMMaterialValueTexture* TextureValue = Cast<UDMMaterialValueTexture>(InValue))
+			{
+				SetInputDefault_Texture(InFunctionInput, TextureValue);
+			}
+			break;
+
+		default:
+			// Not possible
+			break;
+	}
+}
+
+void UDMMaterialFunctionFunctionLibrary::SetInputDefault_Texture(UMaterialExpressionFunctionInput* InFunctionInput, UDMMaterialValueTexture* InTextureValue)
+{
+	if (!InFunctionInput || !InTextureValue)
+	{
+		return;
+	}
+
+	UMaterialExpressionTextureObject* TextureObject = Cast<UMaterialExpressionTextureObject>(InFunctionInput->Preview.Expression);
+
+	if (!TextureObject || !TextureObject->Texture)
+	{
+		return;
+	}
+
+	InTextureValue->SetDefaultValue(TextureObject->Texture);
+	InTextureValue->ApplyDefaultValue();
+}
+
+#undef LOCTEXT_NAMESPACE
diff --git a/Engine/Plugins/Experimental/DynamicMaterial/Source/DynamicMaterialEditor/Public/Utils/DMMaterialFunctionFunctionLibrary.h b/Engine/Plugins/Experimental/DynamicMaterial/Source/DynamicMaterialEditor/Public/Utils/DMMaterialFunctionFunctionLibrary.h
new file mode 100644
--- /dev/null
+++ b/Engine/Plugins/Experimental/DynamicMaterial/Source/DynamicMaterialEditor/Public/Utils/DMMaterialFunctionFunctionLibrary.h
@@ -0,0 +1,36 @@
+// Copyright Epic Games, Inc. All Rights Reserved.
+
+#pragma once
+
+#include "Kismet/BlueprintFunctionLibrary.h"
+
+#include "Templates/SharedPointerFwd.h"
+
+#include "DMMaterialFunctionFunctionLibrary.generated.h"
+
+class IPropertyHandle;
+class UDMMaterialValue;
+class UDMMaterialValueTexture;
+class UMaterialExpressionFunctionInput;
+enum class EDMValueType : uint8;
+struct FFunctionExpressionInput;
+
+/**
+ * Material Function Function Library
+ */
+UCLASS()
+class UDMMaterialFunctionFunctionLibrary : public UBlueprintFunctionLibrary
+{
+	GENERATED_BODY()
+
+public:
+	DYNAMICMATERIALEDITOR_API static void ApplyMetaData(const FFunctionExpressionInput& InFunctionInput, 
+		const TSharedRef<IPropertyHandle>& InPropertyHandle);
+
+	DYNAMICMATERIALEDITOR_API static EDMValueType GetInputValueType(UMaterialExpressionFunctionInput* InFunctionInput);
+
+	DYNAMICMATERIALEDITOR_API static void SetInputDefault(UMaterialExpressionFunctionInput* InFunctionInput, UDMMaterialValue* InValue);
+
+private:
+	static void SetInputDefault_Texture(UMaterialExpressionFunctionInput* InFunctionInput, UDMMaterialValueTexture* InTextureValue);
+};
