Material Designer
- Fixing duplicate localization key.

diff --git a/Engine/Plugins/Experimental/DynamicMaterial/Source/DynamicMaterialEditor/Private/UI/Widgets/Editor/SDMMaterialSlotEditor.cpp b/Engine/Plugins/Experimental/DynamicMaterial/Source/DynamicMaterialEditor/Private/UI/Widgets/Editor/SDMMaterialSlotEditor.cpp
--- a/Engine/Plugins/Experimental/DynamicMaterial/Source/DynamicMaterialEditor/Private/UI/Widgets/Editor/SDMMaterialSlotEditor.cpp
+++ b/Engine/Plugins/Experimental/DynamicMaterial/Source/DynamicMaterialEditor/Private/UI/Widgets/Editor/SDMMaterialSlotEditor.cpp
@@ -810,7 +810,7 @@ TSharedRef<SWidget> SDMMaterialSlotEditor::CreateSlot_LayerSettings()
 			SNew(SButton)
 			.ContentPadding(4.0f)
 			.ButtonStyle(FDynamicMaterialEditorStyle::Get(), "HoverHintOnly.Bordered.Dark")
-			.ToolTipText(LOCTEXT("DuplicateLayer", "Duplicate Selected Layer"))
+			.ToolTipText(LOCTEXT("DuplicateSelectedLayer", "Duplicate Selected Layer"))
 			.IsEnabled(this, &SDMMaterialSlotEditor::GetLayerRowsButtonsCanDuplicate)
 			.OnClicked(this, &SDMMaterialSlotEditor::OnLayerRowButtonsDuplicateClicked)
 			[
