Material Designer
- Small ensure that things have their parameter name checked after duplication.

diff --git a/Engine/Plugins/Experimental/DynamicMaterial/Source/DynamicMaterial/Private/Components/DMMaterialValue.cpp b/Engine/Plugins/Experimental/DynamicMaterial/Source/DynamicMaterial/Private/Components/DMMaterialValue.cpp
--- a/Engine/Plugins/Experimental/DynamicMaterial/Source/DynamicMaterial/Private/Components/DMMaterialValue.cpp
+++ b/Engine/Plugins/Experimental/DynamicMaterial/Source/DynamicMaterial/Private/Components/DMMaterialValue.cpp
@@ -293,6 +293,7 @@ void UDMMaterialValue::PostEditorDuplicate(UDynamicMaterialModel* InMaterialMode
 	if (GetOuter() == InMaterialModel)
 	{
 		Super::PostEditorDuplicate(InMaterialModel, InParent);
+		UpdateCachedParameterName();
 		return;
 	}
 
diff --git a/Engine/Plugins/Experimental/DynamicMaterial/Source/DynamicMaterialEditor/Private/Components/DMMaterialSlot.cpp b/Engine/Plugins/Experimental/DynamicMaterial/Source/DynamicMaterialEditor/Private/Components/DMMaterialSlot.cpp
--- a/Engine/Plugins/Experimental/DynamicMaterial/Source/DynamicMaterialEditor/Private/Components/DMMaterialSlot.cpp
+++ b/Engine/Plugins/Experimental/DynamicMaterial/Source/DynamicMaterialEditor/Private/Components/DMMaterialSlot.cpp
@@ -594,6 +594,15 @@ bool UDMMaterialSlot::PasteLayer(UDMMaterialLayerObject* InLayer)
 	{
 		const FDMUpdateGuard Guard;
 		InLayer->SetMaterialProperty(NewPropertyType);
+
+		UDynamicMaterialModel* MaterialModel = nullptr;
+
+		if (UDynamicMaterialModelEditorOnlyData* EditorOnlyData = GetMaterialModelEditorOnlyData())
+		{
+			MaterialModel = EditorOnlyData->GetMaterialModel();
+		}
+
+		InLayer->PostEditorDuplicate(MaterialModel, this);
 	}
 
 	LayerObjects.Add(InLayer);
