Material Designer
- Properly responds to whether properties are enabled are not.
- Clears the node connected to material properties when rebuilding materials again (!)

diff --git a/Engine/Plugins/Experimental/DynamicMaterial/Source/DynamicMaterialEditor/Private/Model/DynamicMaterialModelEditorOnlyData.cpp b/Engine/Plugins/Experimental/DynamicMaterial/Source/DynamicMaterialEditor/Private/Model/DynamicMaterialModelEditorOnlyData.cpp
--- a/Engine/Plugins/Experimental/DynamicMaterial/Source/DynamicMaterialEditor/Private/Model/DynamicMaterialModelEditorOnlyData.cpp
+++ b/Engine/Plugins/Experimental/DynamicMaterial/Source/DynamicMaterialEditor/Private/Model/DynamicMaterialModelEditorOnlyData.cpp
@@ -495,9 +495,14 @@ void UDynamicMaterialModelEditorOnlyData::BuildMaterial(bool bInDirtyAssets)
 	ForEachMaterialPropertyType(
 		[this, &BuildState](EDMMaterialPropertyType InType)
 		{
+			if (FExpressionInput* Input = BuildState->GetMaterialProperty(InType))
+			{
+				Input->Expression = nullptr;
+			}
+
 			if (UDMMaterialProperty* Property = GetMaterialProperty(InType))
 			{
-				if (Property->IsMaterialPin())
+				if (Property->IsEnabled() && Property->IsMaterialPin())
 				{
 					Property->GenerateExpressions(BuildState);
 
@@ -520,16 +525,18 @@ void UDynamicMaterialModelEditorOnlyData::BuildMaterial(bool bInDirtyAssets)
 		 */
 		EDMMaterialPropertyType GenerateOpacityInput = EDMMaterialPropertyType::None;
 
+		// Masked materials use the mask property.
 		if (BlendMode == EBlendMode::BLEND_Masked)
 		{
-			if (GetSlotForMaterialProperty(EDMMaterialPropertyType::OpacityMask) == nullptr)
+			if (GetSlotForEnabledMaterialProperty(EDMMaterialPropertyType::OpacityMask) == nullptr)
 			{
 				GenerateOpacityInput = EDMMaterialPropertyType::OpacityMask;
 			}
 		}
+		// Any other translucent material will use the opacity property.
 		else if (BlendMode != EBlendMode::BLEND_Opaque)
 		{
-			if (GetSlotForMaterialProperty(EDMMaterialPropertyType::Opacity) == nullptr)
+			if (GetSlotForEnabledMaterialProperty(EDMMaterialPropertyType::Opacity) == nullptr)
 			{
 				GenerateOpacityInput = EDMMaterialPropertyType::Opacity;
 			}
@@ -540,12 +547,12 @@ void UDynamicMaterialModelEditorOnlyData::BuildMaterial(bool bInDirtyAssets)
 			UDMMaterialSlot* OpacitySlot = nullptr;
 			EDMMaterialPropertyType OpacityProperty = EDMMaterialPropertyType::None;
 
-			if (UDMMaterialSlot* BaseColorSlot = GetSlotForMaterialProperty(EDMMaterialPropertyType::BaseColor))
+			if (UDMMaterialSlot* BaseColorSlot = GetSlotForEnabledMaterialProperty(EDMMaterialPropertyType::BaseColor))
 			{
 				OpacitySlot = BaseColorSlot;
 				OpacityProperty = EDMMaterialPropertyType::BaseColor;
 			}
-			else if (UDMMaterialSlot* EmissiveSlot = GetSlotForMaterialProperty(EDMMaterialPropertyType::EmissiveColor))
+			else if (UDMMaterialSlot* EmissiveSlot = GetSlotForEnabledMaterialProperty(EDMMaterialPropertyType::EmissiveColor))
 			{
 				OpacitySlot = EmissiveSlot;
 				OpacityProperty = EDMMaterialPropertyType::EmissiveColor;
@@ -593,7 +600,7 @@ void UDynamicMaterialModelEditorOnlyData::BuildMaterial(bool bInDirtyAssets)
 			OpacityProperty = GetMaterialProperty(EDMMaterialPropertyType::Opacity);
 		}
 
-		if (OpacityProperty != nullptr)
+		if (OpacityProperty != nullptr && OpacityProperty->IsEnabled())
 		{
 			OpacityProperty->AddAlphaMultiplier(BuildState);
 		}
@@ -607,7 +614,7 @@ void UDynamicMaterialModelEditorOnlyData::BuildMaterial(bool bInDirtyAssets)
 		{
 			if (UDMMaterialProperty* Property = GetMaterialProperty(InType))
 			{
-				if (Property->IsMaterialPin())
+				if (Property->IsMaterialPin() && Property->IsEnabled())
 				{
 					Property->AddOutputProcessor(BuildState);
 				}
@@ -1253,6 +1260,18 @@ UDMMaterialSlot* UDynamicMaterialModelEditorOnlyData::GetSlotForMaterialProperty
 	return nullptr;
 }
 
+UDMMaterialSlot* UDynamicMaterialModelEditorOnlyData::GetSlotForEnabledMaterialProperty(EDMMaterialPropertyType InType) const
+{
+	UDMMaterialProperty* Property = GetMaterialProperty(InType);
+
+	if (!Property || !Property->IsEnabled())
+	{
+		return nullptr;
+	}
+
+	return GetSlotForMaterialProperty(InType);
+}
+
 UDMMaterialSlot* UDynamicMaterialModelEditorOnlyData::AddSlot()
 {
 	UDMMaterialSlot* NewSlot = nullptr;
diff --git a/Engine/Plugins/Experimental/DynamicMaterial/Source/DynamicMaterialEditor/Public/Model/DynamicMaterialModelEditorOnlyData.h b/Engine/Plugins/Experimental/DynamicMaterial/Source/DynamicMaterialEditor/Public/Model/DynamicMaterialModelEditorOnlyData.h
--- a/Engine/Plugins/Experimental/DynamicMaterial/Source/DynamicMaterialEditor/Public/Model/DynamicMaterialModelEditorOnlyData.h
+++ b/Engine/Plugins/Experimental/DynamicMaterial/Source/DynamicMaterialEditor/Public/Model/DynamicMaterialModelEditorOnlyData.h
@@ -176,6 +176,10 @@ public:
 	UFUNCTION(BlueprintPure, Category = "Material Designer")
 	DYNAMICMATERIALEDITOR_API UDMMaterialSlot* GetSlotForMaterialProperty(EDMMaterialPropertyType InType) const;
 
+	/** Same as the above method, but will only return the slot if the material property is enabled. */
+	UFUNCTION(BlueprintPure, Category = "Material Designer")
+	DYNAMICMATERIALEDITOR_API UDMMaterialSlot* GetSlotForEnabledMaterialProperty(EDMMaterialPropertyType InType) const;
+
 	/** Adds the next available slot. Highly recommended to use AddSlotForMaterialProperty(PropertyType). */
 	UFUNCTION(BlueprintCallable, Category = "Material Designer")
 	DYNAMICMATERIALEDITOR_API UDMMaterialSlot* AddSlot();
