Material Designer
- Added missing assets and fixed path.

#virtualized

diff --git a/Engine/Plugins/Experimental/DynamicMaterial/Source/DynamicMaterialEditor/Private/DynamicMaterialEditorSettings.cpp b/Engine/Plugins/Experimental/DynamicMaterial/Source/DynamicMaterialEditor/Private/DynamicMaterialEditorSettings.cpp
--- a/Engine/Plugins/Experimental/DynamicMaterial/Source/DynamicMaterialEditor/Private/DynamicMaterialEditorSettings.cpp
+++ b/Engine/Plugins/Experimental/DynamicMaterial/Source/DynamicMaterialEditor/Private/DynamicMaterialEditorSettings.cpp
@@ -34,7 +34,7 @@ namespace UE::DynamicMaterialEditor::Private
 			}},
 		{EDMMaterialPropertyType::AmbientOcclusion,    TSoftObjectPtr<UTexture>(FSoftObjectPath(TEXT("/Script/Engine.Texture2D'/DynamicMaterial/Textures/SlotDefaults/T_MD_AmbientOcclusion.T_MD_AmbientOcclusion'")))},
 		{EDMMaterialPropertyType::Displacement,        TSoftObjectPtr<UTexture>(FSoftObjectPath(TEXT("/Script/Engine.Texture2D'/DynamicMaterial/Textures/SlotDefaults/T_MD_Displacement.T_MD_Displacement'")))},
-		{EDMMaterialPropertyType::SubsurfaceColor,     TSoftObjectPtr<UTexture>(FSoftObjectPath(TEXT("/Script/Engine.Texture2D'/Engine/EngineResources/T_MD_SubsurfaceColor.T_MD_SubsurfaceColor'")))},
+		{EDMMaterialPropertyType::SubsurfaceColor,     TSoftObjectPtr<UTexture>(FSoftObjectPath(TEXT("/Script/Engine.Texture2D'/DynamicMaterial/Textures/SlotDefaults/T_MD_SubsurfaceColor.T_MD_SubsurfaceColor'")))},
 		{EDMMaterialPropertyType::SurfaceThickness,    FLinearColor(0.0f, 0.0f, 0.0f, 1.f)},
 	};
 }
