Material Designer
- Change global parameter properties back to single uproperties, instead of the malfunctioning array
- Added more global parameters.
- Fixed properties generating code when they had no associated slot.
- Fixed properties (not) being added to the global values list in general settings inconsistently.

diff --git a/Engine/Plugins/Experimental/DynamicMaterial/Source/DynamicMaterial/Private/DMDefs.cpp b/Engine/Plugins/Experimental/DynamicMaterial/Source/DynamicMaterial/Private/DMDefs.cpp
--- a/Engine/Plugins/Experimental/DynamicMaterial/Source/DynamicMaterial/Private/DMDefs.cpp
+++ b/Engine/Plugins/Experimental/DynamicMaterial/Source/DynamicMaterial/Private/DMDefs.cpp
@@ -4,3 +4,17 @@
 
 int32 FDMUpdateGuard::GuardCount = 0;
 uint32 FDMInitializationGuard::GuardCount = 0;
+
+void UE::DynamicMaterial::ForEachMaterialPropertyType(TFunctionRef<EDMIterationResult(EDMMaterialPropertyType InType)> InCallable,
+	EDMMaterialPropertyType InStart, EDMMaterialPropertyType InEnd)
+{
+	for (uint8 PropertyIndex = static_cast<uint8>(InStart); PropertyIndex <= static_cast<uint8>(InEnd); ++PropertyIndex)
+	{
+		const EDMMaterialPropertyType Property = static_cast<EDMMaterialPropertyType>(PropertyIndex);
+
+		if (InCallable(Property) == EDMIterationResult::Break)
+		{
+			break;
+		}
+	}
+}
diff --git a/Engine/Plugins/Experimental/DynamicMaterial/Source/DynamicMaterial/Private/Model/DynamicMaterialModel.cpp b/Engine/Plugins/Experimental/DynamicMaterial/Source/DynamicMaterial/Private/Model/DynamicMaterialModel.cpp
--- a/Engine/Plugins/Experimental/DynamicMaterial/Source/DynamicMaterial/Private/Model/DynamicMaterialModel.cpp
+++ b/Engine/Plugins/Experimental/DynamicMaterial/Source/DynamicMaterial/Private/Model/DynamicMaterialModel.cpp
@@ -19,6 +19,10 @@
 
 const FString UDynamicMaterialModel::ValuesPathToken = FString(TEXT("Values"));
 const FString UDynamicMaterialModel::ParametersPathToken = FString(TEXT("Parameters"));
+const FLazyName UDynamicMaterialModel::GlobalBaseColorValueName = FLazyName(TEXT("GlobalBaseColorValue"));
+const FLazyName UDynamicMaterialModel::GlobalBaseColorParameterName = FLazyName(TEXT("GlobalBaseColor"));
+const FLazyName UDynamicMaterialModel::GlobalEmissiveColorValueName = FLazyName(TEXT("GlobalEmissiveColorValue"));
+const FLazyName UDynamicMaterialModel::GlobalEmissiveColorParameterName = FLazyName(TEXT("GlobalEmissiveColor"));
 const FLazyName UDynamicMaterialModel::GlobalOpacityValueName = FLazyName(TEXT("GlobalOpacityValue"));
 const FLazyName UDynamicMaterialModel::GlobalOpacityParameterName = FLazyName(TEXT("GlobalOpacity"));
 const FLazyName UDynamicMaterialModel::GlobalMetallicValueName = FLazyName(TEXT("GlobalMetallicValue"));
@@ -37,6 +41,8 @@ const FLazyName UDynamicMaterialModel::GlobalAmbientOcclusionValueName = FLazyNa
 const FLazyName UDynamicMaterialModel::GlobalAmbientOcclusionParameterName = FLazyName(TEXT("GlobalAmbientOcclusion"));
 const FLazyName UDynamicMaterialModel::GlobalRefractionValueName = FLazyName(TEXT("GlobalRefractionValue"));
 const FLazyName UDynamicMaterialModel::GlobalRefractionParameterName = FLazyName(TEXT("GlobalRefraction"));
+const FLazyName UDynamicMaterialModel::GlobalTangentValueName = FLazyName(TEXT("GlobalTangentValue"));
+const FLazyName UDynamicMaterialModel::GlobalTangentParameterName = FLazyName(TEXT("GlobalTangent"));
 const FLazyName UDynamicMaterialModel::GlobalPixelDepthOffsetValueName = FLazyName(TEXT("GlobalPixelDepthOffsetValue"));
 const FLazyName UDynamicMaterialModel::GlobalPixelDepthOffsetParameterName = FLazyName(TEXT("GlobalPixelDepthOffset"));
 const FLazyName UDynamicMaterialModel::GlobalDisplacementValueName = FLazyName(TEXT("GlobalDisplacementValue"));
@@ -45,6 +51,7 @@ const FLazyName UDynamicMaterialModel::GlobalSubsurfaceColorValueName = FLazyNam
 const FLazyName UDynamicMaterialModel::GlobalSubsurfaceColorParameterName = FLazyName(TEXT("GlobalSubsurfaceColor"));
 const FLazyName UDynamicMaterialModel::GlobalSurfaceThicknessValueName = FLazyName(TEXT("GlobalSurfaceThicknessValue"));
 const FLazyName UDynamicMaterialModel::GlobalSurfaceThicknessParameterName = FLazyName(TEXT("GlobalSurfaceThickness"));
+
 const FLazyName UDynamicMaterialModel::GlobalOffsetValueName = FLazyName(TEXT("GlobalOffsetValue"));
 const FLazyName UDynamicMaterialModel::GlobalOffsetParameterName = FLazyName(TEXT("GlobalOffset"));
 const FLazyName UDynamicMaterialModel::GlobalTilingValueName = FLazyName(TEXT("GlobalTilingValue"));
@@ -52,6 +59,28 @@ const FLazyName UDynamicMaterialModel::GlobalTilingParameterName = FLazyName(TEX
 const FLazyName UDynamicMaterialModel::GlobalRotationValueName = FLazyName(TEXT("GlobalRotationValue"));
 const FLazyName UDynamicMaterialModel::GlobalRotationParameterName = FLazyName(TEXT("GlobalRotation"));
 
+namespace UE::DynamicMaterial::Private
+{
+	const TMap<FName, EDMMaterialPropertyType> ValueNameToMaterialProperty = {
+		{UDynamicMaterialModel::GlobalBaseColorValueName,           EDMMaterialPropertyType::BaseColor},
+		{UDynamicMaterialModel::GlobalEmissiveColorValueName,       EDMMaterialPropertyType::EmissiveColor},
+		{UDynamicMaterialModel::GlobalOpacityValueName,             EDMMaterialPropertyType::Opacity},
+		{UDynamicMaterialModel::GlobalMetallicValueName,            EDMMaterialPropertyType::Metallic},
+		{UDynamicMaterialModel::GlobalSpecularValueName,            EDMMaterialPropertyType::Specular},
+		{UDynamicMaterialModel::GlobalRoughnessValueName,           EDMMaterialPropertyType::Roughness},
+		{UDynamicMaterialModel::GlobalAnisotropyValueName,          EDMMaterialPropertyType::Anisotropy},
+		{UDynamicMaterialModel::GlobalNormalValueName,              EDMMaterialPropertyType::Normal},
+		{UDynamicMaterialModel::GlobalWorldPositionOffsetValueName, EDMMaterialPropertyType::WorldPositionOffset},
+		{UDynamicMaterialModel::GlobalAmbientOcclusionValueName,    EDMMaterialPropertyType::AmbientOcclusion},
+		{UDynamicMaterialModel::GlobalRefractionValueName,          EDMMaterialPropertyType::Refraction},
+		{UDynamicMaterialModel::GlobalTangentValueName,             EDMMaterialPropertyType::Tangent},
+		{UDynamicMaterialModel::GlobalPixelDepthOffsetValueName,    EDMMaterialPropertyType::PixelDepthOffset},
+		{UDynamicMaterialModel::GlobalDisplacementValueName,        EDMMaterialPropertyType::Displacement},
+		{UDynamicMaterialModel::GlobalSubsurfaceColorValueName,     EDMMaterialPropertyType::SubsurfaceColor},
+		{UDynamicMaterialModel::GlobalSurfaceThicknessValueName,    EDMMaterialPropertyType::SurfaceThickness}
+	};
+}
+
 UDynamicMaterialModel::UDynamicMaterialModel()
 {
 	DynamicMaterialInstance = nullptr;
@@ -60,15 +89,20 @@ UDynamicMaterialModel::UDynamicMaterialModel()
 	EditorOnlyDataSI.SetObject(nullptr);
 #endif
 
-	auto AddFloatParameter = [this](FName InPropertyName, FName InParameterName, float InDefaultValue = 1.f)
+	auto AddFloatParameter = [this](TObjectPtr<UDMMaterialValue>& InValuePtr, FName InPropertyName, FName InParameterName, float InDefaultValue = 1.f, bool bInSetRange = false)
 		{
 			UDMMaterialValueFloat1* Property = CreateDefaultSubobject<UDMMaterialValueFloat1>(InPropertyName);
 			Property->CachedParameterName = InParameterName;
-			GlobalParameterValues.Add(InPropertyName, Property);
+			InValuePtr = Property;
 
 #if WITH_EDITOR
 			Property->SetDefaultValue(InDefaultValue);
 			Property->ApplyDefaultValue();
+
+			if (bInSetRange)
+			{
+				Property->SetValueRange({0.f, 1.f});
+			}
 #endif
 
 			UDMMaterialParameter* Parameter = CreateDefaultSubobject<UDMMaterialParameter>(FName(*(InParameterName.GetPlainNameString() + TEXT("Parameter"))));
@@ -78,11 +112,11 @@ UDynamicMaterialModel::UDynamicMaterialModel()
 			ParameterMap.Add(InParameterName, Parameter);
 		};
 
-	auto AddVector2Parameter = [this](FName InPropertyName, FName InParameterName, const FVector2D& InDefaultValue)
+	auto AddVector2Parameter = [this](TObjectPtr<UDMMaterialValue>& InValuePtr, FName InPropertyName, FName InParameterName, const FVector2D& InDefaultValue)
 		{
 			UDMMaterialValueFloat2* Property = CreateDefaultSubobject<UDMMaterialValueFloat2>(InPropertyName);
 			Property->CachedParameterName = InParameterName;
-			GlobalParameterValues.Add(InPropertyName, Property);
+			InValuePtr = Property;
 
 #if WITH_EDITOR
 			Property->SetDefaultValue(InDefaultValue);
@@ -97,21 +131,27 @@ UDynamicMaterialModel::UDynamicMaterialModel()
 		};
 
 	FDMUpdateGuard Guard;
-	AddFloatParameter(GlobalOpacityValueName, GlobalOpacityParameterName);
-	AddFloatParameter(GlobalMetallicValueName, GlobalMetallicParameterName);
-	AddFloatParameter(GlobalSpecularValueName, GlobalSpecularParameterName);
-	AddFloatParameter(GlobalRoughnessValueName, GlobalRoughnessParameterName);
-	AddFloatParameter(GlobalNormalValueName, GlobalNormalParameterName);
-	AddFloatParameter(GlobalAnisotropyValueName, GlobalAnisotropyParameterName);
-	AddFloatParameter(GlobalWorldPositionOffsetValueName, GlobalWorldPositionOffsetParameterName);
-	AddFloatParameter(GlobalAmbientOcclusionValueName, GlobalAmbientOcclusionParameterName);
-	AddFloatParameter(GlobalRefractionValueName, GlobalRefractionParameterName);
-	AddFloatParameter(GlobalPixelDepthOffsetValueName, GlobalPixelDepthOffsetParameterName);
-	AddFloatParameter(GlobalDisplacementValueName, GlobalDisplacementParameterName);
-
-	AddVector2Parameter(GlobalOffsetValueName, GlobalOffsetParameterName, FVector2D::ZeroVector);
-	AddVector2Parameter(GlobalTilingValueName, GlobalTilingParameterName, FVector2D::UnitVector);
-	AddFloatParameter(GlobalRotationValueName, GlobalRotationParameterName, /* Default Value */ 0.f);
+
+	AddFloatParameter(GlobalBaseColorParameterValue, GlobalBaseColorValueName, GlobalBaseColorParameterName);
+	AddFloatParameter(GlobalEmissiveColorParameterValue, GlobalEmissiveColorValueName, GlobalEmissiveColorParameterName);
+	AddFloatParameter(GlobalOpacityParameterValue, GlobalOpacityValueName, GlobalOpacityParameterName, /* Set Range */ true);
+	AddFloatParameter(GlobalMetallicParameterValue, GlobalMetallicValueName, GlobalMetallicParameterName);
+	AddFloatParameter(GlobalSpecularParameterValue, GlobalSpecularValueName, GlobalSpecularParameterName);
+	AddFloatParameter(GlobalRoughnessParameterValue, GlobalRoughnessValueName, GlobalRoughnessParameterName);
+	AddFloatParameter(GlobalNormalParameterValue, GlobalNormalValueName, GlobalNormalParameterName);
+	AddFloatParameter(GlobalAnisotropyParameterValue, GlobalAnisotropyValueName, GlobalAnisotropyParameterName);
+	AddFloatParameter(GlobalWorldPositionOffsetParameterValue, GlobalWorldPositionOffsetValueName, GlobalWorldPositionOffsetParameterName);
+	AddFloatParameter(GlobalAmbientOcclusionParameterValue, GlobalAmbientOcclusionValueName, GlobalAmbientOcclusionParameterName);
+	AddFloatParameter(GlobalTangentParameterValue, GlobalTangentValueName, GlobalTangentParameterName);
+	AddFloatParameter(GlobalRefractionParameterValue, GlobalRefractionValueName, GlobalRefractionParameterName);
+	AddFloatParameter(GlobalPixelDepthOffsetParameterValue, GlobalPixelDepthOffsetValueName, GlobalPixelDepthOffsetParameterName);
+	AddFloatParameter(GlobalDisplacementParameterValue, GlobalDisplacementValueName, GlobalDisplacementParameterName);
+	AddFloatParameter(GlobalSubsurfaceColorParameterValue, GlobalSubsurfaceColorValueName, GlobalSubsurfaceColorParameterName);
+	AddFloatParameter(GlobalSurfaceThicknessParameterValue, GlobalSurfaceThicknessValueName, GlobalSurfaceThicknessParameterName);
+
+	AddVector2Parameter(GlobalOffsetParameterValue, GlobalOffsetValueName, GlobalOffsetParameterName, FVector2D::ZeroVector);
+	AddVector2Parameter(GlobalTilingParameterValue, GlobalTilingValueName, GlobalTilingParameterName, FVector2D::UnitVector);
+	AddFloatParameter(GlobalRotationParameterValue, GlobalRotationValueName, GlobalRotationParameterName, /* Default Value */ 0.f);
 }
 
 void UDynamicMaterialModel::SetDynamicMaterialInstance(UDynamicMaterialInstance* InDynamicMaterialInstance)
@@ -141,18 +181,30 @@ UDMMaterialValueFloat1* UDynamicMaterialModel::GetGlobalOpacityValue() const
 	return GetTypedGlobalParameterValue<UDMMaterialValueFloat1>(GlobalOpacityValueName);
 }
 
-const TMap<FName, TObjectPtr<UDMMaterialValue>>& UDynamicMaterialModel::GetGlobalParameterValues() const
-{
-	return GlobalParameterValues;
-}
-
 UDMMaterialValue* UDynamicMaterialModel::GetGlobalParameterValue(FName InName) const
 {
-	if (const TObjectPtr<UDMMaterialValue>* ValuePtr = GlobalParameterValues.Find(InName))
+	using namespace UE::DynamicMaterial::Private;
+
+	if (const EDMMaterialPropertyType* MaterialPropertyPtr = ValueNameToMaterialProperty.Find(InName))
+	{
+		return GetGlobalParameterValueForMaterialProperty(*MaterialPropertyPtr);
+	}
+
+	if (InName == GlobalOffsetValueName)
+	{
+		return GlobalOffsetParameterValue;
+	}
+
+	if (InName == GlobalTilingValueName)
+	{
+		return GlobalTilingParameterValue;
+	}
+
+	if (InName == GlobalRotationValueName)
 	{
-		return Cast<UDMMaterialValue>(*ValuePtr);
+		return GlobalRotationParameterValue;
 	}
-	
+
 	return nullptr;
 }
 
@@ -161,6 +213,86 @@ bool UDynamicMaterialModel::IsModelValid() const
 	return (!HasAnyFlags(RF_BeginDestroyed | RF_FinishDestroyed) && IsValid(this));
 }
 
+UDMMaterialValue* UDynamicMaterialModel::GetGlobalParameterValueForMaterialProperty(EDMMaterialPropertyType InProperty) const
+{
+	switch (InProperty)
+	{
+		case EDMMaterialPropertyType::BaseColor:
+			return GlobalBaseColorParameterValue;
+
+		case EDMMaterialPropertyType::EmissiveColor:
+			return GlobalEmissiveColorParameterValue;
+
+		case EDMMaterialPropertyType::Opacity:
+			return GlobalOpacityParameterValue;
+
+		case EDMMaterialPropertyType::Metallic:
+			return GlobalMetallicParameterValue;
+
+		case EDMMaterialPropertyType::Specular:
+			return GlobalSpecularParameterValue;
+
+		case EDMMaterialPropertyType::Roughness:
+			return GlobalRoughnessParameterValue;
+
+		case EDMMaterialPropertyType::Anisotropy:
+			return GlobalAnisotropyParameterValue;
+
+		case EDMMaterialPropertyType::Normal:
+			return GlobalNormalParameterValue;
+
+		case EDMMaterialPropertyType::WorldPositionOffset:
+			return GlobalWorldPositionOffsetParameterValue;
+
+		case EDMMaterialPropertyType::AmbientOcclusion:
+			return GlobalAmbientOcclusionParameterValue;
+
+		case EDMMaterialPropertyType::Refraction:
+			return GlobalRefractionParameterValue;
+
+		case EDMMaterialPropertyType::Tangent:
+			return GlobalTangentParameterValue;
+
+		case EDMMaterialPropertyType::PixelDepthOffset:
+			return GlobalPixelDepthOffsetParameterValue;
+
+		case EDMMaterialPropertyType::Displacement:
+			return GlobalDisplacementParameterValue;
+
+		case EDMMaterialPropertyType::SubsurfaceColor:
+			return GlobalSubsurfaceColorParameterValue;
+
+		case EDMMaterialPropertyType::SurfaceThickness:
+			return GlobalSurfaceThicknessParameterValue;
+
+		default:
+			return nullptr;
+	}
+}
+
+void UDynamicMaterialModel::ForEachGlobalParameter(TFunctionRef<void(UDMMaterialValue* InGlobalParameterValue)> InCallable)
+{
+	InCallable(GlobalBaseColorParameterValue);
+	InCallable(GlobalEmissiveColorParameterValue);
+	InCallable(GlobalOpacityParameterValue);
+	InCallable(GlobalMetallicParameterValue);
+	InCallable(GlobalSpecularParameterValue);
+	InCallable(GlobalRoughnessParameterValue);
+	InCallable(GlobalAnisotropyParameterValue);
+	InCallable(GlobalNormalParameterValue);
+	InCallable(GlobalWorldPositionOffsetParameterValue);
+	InCallable(GlobalAmbientOcclusionParameterValue);
+	InCallable(GlobalRefractionParameterValue);
+	InCallable(GlobalTangentParameterValue);
+	InCallable(GlobalPixelDepthOffsetParameterValue);
+	InCallable(GlobalDisplacementParameterValue);
+	InCallable(GlobalSubsurfaceColorParameterValue);
+	InCallable(GlobalSurfaceThicknessParameterValue);
+	InCallable(GlobalOffsetParameterValue);
+	InCallable(GlobalTilingParameterValue);
+	InCallable(GlobalRotationParameterValue);
+}
+
 UDMMaterialComponent* UDynamicMaterialModel::GetComponentByPath(const FString& InPath) const
 {
 	FDMComponentPath Path(InPath);
@@ -455,17 +587,24 @@ void UDynamicMaterialModel::OnTextureUVUpdated(UDMTextureUV* InTextureUV)
 
 void UDynamicMaterialModel::ApplyComponents(UMaterialInstanceDynamic* InMID)
 {
-	for (const TPair<FName, TObjectPtr<UDMMaterialValue>>& GlobalParameterPair : GlobalParameterValues)
-	{
-#if WITH_EDITOR
-		if (GlobalParameterPair.Value->IsComponentCreated())
+	ForEachGlobalParameter(
+		[this, InMID](UDMMaterialValue* InValue)
 		{
-			GlobalParameterPair.Value->SetComponentState(EDMComponentLifetimeState::Added);
-		}
+			if (!InValue)
+			{
+				return;
+			}
+
+#if WITH_EDITOR
+			if (InValue->IsComponentCreated())
+			{
+				InValue->SetComponentState(EDMComponentLifetimeState::Added);
+			}
 #endif
 
-		GlobalParameterPair.Value->SetMIDParameter(InMID);
-	}
+			InValue->SetMIDParameter(InMID);
+		}
+	);
 
 	for (UDMMaterialValue* Value : Values)
 	{
@@ -505,7 +644,7 @@ void UDynamicMaterialModel::PostLoad()
 	// Open assets in the editor first to fix version upgrades!
 	if (GlobalOpacityValue)
 	{
-		GlobalParameterValues.FindOrAdd(GlobalOpacityValueName) = GlobalOpacityValue;
+		GlobalOpacityParameterValue = GlobalOpacityValue;
 
 		if (IDynamicMaterialModelEditorOnlyDataInterface* EditorOnlyDataInterface = GetEditorOnlyData())
 		{
@@ -580,25 +719,22 @@ void UDynamicMaterialModel::PostDuplicate(bool bInDuplicateForPIE)
 
 void UDynamicMaterialModel::PostEditorDuplicate()
 {
-	for (const TObjectPtr<UDMMaterialValue>& Value : Values)
-	{
-		if (GUndo)
+	ForEachGlobalParameter(
+		[this](UDMMaterialValue* InValue)
 		{
-			Value->Modify();
-		}
+			if (!InValue)
+			{
+				return;
+			}
 
-		Value->PostEditorDuplicate(this, nullptr);
-	}
+			if (GUndo)
+			{
+				InValue->Modify();
+			}
 
-	for (const TPair<FName, TObjectPtr<UDMMaterialValue>>& ValuePair : GlobalParameterValues)
-	{
-		if (GUndo)
-		{
-			ValuePair.Value->Modify();
+			InValue->PostEditorDuplicate(this, nullptr);
 		}
-
-		ValuePair.Value->PostEditorDuplicate(this, nullptr);
-	}
+	);
 
 	for (const TPair<FName, TWeakObjectPtr<UDMMaterialParameter>>& ParameterPair : ParameterMap)
 	{
@@ -622,24 +758,9 @@ void UDynamicMaterialModel::PostEditorDuplicate()
 
 void UDynamicMaterialModel::FixGlobalParameterValues()
 {
-	for (TPair<FName, TObjectPtr<UDMMaterialValue>>& ParameterPair : GlobalParameterValues)
+	if (UDMMaterialValueFloat1* FloatValue = Cast<UDMMaterialValueFloat1>(GlobalOpacityParameterValue))
 	{
-		if (!ParameterPair.Value || !ParameterPair.Value->HasAnyFlags(RF_ArchetypeObject))
-		{
-			continue;
-		}
-
-		const FName Name = ParameterPair.Value->GetFName();
-
-		UDMMaterialValue* Local = FindObjectFast<UDMMaterialValue>(this, Name);
-
-		if (!Local)
-		{
-			UE_LOG(LogDynamicMaterial, Error, TEXT("Archtype global parameter found and could not find local one. [%s]"), *Name.ToString());
-			continue;
-		}
-
-		ParameterPair.Value = Local;
+		FloatValue->SetValueRange({0.f, 1.f});
 	}
 }
 
diff --git a/Engine/Plugins/Experimental/DynamicMaterial/Source/DynamicMaterial/Private/Model/DynamicMaterialModelDynamic.cpp b/Engine/Plugins/Experimental/DynamicMaterial/Source/DynamicMaterial/Private/Model/DynamicMaterialModelDynamic.cpp
--- a/Engine/Plugins/Experimental/DynamicMaterial/Source/DynamicMaterial/Private/Model/DynamicMaterialModelDynamic.cpp
+++ b/Engine/Plugins/Experimental/DynamicMaterial/Source/DynamicMaterial/Private/Model/DynamicMaterialModelDynamic.cpp
@@ -206,16 +206,20 @@ void UDynamicMaterialModelDynamic::InitComponents()
 		return;
 	}
 
-	const TMap<FName, TObjectPtr<UDMMaterialValue>>& GlobalParameterValues = ParentModelLocal->GetGlobalParameterValues();
+	int32 GlobalParamCount = 0;
+	ParentModelLocal->ForEachGlobalParameter([&GlobalParamCount](UDMMaterialValue*) { ++GlobalParamCount; });
+
 	const TArray<UDMMaterialValue*>& ParentValues = ParentModelLocal->GetValues();
 	const TSet<TObjectPtr<UDMMaterialComponent>>& RuntimeComponents = ParentModelLocal->GetRuntimeComponents();
 
-	DynamicComponents.Empty(GlobalParameterValues.Num() + ParentValues.Num() + RuntimeComponents.Num());
+	DynamicComponents.Empty(GlobalParamCount + ParentValues.Num() + RuntimeComponents.Num()); // Rough estimate on the number of global params
 
-	for (const TPair<FName, TObjectPtr<UDMMaterialValue>>& GlobalParameterPair : GlobalParameterValues)
-	{
-		DynamicComponents.Add(GlobalParameterPair.Value->GetFName(), GlobalParameterPair.Value->ToDynamic(this));
-	}
+	ParentModelLocal->ForEachGlobalParameter(
+		[this](UDMMaterialValue* InValue)
+		{
+			DynamicComponents.Add(InValue->GetFName(), InValue->ToDynamic(this));
+		}
+	);
 
 	for (UDMMaterialValue* ParentValue : ParentValues)
 	{
@@ -244,21 +248,27 @@ void UDynamicMaterialModelDynamic::EnsureComponents()
 		return;
 	}
 
-	const TMap<FName, TObjectPtr<UDMMaterialValue>>& GlobalParameterValues = ParentModelLocal->GetGlobalParameterValues();
+	int32 GlobalParamCount = 0;
+	ParentModelLocal->ForEachGlobalParameter([&GlobalParamCount](UDMMaterialValue*) { ++GlobalParamCount; });
+
 	const TArray<UDMMaterialValue*>& ParentValues = ParentModelLocal->GetValues();
 	const TSet<TObjectPtr<UDMMaterialComponent>>& RuntimeComponents = ParentModelLocal->GetRuntimeComponents();
 
-	const int32 RequiredComponentCount = GlobalParameterValues.Num() + ParentValues.Num() + RuntimeComponents.Num();
+	const int32 RequiredComponentCount = GlobalParamCount + ParentValues.Num() + RuntimeComponents.Num();
 
 	DynamicComponents.Reserve(RequiredComponentCount);
 
-	for (const TPair<FName, TObjectPtr<UDMMaterialValue>>& GlobalParameterPair : GlobalParameterValues)
-	{
-		if (!DynamicComponents.Contains(GlobalParameterPair.Key))
+	ParentModelLocal->ForEachGlobalParameter(
+		[this](UDMMaterialValue* InValue)
 		{
-			DynamicComponents.Add(GlobalParameterPair.Value->GetFName(), GlobalParameterPair.Value->ToDynamic(this));
+			const FName ValueName = InValue->GetFName();
+
+			if (!DynamicComponents.Contains(ValueName))
+			{
+				DynamicComponents.Add(ValueName, InValue->ToDynamic(this));
+			}
 		}
-	}
+	);
 
 	for (UDMMaterialValue* ParentValue : ParentValues)
 	{
diff --git a/Engine/Plugins/Experimental/DynamicMaterial/Source/DynamicMaterial/Public/DMDefs.h b/Engine/Plugins/Experimental/DynamicMaterial/Source/DynamicMaterial/Public/DMDefs.h
--- a/Engine/Plugins/Experimental/DynamicMaterial/Source/DynamicMaterial/Public/DMDefs.h
+++ b/Engine/Plugins/Experimental/DynamicMaterial/Source/DynamicMaterial/Public/DMDefs.h
@@ -5,17 +5,12 @@
 #include "Containers/Array.h"
 #include "Internationalization/Text.h"
 #include "Engine/EngineTypes.h"
+
 #include "DMDefs.generated.h"
 
 class UDMMaterialStage;
 class UTexture;
 
-namespace UE::DynamicMaterial
-{
-	constexpr int32 RenameFlags = REN_DontCreateRedirectors | REN_DoNotDirty | REN_ForceNoResetLoaders | REN_NonTransactional;
-	constexpr EClassFlags InvalidClassFlags = CLASS_Abstract | CLASS_Hidden | CLASS_Deprecated | CLASS_NewerVersionExists;
-}
-
 /** Some short names provided so they fit nicely on buttons! */
 UENUM(BlueprintType)
 enum class EDMMaterialPropertyType : uint8
@@ -45,6 +40,22 @@ enum class EDMMaterialPropertyType : uint8
 	Any
 };
 
+enum class EDMIterationResult : uint8
+{
+	Continue,
+	Break
+};
+
+namespace UE::DynamicMaterial
+{
+	constexpr int32 RenameFlags = REN_DontCreateRedirectors | REN_DoNotDirty | REN_ForceNoResetLoaders | REN_NonTransactional;
+	constexpr EClassFlags InvalidClassFlags = CLASS_Abstract | CLASS_Hidden | CLASS_Deprecated | CLASS_NewerVersionExists;
+
+	DYNAMICMATERIAL_API void ForEachMaterialPropertyType(TFunctionRef<EDMIterationResult(EDMMaterialPropertyType InType)> InCallable,
+		EDMMaterialPropertyType InStart = static_cast<EDMMaterialPropertyType>(static_cast<uint8>(EDMMaterialPropertyType::None) + 1),
+		EDMMaterialPropertyType InEnd = static_cast<EDMMaterialPropertyType>(static_cast<uint8>(EDMMaterialPropertyType::Any) - 1));
+}
+
 UENUM(BlueprintType)
 enum class EDMValueType : uint8
 {
diff --git a/Engine/Plugins/Experimental/DynamicMaterial/Source/DynamicMaterial/Public/Model/DynamicMaterialModel.h b/Engine/Plugins/Experimental/DynamicMaterial/Source/DynamicMaterial/Public/Model/DynamicMaterialModel.h
--- a/Engine/Plugins/Experimental/DynamicMaterial/Source/DynamicMaterial/Public/Model/DynamicMaterialModel.h
+++ b/Engine/Plugins/Experimental/DynamicMaterial/Source/DynamicMaterial/Public/Model/DynamicMaterialModel.h
@@ -20,6 +20,7 @@ class UDMMaterialValueFloat1;
 class UDMTextureUV;
 class UMaterialExpression;
 class UMaterialInstanceDynamic;
+enum class EDMIterationResult : uint8;
 enum class EDMMaterialPropertyType : uint8;
 enum class EDMMaterialShadingModel : uint8;
 enum class EDMUpdateType : uint8;
@@ -45,6 +46,10 @@ public:
 	DYNAMICMATERIAL_API static const FString ParametersPathToken;
 
 	/** FNames for global parameters values. */
+	DYNAMICMATERIAL_API static const FLazyName GlobalBaseColorValueName;
+	DYNAMICMATERIAL_API static const FLazyName GlobalBaseColorParameterName;
+	DYNAMICMATERIAL_API static const FLazyName GlobalEmissiveColorValueName;
+	DYNAMICMATERIAL_API static const FLazyName GlobalEmissiveColorParameterName;
 	DYNAMICMATERIAL_API static const FLazyName GlobalOpacityValueName;
 	DYNAMICMATERIAL_API static const FLazyName GlobalOpacityParameterName;
 	DYNAMICMATERIAL_API static const FLazyName GlobalMetallicValueName;
@@ -63,6 +68,8 @@ public:
 	DYNAMICMATERIAL_API static const FLazyName GlobalAmbientOcclusionParameterName;
 	DYNAMICMATERIAL_API static const FLazyName GlobalRefractionValueName;
 	DYNAMICMATERIAL_API static const FLazyName GlobalRefractionParameterName;
+	DYNAMICMATERIAL_API static const FLazyName GlobalTangentValueName;
+	DYNAMICMATERIAL_API static const FLazyName GlobalTangentParameterName;
 	DYNAMICMATERIAL_API static const FLazyName GlobalPixelDepthOffsetValueName;
 	DYNAMICMATERIAL_API static const FLazyName GlobalPixelDepthOffsetParameterName;
 	DYNAMICMATERIAL_API static const FLazyName GlobalDisplacementValueName;
@@ -84,8 +91,16 @@ public:
 	UFUNCTION(BlueprintPure, Category = "Material Designer")
 	DYNAMICMATERIAL_API bool IsModelValid() const;
 
-	/** Returns the map of global parameter values (such as global opacity). */
-	DYNAMICMATERIAL_API const TMap<FName, TObjectPtr<UDMMaterialValue>>& GetGlobalParameterValues() const;
+	/** Returns a specific global parameter value (such as global opacity) for the given material property or nullptr. */
+	UFUNCTION(BlueprintPure, Category = "Material Designer")
+	DYNAMICMATERIAL_API UDMMaterialValue* GetGlobalParameterValueForMaterialProperty(EDMMaterialPropertyType InProperty) const;
+
+	/** Returns a Cast version of a specific global value (such as global opacity) or nullptr. */
+	template<typename InValueClass>
+	InValueClass* GetGlobalParameterValueForMaterialProperty(EDMMaterialPropertyType InProperty) const
+	{
+		return Cast<InValueClass>(GetGlobalParameterValueForMaterialProperty(InProperty));
+	}
 
 	/** Returns a specific global parameter value (such as global opacity) of the given object name (see global parameter FNames) or nullptr. */
 	UFUNCTION(BlueprintPure, Category = "Material Designer")
@@ -98,6 +113,8 @@ public:
 		return Cast<InValueClass>(GetGlobalParameterValue(InName));
 	}
 
+	void ForEachGlobalParameter(TFunctionRef<void(UDMMaterialValue* InGlobalParameterValue)> InCallable);
+
 	/** Searches the model for a specific component based on a path. */
 	UFUNCTION(BlueprintCallable, Category = "Material Designer")
 	DYNAMICMATERIAL_API UDMMaterialComponent* GetComponentByPath(const FString& InPath) const;
@@ -215,10 +232,6 @@ protected:
 	UPROPERTY(VisibleInstanceOnly, BlueprintReadOnly, Instanced, Category = "Material Designer")
 	TSet<TObjectPtr<UDMMaterialComponent>> RuntimeComponents;
 
-	/** Map of the global parameter values, such as global opacity. */
-	UPROPERTY(VisibleInstanceOnly, BlueprintReadOnly, Category = "Material Designer")
-	TMap<FName, TObjectPtr<UDMMaterialValue>> GlobalParameterValues;
-
 	/** Map of parameter names to the objects representing that parameter. */
 	UPROPERTY(VisibleInstanceOnly, TextExportTransient, Category = "Material Designer")
 	TMap<FName, TWeakObjectPtr<UDMMaterialParameter>> ParameterMap;
@@ -237,6 +250,63 @@ protected:
 	TScriptInterface<IDynamicMaterialModelEditorOnlyDataInterface> EditorOnlyDataSI;
 #endif
 
+	UPROPERTY(BlueprintReadOnly, Instanced, Category = "Material Designer")
+	TObjectPtr<UDMMaterialValue> GlobalBaseColorParameterValue;
+
+	UPROPERTY(BlueprintReadOnly, Instanced, Category = "Material Designer")
+	TObjectPtr<UDMMaterialValue> GlobalEmissiveColorParameterValue;
+
+	UPROPERTY(BlueprintReadOnly, Instanced, Category = "Material Designer")
+	TObjectPtr<UDMMaterialValue> GlobalOpacityParameterValue;
+
+	UPROPERTY(BlueprintReadOnly, Instanced, Category = "Material Designer")
+	TObjectPtr<UDMMaterialValue> GlobalRoughnessParameterValue;
+
+	UPROPERTY(BlueprintReadOnly, Instanced, Category = "Material Designer")
+	TObjectPtr<UDMMaterialValue> GlobalSpecularParameterValue;
+
+	UPROPERTY(BlueprintReadOnly, Instanced, Category = "Material Designer")
+	TObjectPtr<UDMMaterialValue> GlobalMetallicParameterValue;
+
+	UPROPERTY(BlueprintReadOnly, Instanced, Category = "Material Designer")
+	TObjectPtr<UDMMaterialValue> GlobalNormalParameterValue;
+
+	UPROPERTY(BlueprintReadOnly, Instanced, Category = "Material Designer")
+	TObjectPtr<UDMMaterialValue> GlobalPixelDepthOffsetParameterValue;
+
+	UPROPERTY(BlueprintReadOnly, Instanced, Category = "Material Designer")
+	TObjectPtr<UDMMaterialValue> GlobalWorldPositionOffsetParameterValue;
+
+	UPROPERTY(BlueprintReadOnly, Instanced, Category = "Material Designer")
+	TObjectPtr<UDMMaterialValue> GlobalAmbientOcclusionParameterValue;
+
+	UPROPERTY(BlueprintReadOnly, Instanced, Category = "Material Designer")
+	TObjectPtr<UDMMaterialValue> GlobalAnisotropyParameterValue;
+
+	UPROPERTY(BlueprintReadOnly, Instanced, Category = "Material Designer")
+	TObjectPtr<UDMMaterialValue> GlobalRefractionParameterValue;
+
+	UPROPERTY(BlueprintReadOnly, Instanced, Category = "Material Designer")
+	TObjectPtr<UDMMaterialValue> GlobalTangentParameterValue;
+
+	UPROPERTY(BlueprintReadOnly, Instanced, Category = "Material Designer")
+	TObjectPtr<UDMMaterialValue> GlobalDisplacementParameterValue;
+
+	UPROPERTY(BlueprintReadOnly, Instanced, Category = "Material Designer")
+	TObjectPtr<UDMMaterialValue> GlobalSubsurfaceColorParameterValue;
+
+	UPROPERTY(BlueprintReadOnly, Instanced, Category = "Material Designer")
+	TObjectPtr<UDMMaterialValue> GlobalSurfaceThicknessParameterValue;
+
+	UPROPERTY(BlueprintReadOnly, Instanced, Category = "Material Designer")
+	TObjectPtr<UDMMaterialValue> GlobalOffsetParameterValue;
+
+	UPROPERTY(BlueprintReadOnly, Instanced, Category = "Material Designer")
+	TObjectPtr<UDMMaterialValue> GlobalTilingParameterValue;
+
+	UPROPERTY(BlueprintReadOnly, Instanced, Category = "Material Designer")
+	TObjectPtr<UDMMaterialValue> GlobalRotationParameterValue;
+
 	/** Delegate called when a value is updated. @See GetOnValueUpdateDelegate. */
 	FDMOnValueUpdated OnValueUpdateDelegate;
 
@@ -244,8 +314,7 @@ protected:
 	FDMOnTextureUVUpdated OnTextureUVUpdateDelegate;
 
 	/**
-	 * Because they are not assigned to individual properties, the global parameter values end up being their archetype versions
-	 * for versions of the object from before the change. This needs to be fixed on post load.
+	 * Attempts to fix global opacity to give it a range of 0-1.
 	 */
 	void FixGlobalParameterValues();
 
diff --git a/Engine/Plugins/Experimental/DynamicMaterial/Source/DynamicMaterialEditor/Private/Model/DynamicMaterialModelEditorOnlyData.cpp b/Engine/Plugins/Experimental/DynamicMaterial/Source/DynamicMaterialEditor/Private/Model/DynamicMaterialModelEditorOnlyData.cpp
--- a/Engine/Plugins/Experimental/DynamicMaterial/Source/DynamicMaterialEditor/Private/Model/DynamicMaterialModelEditorOnlyData.cpp
+++ b/Engine/Plugins/Experimental/DynamicMaterial/Source/DynamicMaterialEditor/Private/Model/DynamicMaterialModelEditorOnlyData.cpp
@@ -207,7 +207,7 @@ UDynamicMaterialModelEditorOnlyData::UDynamicMaterialModelEditorOnlyData()
 	Custom3 = UDMMaterialProperty::CreateCustomMaterialPropertyDefaultSubobject(this, EDMMaterialPropertyType::Custom3, "MaterialProperty_Custom3");
 	Custom4 = UDMMaterialProperty::CreateCustomMaterialPropertyDefaultSubobject(this, EDMMaterialPropertyType::Custom4, "MaterialProperty_Custom4");
 
-	ForEachMaterialPropertyType(
+	UE::DynamicMaterial::ForEachMaterialPropertyType(
 		[this](EDMMaterialPropertyType InType)
 		{
 			if (UDMMaterialProperty* Property = GetMaterialProperty(InType))
@@ -222,6 +222,8 @@ UDynamicMaterialModelEditorOnlyData::UDynamicMaterialModelEditorOnlyData()
 
 void UDynamicMaterialModelEditorOnlyData::AssignPropertyAlphaValues()
 {
+	BaseColor          ->AddComponent(AlphaValueName, MaterialModel->GetGlobalParameterValue(UDynamicMaterialModel::GlobalBaseColorValueName));
+	EmissiveColor      ->AddComponent(AlphaValueName, MaterialModel->GetGlobalParameterValue(UDynamicMaterialModel::GlobalEmissiveColorValueName));
 	Opacity            ->AddComponent(AlphaValueName, MaterialModel->GetGlobalParameterValue(UDynamicMaterialModel::GlobalOpacityValueName));
 	OpacityMask        ->AddComponent(AlphaValueName, MaterialModel->GetGlobalParameterValue(UDynamicMaterialModel::GlobalOpacityValueName));
 	Metallic           ->AddComponent(AlphaValueName, MaterialModel->GetGlobalParameterValue(UDynamicMaterialModel::GlobalMetallicValueName));
@@ -232,8 +234,11 @@ void UDynamicMaterialModelEditorOnlyData::AssignPropertyAlphaValues()
 	WorldPositionOffset->AddComponent(AlphaValueName, MaterialModel->GetGlobalParameterValue(UDynamicMaterialModel::GlobalWorldPositionOffsetValueName));
 	AmbientOcclusion   ->AddComponent(AlphaValueName, MaterialModel->GetGlobalParameterValue(UDynamicMaterialModel::GlobalAmbientOcclusionValueName));
 	Refraction         ->AddComponent(AlphaValueName, MaterialModel->GetGlobalParameterValue(UDynamicMaterialModel::GlobalRefractionValueName));
+	Tangent            ->AddComponent(AlphaValueName, MaterialModel->GetGlobalParameterValue(UDynamicMaterialModel::GlobalTangentValueName));
 	PixelDepthOffset   ->AddComponent(AlphaValueName, MaterialModel->GetGlobalParameterValue(UDynamicMaterialModel::GlobalPixelDepthOffsetValueName));
 	Displacement       ->AddComponent(AlphaValueName, MaterialModel->GetGlobalParameterValue(UDynamicMaterialModel::GlobalDisplacementValueName));
+	SubsurfaceColor    ->AddComponent(AlphaValueName, MaterialModel->GetGlobalParameterValue(UDynamicMaterialModel::GlobalSubsurfaceColorValueName));
+	SurfaceThickness   ->AddComponent(AlphaValueName, MaterialModel->GetGlobalParameterValue(UDynamicMaterialModel::GlobalSurfaceThicknessValueName));
 }
 
 TArray<FName> UDynamicMaterialModelEditorOnlyData::GetPresetOptions() const
@@ -282,7 +287,7 @@ void UDynamicMaterialModelEditorOnlyData::EnsurePresetSlots()
 		return;
 	}
 
-	ForEachMaterialPropertyType(
+	UE::DynamicMaterial::ForEachMaterialPropertyType(
 		[this, Preset](EDMMaterialPropertyType InProperty)
 		{
 			if (InProperty == EDMMaterialPropertyType::OpacityMask)
@@ -492,7 +497,7 @@ void UDynamicMaterialModelEditorOnlyData::BuildMaterial(bool bInDirtyAssets)
 	/**
 	 * Process slots to build base material inputs.
 	 */
-	ForEachMaterialPropertyType(
+	UE::DynamicMaterial::ForEachMaterialPropertyType(
 		[this, &BuildState](EDMMaterialPropertyType InType)
 		{
 			if (FExpressionInput* Input = BuildState->GetMaterialProperty(InType))
@@ -502,7 +507,7 @@ void UDynamicMaterialModelEditorOnlyData::BuildMaterial(bool bInDirtyAssets)
 
 			if (UDMMaterialProperty* Property = GetMaterialProperty(InType))
 			{
-				if (Property->IsEnabled() && Property->IsMaterialPin())
+				if (Property->IsEnabled() && Property->IsMaterialPin() && GetSlotForMaterialProperty(InType))
 				{
 					Property->GenerateExpressions(BuildState);
 
@@ -609,12 +614,12 @@ void UDynamicMaterialModelEditorOnlyData::BuildMaterial(bool bInDirtyAssets)
 	/**
 	 * Apply output processors.
 	 */
-	ForEachMaterialPropertyType(
+	UE::DynamicMaterial::ForEachMaterialPropertyType(
 		[this, &BuildState](EDMMaterialPropertyType InType)
 		{
 			if (UDMMaterialProperty* Property = GetMaterialProperty(InType))
 			{
-				if (Property->IsMaterialPin() && Property->IsEnabled())
+				if (Property->IsMaterialPin() && Property->IsEnabled() && GetSlotForMaterialProperty(InType))
 				{
 					Property->AddOutputProcessor(BuildState);
 				}
@@ -838,34 +843,6 @@ void UDynamicMaterialModelEditorOnlyData::OnWizardComplete()
 	}
 }
 
-void UDynamicMaterialModelEditorOnlyData::ForEachMaterialPropertyType(TFunctionRef<EDMIterationResult(EDMMaterialPropertyType InType)> InCallable,
-	EDMMaterialPropertyType InStart, EDMMaterialPropertyType InEnd)
-{
-	for (uint8 PropertyIndex = static_cast<uint8>(InStart); PropertyIndex <= static_cast<uint8>(InEnd); ++PropertyIndex)
-	{
-		const EDMMaterialPropertyType Property = static_cast<EDMMaterialPropertyType>(PropertyIndex);
-
-		if (InCallable(Property) == EDMIterationResult::Break)
-		{
-			break;
-		}
-	}
-}
-
-void UDynamicMaterialModelEditorOnlyData::ForEachMaterialPropertyType(TFunctionRef<EDMIterationResult(EDMMaterialPropertyType InType)> InCallable, 
-	EDMMaterialPropertyType InStart, EDMMaterialPropertyType InEnd) const
-{
-	for (uint8 PropertyIndex = static_cast<uint8>(InStart); PropertyIndex <= static_cast<uint8>(InEnd); ++PropertyIndex)
-	{
-		const EDMMaterialPropertyType Property = static_cast<EDMMaterialPropertyType>(PropertyIndex);
-
-		if (InCallable(Property) == EDMIterationResult::Break)
-		{
-			break;
-		}
-	}
-}
-
 void UDynamicMaterialModelEditorOnlyData::SetChannelListPreset(FName InPresetName)
 {
 	ChannelListPreset = InPresetName;
@@ -1153,7 +1130,7 @@ TMap<EDMMaterialPropertyType, UDMMaterialProperty*> UDynamicMaterialModelEditorO
 {
 	TMap<EDMMaterialPropertyType, UDMMaterialProperty*> LocalProperties;
 
-	ForEachMaterialPropertyType(
+	UE::DynamicMaterial::ForEachMaterialPropertyType(
 		[this, &LocalProperties](EDMMaterialPropertyType InType)
 		{
 			if (UDMMaterialProperty* Property = GetMaterialProperty(InType))
@@ -1276,7 +1253,7 @@ UDMMaterialSlot* UDynamicMaterialModelEditorOnlyData::AddSlot()
 {
 	UDMMaterialSlot* NewSlot = nullptr;
 
-	ForEachMaterialPropertyType(
+	UE::DynamicMaterial::ForEachMaterialPropertyType(
 		[this, &NewSlot](EDMMaterialPropertyType InProperty)
 		{
 			if (GetSlotForMaterialProperty(InProperty))
@@ -1708,7 +1685,7 @@ void UDynamicMaterialModelEditorOnlyData::PostEditorDuplicate()
 		Modify();
 	}
 
-	ForEachMaterialPropertyType(
+	UE::DynamicMaterial::ForEachMaterialPropertyType(
 		[this](EDMMaterialPropertyType InType)
 		{
 			if (UDMMaterialProperty* Property = GetMaterialProperty(InType))
diff --git a/Engine/Plugins/Experimental/DynamicMaterial/Source/DynamicMaterialEditor/Private/UI/PropertyGenerators/DMMaterialModelPropertyRowGenerator.cpp b/Engine/Plugins/Experimental/DynamicMaterial/Source/DynamicMaterialEditor/Private/UI/PropertyGenerators/DMMaterialModelPropertyRowGenerator.cpp
--- a/Engine/Plugins/Experimental/DynamicMaterial/Source/DynamicMaterialEditor/Private/UI/PropertyGenerators/DMMaterialModelPropertyRowGenerator.cpp
+++ b/Engine/Plugins/Experimental/DynamicMaterial/Source/DynamicMaterialEditor/Private/UI/PropertyGenerators/DMMaterialModelPropertyRowGenerator.cpp
@@ -38,18 +38,34 @@ void FDMMaterialModelPropertyRowGenerator::AddMaterialModelProperties(const TSha
 	{
 		if (EditorOnlyData->GetBlendMode() != BLEND_Opaque)
 		{
-			AddGlobalValue(InGlobalSettingEditorWidget, InMaterialModelBase, InOutPropertyRows,  
-				MaterialModel->GetGlobalParameterValue(UDynamicMaterialModel::GlobalOpacityValueName), 
-				LOCTEXT("GlobalOpacity", "Global Opacity"));
+			UDMMaterialProperty* OpacityProperty = EditorOnlyData->GetMaterialProperty(EDMMaterialPropertyType::Opacity);
+			UDMMaterialProperty* OpacityMaskProperty = EditorOnlyData->GetMaterialProperty(EDMMaterialPropertyType::OpacityMask);
+
+			const bool bOpacityValid = OpacityProperty && OpacityProperty->IsEnabled() && OpacityProperty->IsValidForModel(*EditorOnlyData)
+				&& EditorOnlyData->GetSlotForMaterialProperty(EDMMaterialPropertyType::Opacity);
+
+			const bool bOpacityMaskValid = OpacityMaskProperty && OpacityMaskProperty->IsEnabled() && OpacityMaskProperty->IsValidForModel(*EditorOnlyData)
+				&& EditorOnlyData->GetSlotForMaterialProperty(EDMMaterialPropertyType::OpacityMask);
+
+			if (bOpacityValid || bOpacityMaskValid)
+			{
+				AddGlobalValue(InGlobalSettingEditorWidget, InMaterialModelBase, InOutPropertyRows,
+					MaterialModel->GetGlobalParameterValue(UDynamicMaterialModel::GlobalOpacityValueName),
+					LOCTEXT("GlobalOpacity", "Global Opacity"));
+			}
 		}
 
-		EditorOnlyData->ForEachMaterialPropertyType(
+		UE::DynamicMaterial::ForEachMaterialPropertyType(
 			[&InGlobalSettingEditorWidget, InMaterialModelBase, &InOutPropertyRows, EditorOnlyData]
 			(EDMMaterialPropertyType InProperty)
 			{
-				return ForEachProperty(InProperty, InGlobalSettingEditorWidget, InMaterialModelBase, InOutPropertyRows, EditorOnlyData);
-			},
-			/* Start from */ EDMMaterialPropertyType::Roughness
+				if (InProperty != EDMMaterialPropertyType::Opacity && InProperty != EDMMaterialPropertyType::OpacityMask)
+				{
+					AddGlobalMaterialParameterValue(InProperty, InGlobalSettingEditorWidget, InMaterialModelBase, InOutPropertyRows, EditorOnlyData);
+				}
+
+				return EDMIterationResult::Continue;
+			}
 		);
 	}
 
@@ -84,7 +100,7 @@ void FDMMaterialModelPropertyRowGenerator::AddMaterialModelProperties(const TSha
 	}
 }
 
-EDMIterationResult FDMMaterialModelPropertyRowGenerator::ForEachProperty(EDMMaterialPropertyType InProperty, 
+void FDMMaterialModelPropertyRowGenerator::AddGlobalMaterialParameterValue(EDMMaterialPropertyType InProperty, 
 	const TSharedRef<SDMMaterialGlobalSettingsEditor>& InGlobalSettingEditorWidget, UDynamicMaterialModelBase* InMaterialModelBase, 
 	TArray<FDMPropertyHandle>& InOutPropertyRows, UDynamicMaterialModelEditorOnlyData* InEditorOnlyData)
 {
@@ -112,8 +128,6 @@ EDMIterationResult FDMMaterialModelPropertyRowGenerator::ForEachProperty(EDMMate
 			}
 		}
 	}
-
-	return EDMIterationResult::Continue;
 }
 
 void FDMMaterialModelPropertyRowGenerator::AddGlobalValue(const TSharedRef<SDMMaterialGlobalSettingsEditor>& InGlobalSettingEditorWidget,
diff --git a/Engine/Plugins/Experimental/DynamicMaterial/Source/DynamicMaterialEditor/Private/UI/PropertyGenerators/DMMaterialModelPropertyRowGenerator.h b/Engine/Plugins/Experimental/DynamicMaterial/Source/DynamicMaterialEditor/Private/UI/PropertyGenerators/DMMaterialModelPropertyRowGenerator.h
--- a/Engine/Plugins/Experimental/DynamicMaterial/Source/DynamicMaterialEditor/Private/UI/PropertyGenerators/DMMaterialModelPropertyRowGenerator.h
+++ b/Engine/Plugins/Experimental/DynamicMaterial/Source/DynamicMaterialEditor/Private/UI/PropertyGenerators/DMMaterialModelPropertyRowGenerator.h
@@ -11,7 +11,6 @@ class SDMMaterialGlobalSettingsEditor;
 class UDMMaterialComponent;
 class UDynamicMaterialModelBase;
 class UDynamicMaterialModelEditorOnlyData;
-enum class EDMIterationResult : uint8;
 enum class EDMMaterialPropertyType : uint8;
 struct FDMPropertyHandle;
 
@@ -28,7 +27,7 @@ public:
 	static void AddMaterialModelProperties(const TSharedRef<SDMMaterialGlobalSettingsEditor>& InGlobalSettingEditorWidget, 
 		UDynamicMaterialModelBase* InMaterialModelBase, TArray<FDMPropertyHandle>& InOutPropertyRows);
 
-	static EDMIterationResult ForEachProperty(EDMMaterialPropertyType InProperty, const TSharedRef<SDMMaterialGlobalSettingsEditor>& InGlobalSettingEditorWidget,
+	static void AddGlobalMaterialParameterValue(EDMMaterialPropertyType InProperty, const TSharedRef<SDMMaterialGlobalSettingsEditor>& InGlobalSettingEditorWidget,
 		UDynamicMaterialModelBase* InMaterialModelBase, TArray<FDMPropertyHandle>& InOutPropertyRows,
 		UDynamicMaterialModelEditorOnlyData* InEditorOnlyData);
 
diff --git a/Engine/Plugins/Experimental/DynamicMaterial/Source/DynamicMaterialEditor/Public/DMEDefs.h b/Engine/Plugins/Experimental/DynamicMaterial/Source/DynamicMaterialEditor/Public/DMEDefs.h
--- a/Engine/Plugins/Experimental/DynamicMaterial/Source/DynamicMaterialEditor/Public/DMEDefs.h
+++ b/Engine/Plugins/Experimental/DynamicMaterial/Source/DynamicMaterialEditor/Public/DMEDefs.h
@@ -26,12 +26,6 @@ class UDynamicMaterialInstance;
 class UDynamicMaterialModel;
 class UPrimitiveComponent;
 
-enum class EDMIterationResult : uint8
-{
-	Continue,
-	Break
-};
-
 UINTERFACE(MinimalAPI, Blueprintable, BlueprintType)
 class UDMBuildable : public UInterface
 {
diff --git a/Engine/Plugins/Experimental/DynamicMaterial/Source/DynamicMaterialEditor/Public/Model/DynamicMaterialModelEditorOnlyData.h b/Engine/Plugins/Experimental/DynamicMaterial/Source/DynamicMaterialEditor/Public/Model/DynamicMaterialModelEditorOnlyData.h
--- a/Engine/Plugins/Experimental/DynamicMaterial/Source/DynamicMaterialEditor/Public/Model/DynamicMaterialModelEditorOnlyData.h
+++ b/Engine/Plugins/Experimental/DynamicMaterial/Source/DynamicMaterialEditor/Public/Model/DynamicMaterialModelEditorOnlyData.h
@@ -222,14 +222,6 @@ public:
 
 	DYNAMICMATERIALEDITOR_API void OnWizardComplete();
 
-	DYNAMICMATERIALEDITOR_API void ForEachMaterialPropertyType(TFunctionRef<EDMIterationResult(EDMMaterialPropertyType InType)> InCallable,
-		EDMMaterialPropertyType InStart = static_cast<EDMMaterialPropertyType>(static_cast<uint8>(EDMMaterialPropertyType::None) + 1),
-		EDMMaterialPropertyType InEnd = static_cast<EDMMaterialPropertyType>(static_cast<uint8>(EDMMaterialPropertyType::Any) - 1));
-
-	DYNAMICMATERIALEDITOR_API void ForEachMaterialPropertyType(TFunctionRef<EDMIterationResult(EDMMaterialPropertyType InType)> InCallable,
-		EDMMaterialPropertyType InStart = static_cast<EDMMaterialPropertyType>(static_cast<uint8>(EDMMaterialPropertyType::None) + 1),
-		EDMMaterialPropertyType InEnd = static_cast<EDMMaterialPropertyType>(static_cast<uint8>(EDMMaterialPropertyType::Any) - 1)) const;
-
 	void SaveEditor();
 
 	//~ Begin FNotifyHook
@@ -329,67 +321,67 @@ protected:
 	FDMOnValueListUpdated OnValueListUpdateDelegate;
 	FDMOnSlotListUpdated OnSlotListUpdateDelegate;
 
-	UPROPERTY(VisibleInstanceOnly, BlueprintReadOnly, Instanced, Category = "Material Designer")
+	UPROPERTY(BlueprintReadOnly, Instanced, Category = "Material Designer")
 	TObjectPtr<UDMMaterialProperty> BaseColor;
 
-	UPROPERTY(VisibleInstanceOnly, BlueprintReadOnly, Instanced, Category = "Material Designer")
+	UPROPERTY(BlueprintReadOnly, Instanced, Category = "Material Designer")
 	TObjectPtr<UDMMaterialProperty> EmissiveColor;
 
-	UPROPERTY(VisibleInstanceOnly, BlueprintReadOnly, Instanced, Category = "Material Designer")
+	UPROPERTY(BlueprintReadOnly, Instanced, Category = "Material Designer")
 	TObjectPtr<UDMMaterialProperty> Opacity;
 
-	UPROPERTY(VisibleInstanceOnly, BlueprintReadOnly, Instanced, Category = "Material Designer")
+	UPROPERTY(BlueprintReadOnly, Instanced, Category = "Material Designer")
 	TObjectPtr<UDMMaterialProperty> OpacityMask;
 
-	UPROPERTY(VisibleInstanceOnly, BlueprintReadOnly, Instanced, Category = "Material Designer")
+	UPROPERTY(BlueprintReadOnly, Instanced, Category = "Material Designer")
 	TObjectPtr<UDMMaterialProperty> Roughness;
 
-	UPROPERTY(VisibleInstanceOnly, BlueprintReadOnly, Instanced, Category = "Material Designer")
+	UPROPERTY(BlueprintReadOnly, Instanced, Category = "Material Designer")
 	TObjectPtr<UDMMaterialProperty> Specular;
 
-	UPROPERTY(VisibleInstanceOnly, BlueprintReadOnly, Instanced, Category = "Material Designer")
+	UPROPERTY(BlueprintReadOnly, Instanced, Category = "Material Designer")
 	TObjectPtr<UDMMaterialProperty> Metallic;
 
-	UPROPERTY(VisibleInstanceOnly, BlueprintReadOnly, Instanced, Category = "Material Designer")
+	UPROPERTY(BlueprintReadOnly, Instanced, Category = "Material Designer")
 	TObjectPtr<UDMMaterialProperty> Normal;
 
-	UPROPERTY(VisibleInstanceOnly, BlueprintReadOnly, Instanced, Category = "Material Designer")
+	UPROPERTY(BlueprintReadOnly, Instanced, Category = "Material Designer")
 	TObjectPtr<UDMMaterialProperty> PixelDepthOffset;
 
-	UPROPERTY(VisibleInstanceOnly, BlueprintReadOnly, Instanced, Category = "Material Designer")
+	UPROPERTY(BlueprintReadOnly, Instanced, Category = "Material Designer")
 	TObjectPtr<UDMMaterialProperty> WorldPositionOffset;
 
-	UPROPERTY(VisibleInstanceOnly, BlueprintReadOnly, Instanced, Category = "Material Designer")
+	UPROPERTY(BlueprintReadOnly, Instanced, Category = "Material Designer")
 	TObjectPtr<UDMMaterialProperty> AmbientOcclusion;
 
-	UPROPERTY(VisibleInstanceOnly, BlueprintReadOnly, Instanced, Category = "Material Designer")
+	UPROPERTY(BlueprintReadOnly, Instanced, Category = "Material Designer")
 	TObjectPtr<UDMMaterialProperty> Anisotropy;
 
-	UPROPERTY(VisibleInstanceOnly, BlueprintReadOnly, Instanced, Category = "Material Designer")
+	UPROPERTY(BlueprintReadOnly, Instanced, Category = "Material Designer")
 	TObjectPtr<UDMMaterialProperty> Refraction;
 
-	UPROPERTY(VisibleInstanceOnly, BlueprintReadOnly, Instanced, Category = "Material Designer")
+	UPROPERTY(BlueprintReadOnly, Instanced, Category = "Material Designer")
 	TObjectPtr<UDMMaterialProperty> Tangent;
 
-	UPROPERTY(VisibleInstanceOnly, BlueprintReadOnly, Instanced, Category = "Material Designer")
+	UPROPERTY(BlueprintReadOnly, Instanced, Category = "Material Designer")
 	TObjectPtr<UDMMaterialProperty> Displacement;
 
-	UPROPERTY(VisibleInstanceOnly, BlueprintReadOnly, Instanced, Category = "Material Designer")
+	UPROPERTY(BlueprintReadOnly, Instanced, Category = "Material Designer")
 	TObjectPtr<UDMMaterialProperty> SubsurfaceColor;
 
-	UPROPERTY(VisibleInstanceOnly, BlueprintReadOnly, Instanced, Category = "Material Designer")
+	UPROPERTY(BlueprintReadOnly, Instanced, Category = "Material Designer")
 	TObjectPtr<UDMMaterialProperty> SurfaceThickness;
 
-	UPROPERTY(VisibleInstanceOnly, BlueprintReadOnly, Instanced, Category = "Material Designer")
+	UPROPERTY(BlueprintReadOnly, Instanced, Category = "Material Designer")
 	TObjectPtr<UDMMaterialProperty> Custom1;
 
-	UPROPERTY(VisibleInstanceOnly, BlueprintReadOnly, Instanced, Category = "Material Designer")
+	UPROPERTY(BlueprintReadOnly, Instanced, Category = "Material Designer")
 	TObjectPtr<UDMMaterialProperty> Custom2;
 
-	UPROPERTY(VisibleInstanceOnly, BlueprintReadOnly, Instanced, Category = "Material Designer")
+	UPROPERTY(BlueprintReadOnly, Instanced, Category = "Material Designer")
 	TObjectPtr<UDMMaterialProperty> Custom3;
 
-	UPROPERTY(VisibleInstanceOnly, BlueprintReadOnly, Instanced, Category = "Material Designer")
+	UPROPERTY(BlueprintReadOnly, Instanced, Category = "Material Designer")
 	TObjectPtr<UDMMaterialProperty> Custom4;
 
 	void CreateMaterial();
