Material Designer
- The slot and stage menus now workly correctly with context.
- The stage menu is now displayed when right clicking on a stage.
- The proper cursor is now displayed (grab hand) to show that layers can be dragged.
- The proper cursor is now displayed (grab hand) to show that effects can be dragged.

diff --git a/Engine/Plugins/Experimental/DynamicMaterial/Source/DynamicMaterialEditor/Private/UI/Menus/DMMaterialSlotLayerMenus.cpp b/Engine/Plugins/Experimental/DynamicMaterial/Source/DynamicMaterialEditor/Private/UI/Menus/DMMaterialSlotLayerMenus.cpp
--- a/Engine/Plugins/Experimental/DynamicMaterial/Source/DynamicMaterialEditor/Private/UI/Menus/DMMaterialSlotLayerMenus.cpp
+++ b/Engine/Plugins/Experimental/DynamicMaterial/Source/DynamicMaterialEditor/Private/UI/Menus/DMMaterialSlotLayerMenus.cpp
@@ -40,36 +40,39 @@ namespace UE::DynamicMaterialEditor::Private
 	static const FName GlobalValuesSectionName = TEXT("GlobalValues");
 }
 
-UToolMenu* FDMMaterialSlotLayerMenus::GenerateSlotLayerMenu(const TSharedPtr<SDMMaterialSlotEditor>& InSlotWidget, UDMMaterialLayerObject* InLayerObject)
+TSharedRef<SWidget> FDMMaterialSlotLayerMenus::GenerateSlotLayerMenu(const TSharedPtr<SDMMaterialSlotEditor>& InSlotWidget, UDMMaterialLayerObject* InLayerObject)
 {
 	using namespace UE::DynamicMaterialEditor::Private;
 
-	UToolMenu* NewToolMenu = UDMMenuContext::GenerateContextMenuLayer(SlotLayerMenuName, InSlotWidget->GetEditorWidget(), InLayerObject);
-
-	if (!NewToolMenu)
+	if (!UToolMenus::Get()->IsMenuRegistered(SlotLayerMenuName))
 	{
-		return nullptr;
-	}
+		UToolMenu* NewToolMenu = UDMMenuContext::GenerateContextMenuDefault(SlotLayerMenuName);
 
-	AddAddLayerSection(NewToolMenu);
+		if (!NewToolMenu)
+		{
+			return SNullWidget::NullWidget;
+		}
 
-	if constexpr (UE::DynamicMaterialEditor::bGlobalValuesEnabled)
-	{
-		AddGlobalValueSection(NewToolMenu);
-	}
+		NewToolMenu->AddDynamicSection(NAME_None, FNewToolMenuDelegate::CreateStatic(&FDMMaterialSlotLayerMenus::AddAddLayerSection));
+
+		if constexpr (UE::DynamicMaterialEditor::bGlobalValuesEnabled)
+		{
+			NewToolMenu->AddDynamicSection(NAME_None, FNewToolMenuDelegate::CreateStatic(&FDMMaterialSlotLayerMenus::AddGlobalValueSection));
+		}
 
-	AddLayerAddEffectsSection(NewToolMenu, InLayerObject);
+		NewToolMenu->AddDynamicSection(NAME_None, FNewToolMenuDelegate::CreateStatic(&FDMMaterialSlotLayerMenus::AddLayerModifySection));
+	}
 
-	AddLayerModifySection(NewToolMenu);
+	FToolMenuContext MenuContext(UDMMenuContext::CreateEditor(InSlotWidget->GetEditorWidget()));
 
-	return NewToolMenu;
+	return UToolMenus::Get()->GenerateWidget(SlotLayerMenuName, MenuContext);
 }
 
 void FDMMaterialSlotLayerMenus::AddAddLayerSection(UToolMenu* InMenu)
 {
 	using namespace UE::DynamicMaterialEditor::Private;
 
-	if (!IsValid(InMenu) || InMenu->ContainsSection(SlotLayerAddSectionName))
+	if (!IsValid(InMenu))
 	{
 		return;
 	}
@@ -347,7 +350,7 @@ void FDMMaterialSlotLayerMenus::AddLayerModifySection(UToolMenu* InMenu)
 {
 	using namespace UE::DynamicMaterialEditor::Private;
 
-	if (!IsValid(InMenu) || InMenu->ContainsSection(SlotLayerModifySectionName))
+	if (!IsValid(InMenu))
 	{
 		return;
 	}
@@ -368,16 +371,11 @@ void FDMMaterialSlotLayerMenus::AddLayerModifySection(UToolMenu* InMenu)
 	NewSection.AddMenuEntry(FGenericCommands::Get().Delete);
 }
 
-void FDMMaterialSlotLayerMenus::AddLayerAddEffectsSection(UToolMenu* InMenu, UDMMaterialLayerObject* InLayerObject)
-{
-	FDMMaterialSlotLayerAddEffectMenus::AddEffectSubMenu(InMenu, InLayerObject);
-}
-
 void FDMMaterialSlotLayerMenus::AddGlobalValueSection(UToolMenu* InMenu)
 {
 	using namespace UE::DynamicMaterialEditor::Private;
 
-	if (!IsValid(InMenu) || InMenu->ContainsSection(GlobalValuesSectionName))
+	if (!IsValid(InMenu))
 	{
 		return;
 	}
diff --git a/Engine/Plugins/Experimental/DynamicMaterial/Source/DynamicMaterialEditor/Private/UI/Menus/DMMaterialSlotLayerMenus.h b/Engine/Plugins/Experimental/DynamicMaterial/Source/DynamicMaterialEditor/Private/UI/Menus/DMMaterialSlotLayerMenus.h
--- a/Engine/Plugins/Experimental/DynamicMaterial/Source/DynamicMaterialEditor/Private/UI/Menus/DMMaterialSlotLayerMenus.h
+++ b/Engine/Plugins/Experimental/DynamicMaterial/Source/DynamicMaterialEditor/Private/UI/Menus/DMMaterialSlotLayerMenus.h
@@ -6,6 +6,7 @@
 #include "Templates/SharedPointerFwd.h"
 
 class SDMMaterialSlotEditor;
+class SWidget;
 class UDMMaterialLayerObject;
 class UDMMaterialSlot;
 class UToolMenu;
@@ -14,15 +15,13 @@ enum class EDMMaterialPropertyType : uint8;
 class FDMMaterialSlotLayerMenus final
 {
 public:
-	static UToolMenu* GenerateSlotLayerMenu(const TSharedPtr<SDMMaterialSlotEditor>& InSlotWidget, UDMMaterialLayerObject* InLayerObject);
+	static TSharedRef<SWidget> GenerateSlotLayerMenu(const TSharedPtr<SDMMaterialSlotEditor>& InSlotWidget, UDMMaterialLayerObject* InLayerObject);
 
 	static void AddAddLayerSection(UToolMenu* InMenu);
 
 private:
 	static void AddLayerModifySection(UToolMenu* InMenu);
 
-	static void AddLayerAddEffectsSection(UToolMenu* InMenu, UDMMaterialLayerObject* InLayerObject);
-
 	static void AddGlobalValueSection(UToolMenu* InMenu);
 
 	static void AddSlotMenuEntry(const TSharedPtr<SDMMaterialSlotEditor> InSlotWidget, UToolMenu* InMenu, const FText& InName, 
diff --git a/Engine/Plugins/Experimental/DynamicMaterial/Source/DynamicMaterialEditor/Private/UI/Menus/DMMaterialStageMenus.cpp b/Engine/Plugins/Experimental/DynamicMaterial/Source/DynamicMaterialEditor/Private/UI/Menus/DMMaterialStageMenus.cpp
--- a/Engine/Plugins/Experimental/DynamicMaterial/Source/DynamicMaterialEditor/Private/UI/Menus/DMMaterialStageMenus.cpp
+++ b/Engine/Plugins/Experimental/DynamicMaterial/Source/DynamicMaterialEditor/Private/UI/Menus/DMMaterialStageMenus.cpp
@@ -8,6 +8,7 @@
 #include "ScopedTransaction.h"
 #include "Templates/SharedPointer.h"
 #include "ToolMenu.h"
+#include "ToolMenus.h"
 #include "UI/Menus/DMMenuContext.h"
 #include "UI/Widgets/Editor/SDMMaterialSlotEditor.h"
 #include "UI/Widgets/Editor/SlotEditor/SDMMaterialStage.h"
@@ -15,33 +16,41 @@
 
 #define LOCTEXT_NAMESPACE "FDMMaterialStageMenus"
 
-FName FDMMaterialStageMenus::GetStageMenuName()
+namespace UE::DynamicMaterialEditor::Private
 {
-	return "MaterialDesigner.MaterialStage";
+	const FLazyName StageMenuName = TEXT("MaterialDesigner.MaterialStage");
+	const FLazyName StageMenuToggleName = TEXT("StageToggle");
 }
 
-FName FDMMaterialStageMenus::GetStageToggleSectionName()
-{
-	return "StageToggle";
-}
-
-UToolMenu* FDMMaterialStageMenus::GenerateStageMenu(const TSharedPtr<SDMMaterialSlotEditor>& InSlotWidget, 
+TSharedRef<SWidget> FDMMaterialStageMenus::GenerateStageMenu(const TSharedPtr<SDMMaterialSlotEditor>& InSlotWidget,
 	const TSharedPtr<SDMMaterialStage>& InStageWidget)
 {
-	UToolMenu* NewToolMenu = UDMMenuContext::GenerateContextMenuStage(GetStageMenuName(), InSlotWidget->GetEditorWidget(), InStageWidget);
-	if (!NewToolMenu)
+	using namespace UE::DynamicMaterialEditor::Private;
+
+	UToolMenus* ToolMenus = UToolMenus::Get();
+
+	if (!ToolMenus->IsMenuRegistered(StageMenuName))
 	{
-		return nullptr;
+		UToolMenu* NewToolMenu = UDMMenuContext::GenerateContextMenuDefault(StageMenuName);
+
+		if (!NewToolMenu)
+		{
+			return SNullWidget::NullWidget;
+		}
+
+		NewToolMenu->AddDynamicSection(NAME_None, FNewToolMenuDelegate::CreateStatic(&FDMMaterialStageMenus::AddStageSection));
 	}
 
-	AddStageSection(NewToolMenu);
+	FToolMenuContext MenuContext(UDMMenuContext::CreateStage(InSlotWidget->GetEditorWidget(), InStageWidget));
 
-	return NewToolMenu;
+	return ToolMenus->GenerateWidget(StageMenuName, MenuContext);
 }
 
 void FDMMaterialStageMenus::AddStageSection(UToolMenu* InMenu)
 {
-	if (!IsValid(InMenu) || InMenu->ContainsSection(GetStageToggleSectionName()))
+	using namespace UE::DynamicMaterialEditor::Private;
+
+	if (!IsValid(InMenu) || InMenu->ContainsSection(StageMenuToggleName))
 	{
 		return;
 	}
@@ -58,7 +67,7 @@ void FDMMaterialStageMenus::AddStageSection(UToolMenu* InMenu)
 		return;
 	}
 
-	UDMMaterialLayerObject* Layer = MenuContext->GetLayer();
+	UDMMaterialLayerObject* Layer = Stage->GetLayer();
 	if (!Layer)
 	{
 		return;
@@ -70,29 +79,54 @@ void FDMMaterialStageMenus::AddStageSection(UToolMenu* InMenu)
 		return;
 	}
 
-	const int32 LayerIndex = Layer->FindIndex();
-
 	const bool bAllowRemoveLayer = Slot->CanRemoveLayer(Layer);
 	const EDMMaterialLayerStage StageType = Layer->GetStageType(Stage);
 
-	if (bAllowRemoveLayer || StageType == EDMMaterialLayerStage::Mask)
+	FToolMenuSection& NewSection = InMenu->AddSection(StageMenuToggleName, LOCTEXT("MaterialStageMenu", "Material Stage"));
+
+	if (!bAllowRemoveLayer && StageType != EDMMaterialLayerStage::Mask)
+	{
+		return;
+	}
+
+	if (StageType == EDMMaterialLayerStage::Mask)
 	{
-		FToolMenuSection& NewSection = InMenu->AddSection(GetStageToggleSectionName(), LOCTEXT("MaterialStageMenu", "Material Stage"));
+		NewSection.AddMenuEntry(NAME_None,
+			LOCTEXT("ToggleLayerMask", "Toggle Mask"),
+			LOCTEXT("ToggleLayerMaskTooltip", "Toggle the layer mask on and off.\n\nAlt+Shift+Left Click"),
+			FSlateIcon(),
+			FUIAction(FExecuteAction::CreateWeakLambda(
+				Layer,
+				[Layer]()
+				{
+					FScopedTransaction Transaction(LOCTEXT("ToggleBaseStageEnabled", "Toggle Base Stage Enabled"));
+
+					if (UDMMaterialStage* Stage = Layer->GetStage(EDMMaterialLayerStage::Mask))
+					{
+						Stage->Modify();
+						Stage->SetEnabled(!Stage->IsEnabled());
+					}
+				}
+			))
+		);
+	}
 
-		if (bAllowRemoveLayer)
+	if (bAllowRemoveLayer)
+	{
+		if (StageType == EDMMaterialLayerStage::Base)
 		{
 			NewSection.AddMenuEntry(NAME_None,
-				LOCTEXT("ToggleLayer", "Toggle"),
-				LOCTEXT("ToggleLayerTooltip", "Toggle the entire layer on and off.\n\n"
-					"Warning: Toggling a layer off may result in inputs being reset where incompatibilities are found.\n\nAlt+Left Click"),
+				LOCTEXT("ToggleLayerBase", "Toggle Base"),
+				LOCTEXT("ToggleLayerBaseTooltip", "Toggle the layer base on and off.\n\n"
+					"Warning: Toggling a layer base off may result in inputs being reset where incompatibilities are found.\n\nAlt+Shift+Left Click"),
 				FSlateIcon(),
 				FUIAction(FExecuteAction::CreateWeakLambda(
 					Layer,
 					[Layer]()
 					{
-						FScopedTransaction Transaction(LOCTEXT("ToggleAllStageEnabled", "Toggle All Stage Enabled"));
+						FScopedTransaction Transaction(LOCTEXT("ToggleBaseStageEnabled", "Toggle Base Stage Enabled"));
 
-						for (UDMMaterialStage* Stage : Layer->GetStages(EDMMaterialLayerStage::All))
+						if (UDMMaterialStage* Stage = Layer->GetStage(EDMMaterialLayerStage::Base))
 						{
 							Stage->Modify();
 							Stage->SetEnabled(!Stage->IsEnabled());
@@ -100,77 +134,52 @@ void FDMMaterialStageMenus::AddStageSection(UToolMenu* InMenu)
 					}
 				))
 			);
-
-			if (StageType == EDMMaterialLayerStage::Base)
-			{
-				NewSection.AddMenuEntry(NAME_None,
-					LOCTEXT("ToggleLayerBase", "Toggle Base"),
-					LOCTEXT("ToggleLayerBaseTooltip", "Toggle the layer base on and off.\n\n"
-						"Warning: Toggling a layer base off may result in inputs being reset where incompatibilities are found.\n\nAlt+Shift+Left Click"),
-					FSlateIcon(),
-					FUIAction(FExecuteAction::CreateWeakLambda(
-						Layer,
-						[Layer]()
-						{
-							FScopedTransaction Transaction(LOCTEXT("ToggleBaseStageEnabled", "Toggle Base Stage Enabled"));
-
-							if (UDMMaterialStage* Stage = Layer->GetStage(EDMMaterialLayerStage::Base))
-							{
-								Stage->Modify();
-								Stage->SetEnabled(!Stage->IsEnabled());
-							}
-						}
-					))
-				);
-			}
 		}
 
-		if (StageType == EDMMaterialLayerStage::Mask)
+		NewSection.AddMenuEntry(NAME_None,
+			LOCTEXT("ToggleLayer", "Toggle Layer"),
+			LOCTEXT("ToggleLayerTooltip", "Toggle the entire layer on and off.\n\n"
+				"Warning: Toggling a layer off may result in inputs being reset where incompatibilities are found.\n\nAlt+Left Click"),
+			FSlateIcon(),
+			FUIAction(FExecuteAction::CreateWeakLambda(
+				Layer,
+				[Layer]()
+				{
+					FScopedTransaction Transaction(LOCTEXT("ToggleAllStageEnabled", "Toggle All Stage Enabled"));
+
+					for (UDMMaterialStage* Stage : Layer->GetStages(EDMMaterialLayerStage::All))
+					{
+						Stage->Modify();
+						Stage->SetEnabled(!Stage->IsEnabled());
+					}
+				}
+			))
+		);			
+	}
+
+	if (bAllowRemoveLayer)
+	{
+		if (TSharedPtr<SDMMaterialEditor> EditorWidget = MenuContext->GetEditorWidget())
 		{
 			NewSection.AddMenuEntry(NAME_None,
-				LOCTEXT("ToggleLayerMask", "Toggle Mask"),
-				LOCTEXT("ToggleLayerMaskTooltip", "Toggle the layer mask on and off.\n\nAlt+Shift+Left Click"),
+				LOCTEXT("RemoveLayer", "Remove Layer"),
+				LOCTEXT("RemoveLayerTooltip", "Remove this layer from its Material Slot.\n\n"
+					"Warning: Removing a stage off may result in inputs being reset where incompatibilities are found."),
 				FSlateIcon(),
 				FUIAction(FExecuteAction::CreateWeakLambda(
 					Layer,
 					[Layer]()
 					{
-						FScopedTransaction Transaction(LOCTEXT("ToggleBaseStageEnabled", "Toggle Base Stage Enabled"));
+						FScopedTransaction Transaction(LOCTEXT("RemoverLayer", "Remover Layer"));
 
-						if (UDMMaterialStage* Stage = Layer->GetStage(EDMMaterialLayerStage::Mask))
+						if (UDMMaterialSlot* Slot = Layer->GetSlot())
 						{
-							Stage->Modify();
-							Stage->SetEnabled(!Stage->IsEnabled());
+							Slot->RemoveLayer(Layer);
 						}
 					}
 				))
 			);
 		}
-
-		if (bAllowRemoveLayer)
-		{
-			if (TSharedPtr<SDMMaterialEditor> EditorWidget = MenuContext->GetEditorWidget())
-			{
-				NewSection.AddMenuEntry(NAME_None,
-					LOCTEXT("RemoveLayer", "Remove"),
-					LOCTEXT("RemoveLayerTooltip", "Remove this layer from its Material Slot.\n\n"
-						"Warning: Removing a stage off may result in inputs being reset where incompatibilities are found."),
-					FSlateIcon(),
-					FUIAction(FExecuteAction::CreateWeakLambda(
-						Layer,
-						[Layer]()
-						{
-							FScopedTransaction Transaction(LOCTEXT("RemoverLayer", "Remover Layer"));
-
-							if (UDMMaterialSlot* Slot = Layer->GetSlot())
-							{
-								Slot->RemoveLayer(Layer);
-							}
-						}
-					))
-				);
-			}
-		}
 	}
 }
 
diff --git a/Engine/Plugins/Experimental/DynamicMaterial/Source/DynamicMaterialEditor/Private/UI/Menus/DMMaterialStageMenus.h b/Engine/Plugins/Experimental/DynamicMaterial/Source/DynamicMaterialEditor/Private/UI/Menus/DMMaterialStageMenus.h
--- a/Engine/Plugins/Experimental/DynamicMaterial/Source/DynamicMaterialEditor/Private/UI/Menus/DMMaterialStageMenus.h
+++ b/Engine/Plugins/Experimental/DynamicMaterial/Source/DynamicMaterialEditor/Private/UI/Menus/DMMaterialStageMenus.h
@@ -7,17 +7,14 @@
 class FName;
 class SDMMaterialSlotEditor;
 class SDMMaterialStage;
+class SWidget;
 class UToolMenu;
 
 class FDMMaterialStageMenus final
 {
 public:
-	static UToolMenu* GenerateStageMenu(const TSharedPtr<SDMMaterialSlotEditor>& InSlotWidget, const TSharedPtr<SDMMaterialStage>& InStageWidget);
+	static TSharedRef<SWidget> GenerateStageMenu(const TSharedPtr<SDMMaterialSlotEditor>& InSlotWidget, const TSharedPtr<SDMMaterialStage>& InStageWidget);
 
 private:
-	static FName GetStageMenuName();
-
-	static FName GetStageToggleSectionName();
-
 	static void AddStageSection(UToolMenu* InMenu);
 };
diff --git a/Engine/Plugins/Experimental/DynamicMaterial/Source/DynamicMaterialEditor/Private/UI/Widgets/Editor/SDMMaterialSlotEditor.cpp b/Engine/Plugins/Experimental/DynamicMaterial/Source/DynamicMaterialEditor/Private/UI/Widgets/Editor/SDMMaterialSlotEditor.cpp
--- a/Engine/Plugins/Experimental/DynamicMaterial/Source/DynamicMaterialEditor/Private/UI/Widgets/Editor/SDMMaterialSlotEditor.cpp
+++ b/Engine/Plugins/Experimental/DynamicMaterial/Source/DynamicMaterialEditor/Private/UI/Widgets/Editor/SDMMaterialSlotEditor.cpp
@@ -884,8 +884,7 @@ TSharedRef<SWidget> SDMMaterialSlotEditor::GetLayerButtonsMenuContent()
 {
 	if (UDMMaterialLayerObject* LayerObject = LayerViewSlot->GetSelectedLayer())
 	{
-		UToolMenu* ContextMenu = FDMMaterialSlotLayerMenus::GenerateSlotLayerMenu(SharedThis(this), LayerObject);
-		return UToolMenus::Get()->GenerateWidget(ContextMenu);
+		return FDMMaterialSlotLayerMenus::GenerateSlotLayerMenu(SharedThis(this), LayerObject);
 	}
 
 	return SNullWidget::NullWidget;
diff --git a/Engine/Plugins/Experimental/DynamicMaterial/Source/DynamicMaterialEditor/Private/UI/Widgets/Editor/SlotEditor/SDMMaterialSlotLayerEffectItem.cpp b/Engine/Plugins/Experimental/DynamicMaterial/Source/DynamicMaterialEditor/Private/UI/Widgets/Editor/SlotEditor/SDMMaterialSlotLayerEffectItem.cpp
--- a/Engine/Plugins/Experimental/DynamicMaterial/Source/DynamicMaterialEditor/Private/UI/Widgets/Editor/SlotEditor/SDMMaterialSlotLayerEffectItem.cpp
+++ b/Engine/Plugins/Experimental/DynamicMaterial/Source/DynamicMaterialEditor/Private/UI/Widgets/Editor/SlotEditor/SDMMaterialSlotLayerEffectItem.cpp
@@ -35,6 +35,8 @@ void SDMMaterialSlotLayerEffectItem::Construct(const FArguments& InArgs, const T
 	);
 
 	SetContent(CreateMainContent());
+
+	SetCursor(EMouseCursor::GrabHand);
 }
 
 TSharedPtr<SDMMaterialSlotLayerEffectView> SDMMaterialSlotLayerEffectItem::GetEffectView() const
@@ -75,37 +77,33 @@ TSharedRef<SWidget> SDMMaterialSlotLayerEffectItem::CreateMainContent()
 	constexpr float HorizontalSpacingEnd = 3.f;
 	constexpr float VerticalSpacing = 2.f;
 
-	return SNew(SVerticalBox)
-		+ SVerticalBox::Slot()
+	return SNew(SHorizontalBox)
+		.ToolTipText(GetToolTipText())
+		+ SHorizontalBox::Slot()
+		.AutoWidth()
+		.HAlign(HAlign_Left)
+		.VAlign(VAlign_Center)
+		.Padding(HorizontalSpacingEnd, VerticalSpacing, HorizontalSpacingEnd, VerticalSpacing)
+		[
+			CreateLayerBypassButton()
+		]
+		+ SHorizontalBox::Slot()
+		.FillWidth(1.0f)
+		.HAlign(HAlign_Left)
+		.VAlign(VAlign_Center)
+		.Padding(HorizontalSpacing, VerticalSpacing, HorizontalSpacing, VerticalSpacing)
+		[
+			SNew(STextBlock)
+			.TextStyle(FDynamicMaterialEditorStyle::Get(), "SmallFont")
+			.Text(GetLayerHeaderText())
+		]
+		+ SHorizontalBox::Slot()
+		.AutoWidth()
+		.HAlign(HAlign_Right)
+		.VAlign(VAlign_Center)
+		.Padding(HorizontalSpacingEnd, VerticalSpacing, HorizontalSpacingEnd, VerticalSpacing)
 		[
-			SNew(SHorizontalBox)
-			.ToolTipText(GetToolTipText())
-			+ SHorizontalBox::Slot()
-			.AutoWidth()
-			.HAlign(HAlign_Left)
-			.VAlign(VAlign_Center)
-			.Padding(HorizontalSpacingEnd, VerticalSpacing, HorizontalSpacingEnd, VerticalSpacing)
-			[
-				CreateLayerBypassButton()
-			]
-			+ SHorizontalBox::Slot()
-			.FillWidth(1.0f)
-			.HAlign(HAlign_Left)
-			.VAlign(VAlign_Center)
-			.Padding(HorizontalSpacing, VerticalSpacing, HorizontalSpacing, VerticalSpacing)
-			[
-				SNew(STextBlock)
-				.TextStyle(FDynamicMaterialEditorStyle::Get(), "SmallFont")
-				.Text(GetLayerHeaderText())
-			]
-			+ SHorizontalBox::Slot()
-			.AutoWidth()
-			.HAlign(HAlign_Right)
-			.VAlign(VAlign_Center)
-			.Padding(HorizontalSpacingEnd, VerticalSpacing, HorizontalSpacingEnd, VerticalSpacing)
-			[
-				CreateLayerRemoveButton()
-			]
+			CreateLayerRemoveButton()
 		];
 }
 
@@ -115,6 +113,7 @@ TSharedRef<SWidget> SDMMaterialSlotLayerEffectItem::CreateLayerBypassButton()
 		.ContentPadding(FMargin(2.0f, 2.f))
 		.ButtonStyle(FDynamicMaterialEditorStyle::Get(), "HoverHintOnly")
 		.ToolTipText(LOCTEXT("LayerBypassTooltip", "Toggle the bypassing of this layer."))
+		.Cursor(EMouseCursor::Default)
 		.OnClicked(this, &SDMMaterialSlotLayerEffectItem::OnLayerBypassButtonClick)
 		[
 			SNew(SImage)
@@ -129,6 +128,7 @@ TSharedRef<SWidget> SDMMaterialSlotLayerEffectItem::CreateLayerRemoveButton()
 		.ContentPadding(FMargin(2.0f, 2.f))
 		.ButtonStyle(FDynamicMaterialEditorStyle::Get(), "HoverHintOnly")
 		.ToolTipText(LOCTEXT("RemoveEffectTooltip", "Remove Effect"))
+		.Cursor(EMouseCursor::Default)
 		.OnClicked(this, &SDMMaterialSlotLayerEffectItem::OnLayerRemoveButtonClick)
 		[
 			SNew(SImage)
diff --git a/Engine/Plugins/Experimental/DynamicMaterial/Source/DynamicMaterialEditor/Private/UI/Widgets/Editor/SlotEditor/SDMMaterialSlotLayerItem.cpp b/Engine/Plugins/Experimental/DynamicMaterial/Source/DynamicMaterialEditor/Private/UI/Widgets/Editor/SlotEditor/SDMMaterialSlotLayerItem.cpp
--- a/Engine/Plugins/Experimental/DynamicMaterial/Source/DynamicMaterialEditor/Private/UI/Widgets/Editor/SlotEditor/SDMMaterialSlotLayerItem.cpp
+++ b/Engine/Plugins/Experimental/DynamicMaterial/Source/DynamicMaterialEditor/Private/UI/Widgets/Editor/SlotEditor/SDMMaterialSlotLayerItem.cpp
@@ -69,6 +69,8 @@ void SDMMaterialSlotLayerItem::Construct(const FArguments& InArgs, const TShared
 	);
 
 	SetContent(CreateMainContent());
+
+	SetCursor(EMouseCursor::GrabHand);
 }
 
 TSharedPtr<SDMMaterialSlotLayerView> SDMMaterialSlotLayerItem::GetSlotLayerView() const
@@ -294,6 +296,7 @@ TSharedRef<SWidget> SDMMaterialSlotLayerItem::CreateEffectsRowContent()
 {
 	return SNew(SHorizontalBox)
 		.Visibility(this, &SDMMaterialSlotLayerItem::GetEffectsListVisibility)
+		.Cursor(EMouseCursor::Default)
 		+ SHorizontalBox::Slot()
 		.AutoWidth()
 		[
@@ -391,6 +394,7 @@ TSharedRef<SWidget> SDMMaterialSlotLayerItem::CreateLayerBypassButton()
 		.IsEnabled(!bIsDynamic)
 		.ContentPadding(4.0f)
 		.ButtonStyle(FDynamicMaterialEditorStyle::Get(), "HoverHintOnly")
+		.Cursor(EMouseCursor::Default)
 		.ToolTipText(LOCTEXT("LayerBypassTooltip", "Toggle the bypassing of this layer."))
 		.OnClicked(this, &SDMMaterialSlotLayerItem::OnCreateLayerBypassButtonClicked)
 		[
@@ -433,6 +437,7 @@ TSharedRef<SWidget> SDMMaterialSlotLayerItem::CreateLayerBaseToggleButton()
 		.IsEnabled(!bIsDynamic)
 		.ContentPadding(0.0f)
 		.ButtonStyle(FDynamicMaterialEditorStyle::Get(), "HoverHintOnly")
+		.Cursor(EMouseCursor::Default)
 		.ToolTipText(LOCTEXT("MaterialLayerBaseToggleTooltip", "Toggle the layer base on and off.\n\n"
 			"Warning: Toggling a layer base off may result in inputs being reset where incompatibilities are found.\n\n"
 			"The base of the first layer cannot be toggled off."))
@@ -450,6 +455,7 @@ TSharedRef<SWidget> SDMMaterialSlotLayerItem::CreateLayerMaskToggleButton()
 		.ContentPadding(0.0f)
 		.ButtonStyle(FDynamicMaterialEditorStyle::Get(), "HoverHintOnly")
 		.ToolTipText(LOCTEXT("MaterialLayerMaskToggleTooltip", "Toggle the layer mask on and off."))
+		.Cursor(EMouseCursor::Default)
 		.OnClicked(this, &SDMMaterialSlotLayerItem::OnStageToggleButtonClicked, EDMMaterialLayerStage::Mask)
 		[
 			SNew(SImage)
@@ -464,6 +470,7 @@ TSharedRef<SWidget> SDMMaterialSlotLayerItem::CreateLayerLinkToggleButton()
 		.ContentPadding(0.0f)
 		.ButtonStyle(FDynamicMaterialEditorStyle::Get(), "HoverHintOnly")
 		.ToolTipText(LOCTEXT("MaterialStageLinkTooltip", "Click to toggle UV Link on and off."))
+		.Cursor(EMouseCursor::Default)
 		.OnClicked(this, &SDMMaterialSlotLayerItem::OnLayerLinkToggleButton)
 		.Visibility(this, &SDMMaterialSlotLayerItem::GetLayerLinkToggleButtonVisibility)
 		[
@@ -535,8 +542,9 @@ TSharedRef<SWidget> SDMMaterialSlotLayerItem::CreateEffectsToggleButton()
 		.ContentPadding(0.0f)
 		.ButtonStyle(FDynamicMaterialEditorStyle::Get(), "HoverHintOnly")
 		.ToolTipText(LOCTEXT("MaterialLayerFxTooltip", "Show or hide the effect list."))
-		.OnClicked(this, &SDMMaterialSlotLayerItem::OnEffectsToggleButtonClicked)
 		.Visibility(this, &SDMMaterialSlotLayerItem::GetEffectsToggleButtonVisibility)
+		.Cursor(EMouseCursor::Default)
+		.OnClicked(this, &SDMMaterialSlotLayerItem::OnEffectsToggleButtonClicked)
 		[
 			SNew(SImage)
 			.Image(this, &SDMMaterialSlotLayerItem::GetEffectsToggleButtonImage)
diff --git a/Engine/Plugins/Experimental/DynamicMaterial/Source/DynamicMaterialEditor/Private/UI/Widgets/Editor/SlotEditor/SDMMaterialSlotLayerView.cpp b/Engine/Plugins/Experimental/DynamicMaterial/Source/DynamicMaterialEditor/Private/UI/Widgets/Editor/SlotEditor/SDMMaterialSlotLayerView.cpp
--- a/Engine/Plugins/Experimental/DynamicMaterial/Source/DynamicMaterialEditor/Private/UI/Widgets/Editor/SlotEditor/SDMMaterialSlotLayerView.cpp
+++ b/Engine/Plugins/Experimental/DynamicMaterial/Source/DynamicMaterialEditor/Private/UI/Widgets/Editor/SlotEditor/SDMMaterialSlotLayerView.cpp
@@ -285,8 +285,7 @@ TSharedPtr<SWidget> SDMMaterialSlotLayerView::CreateLayerItemContextMenu()
 		}
 	}
 
-	UToolMenu* ContextMenu = FDMMaterialSlotLayerMenus::GenerateSlotLayerMenu(SlotEditor, LayerObject);
-	return UToolMenus::Get()->GenerateWidget(ContextMenu);
+	return FDMMaterialSlotLayerMenus::GenerateSlotLayerMenu(SlotEditor, LayerObject);
 }
 
 TSharedPtr<SDMMaterialSlotLayerItem> SDMMaterialSlotLayerView::WidgetFromLayerItem(const TSharedPtr<FDMMaterialLayerReference>& InItem) const
diff --git a/Engine/Plugins/Experimental/DynamicMaterial/Source/DynamicMaterialEditor/Private/UI/Widgets/Editor/SlotEditor/SDMMaterialStage.cpp b/Engine/Plugins/Experimental/DynamicMaterial/Source/DynamicMaterialEditor/Private/UI/Widgets/Editor/SlotEditor/SDMMaterialStage.cpp
--- a/Engine/Plugins/Experimental/DynamicMaterial/Source/DynamicMaterialEditor/Private/UI/Widgets/Editor/SlotEditor/SDMMaterialStage.cpp
+++ b/Engine/Plugins/Experimental/DynamicMaterial/Source/DynamicMaterialEditor/Private/UI/Widgets/Editor/SlotEditor/SDMMaterialStage.cpp
@@ -14,7 +14,9 @@
 #include "DynamicMaterialEditorSettings.h"
 #include "DynamicMaterialEditorStyle.h"
 #include "Engine/Texture.h"
+#include "Framework/Application/SlateApplication.h"
 #include "SAssetDropTarget.h"
+#include "UI/Menus/DMMaterialStageMenus.h"
 #include "UI/Widgets/Editor/SDMMaterialComponentEditor.h"
 #include "UI/Widgets/Editor/SlotEditor/SDMMaterialSlotLayerItem.h"
 #include "UI/Widgets/Editor/SlotEditor/SDMMaterialSlotLayerView.h"
@@ -38,6 +40,7 @@ void SDMMaterialStage::Construct(const FArguments& InArgs, const TSharedRef<SDMM
 	StageWeak = InStage;
 
 	SetCanTick(false);
+	SetCursor(EMouseCursor::Default);
 
 	if (!IsValid(InStage))
 	{
@@ -136,49 +139,24 @@ UDMMaterialStage* SDMMaterialStage::GetStage() const
 
 FReply SDMMaterialStage::OnMouseButtonDown(const FGeometry& InMyGeometry, const FPointerEvent& InMouseEvent)
 {
-	UDMMaterialStage* Stage = GetStage();
-
-	if (!Stage)
-	{
-		return FReply::Handled();
-	}
-
-	TSharedPtr<SDMMaterialSlotLayerItem> SlotLayerItem = SlotLayerItemWeak.Pin();
-
-	if (!SlotLayerItem.IsValid())
-	{
-		return FReply::Handled();
-	}
-
-	TSharedPtr<SDMMaterialSlotLayerView> SlotLayerView = SlotLayerItem->GetSlotLayerView();
-
-	if (!SlotLayerView.IsValid())
-	{
-		return FReply::Handled();
-	}
-
-	TSharedPtr<SDMMaterialSlotEditor> SlotEditorWidget = SlotLayerView->GetSlotEditorWidget();
-
-	if (!SlotEditorWidget.IsValid())
+	if (InMouseEvent.GetEffectingButton() == EKeys::LeftMouseButton)
 	{
+		OnMouseButtonDown_Left();
 		return FReply::Handled();
 	}
 
-	TSharedPtr<SDMMaterialEditor> EditorWidget = SlotEditorWidget->GetEditorWidget();
+	return FReply::Unhandled();
+}
 
-	if (!EditorWidget.IsValid())
+FReply SDMMaterialStage::OnMouseButtonUp(const FGeometry& InMyGeometry, const FPointerEvent& InMouseEvent)
+{
+	if (InMouseEvent.GetEffectingButton() == EKeys::RightMouseButton)
 	{
+		OnMouseButtonUp_Right();
 		return FReply::Handled();
 	}
 
-	EditorWidget->EditComponent(Stage);
-
-	if (UDMMaterialLayerObject* Layer = Stage->GetLayer())
-	{
-		SlotLayerView->SetSelectedLayer(Layer);
-	}
-
-	return FReply::Handled();
+	return FReply::Unhandled();
 }
 
 bool SDMMaterialStage::IsStageSelected() const
@@ -394,4 +372,90 @@ void SDMMaterialStage::HandleDrop_Texture(UTexture* InTexture)
 	}
 }
 
+void SDMMaterialStage::OnMouseButtonDown_Left()
+{
+	UDMMaterialStage* Stage = GetStage();
+
+	if (!Stage)
+	{
+		return;
+	}
+
+	TSharedPtr<SDMMaterialSlotLayerItem> SlotLayerItem = SlotLayerItemWeak.Pin();
+
+	if (!SlotLayerItem.IsValid())
+	{
+		return;
+	}
+
+	TSharedPtr<SDMMaterialSlotLayerView> SlotLayerView = SlotLayerItem->GetSlotLayerView();
+
+	if (!SlotLayerView.IsValid())
+	{
+		return;
+	}
+
+	TSharedPtr<SDMMaterialSlotEditor> SlotEditorWidget = SlotLayerView->GetSlotEditorWidget();
+
+	if (!SlotEditorWidget.IsValid())
+	{
+		return;
+	}
+
+	TSharedPtr<SDMMaterialEditor> EditorWidget = SlotEditorWidget->GetEditorWidget();
+
+	if (!EditorWidget.IsValid())
+	{
+		return;
+	}
+
+	EditorWidget->EditComponent(Stage);
+
+	if (UDMMaterialLayerObject* Layer = Stage->GetLayer())
+	{
+		SlotLayerView->SetSelectedLayer(Layer);
+	}
+}
+
+void SDMMaterialStage::OnMouseButtonUp_Right()
+{
+	UDMMaterialStage* Stage = GetStage();
+
+	if (!Stage)
+	{
+		return;
+	}
+
+	TSharedPtr<SDMMaterialSlotLayerItem> SlotLayerItem = SlotLayerItemWeak.Pin();
+
+	if (!SlotLayerItem.IsValid())
+	{
+		return;
+	}
+
+	TSharedPtr<SDMMaterialSlotLayerView> SlotLayerView = SlotLayerItem->GetSlotLayerView();
+
+	if (!SlotLayerView.IsValid())
+	{
+		return;
+	}
+
+	TSharedPtr<SDMMaterialSlotEditor> SlotEditorWidget = SlotLayerView->GetSlotEditorWidget();
+
+	if (!SlotEditorWidget.IsValid())
+	{
+		return;
+	}
+
+	FSlateApplication& SlateApplication = FSlateApplication::Get();
+
+	SlateApplication.PushMenu(
+		SharedThis(this),
+		FWidgetPath(),
+		FDMMaterialStageMenus::GenerateStageMenu(SlotEditorWidget.ToSharedRef(), SharedThis(this)),
+		SlateApplication.GetCursorPos(),
+		FPopupTransitionEffect::ContextMenu
+	);
+}
+
 #undef LOCTEXT_NAMESPACE
diff --git a/Engine/Plugins/Experimental/DynamicMaterial/Source/DynamicMaterialEditor/Private/UI/Widgets/Editor/SlotEditor/SDMMaterialStage.h b/Engine/Plugins/Experimental/DynamicMaterial/Source/DynamicMaterialEditor/Private/UI/Widgets/Editor/SlotEditor/SDMMaterialStage.h
--- a/Engine/Plugins/Experimental/DynamicMaterial/Source/DynamicMaterialEditor/Private/UI/Widgets/Editor/SlotEditor/SDMMaterialStage.h
+++ b/Engine/Plugins/Experimental/DynamicMaterial/Source/DynamicMaterialEditor/Private/UI/Widgets/Editor/SlotEditor/SDMMaterialStage.h
@@ -30,6 +30,7 @@ public:
 
 	//~ Begin SWidget
 	virtual FReply OnMouseButtonDown(const FGeometry& InMyGeometry, const FPointerEvent& InMouseEvent) override;
+	virtual FReply OnMouseButtonUp(const FGeometry& InMyGeometry, const FPointerEvent& InMouseEvent) override;
 	//~ End SWidget
 
 protected:
@@ -46,4 +47,8 @@ protected:
 	void OnAssetsDropped(const FDragDropEvent& InDragDropEvent, TArrayView<FAssetData> InAssets);
 
 	void HandleDrop_Texture(UTexture* InTexture);
+
+	void OnMouseButtonDown_Left();
+
+	void OnMouseButtonUp_Right();
 };
