Material Designer
- Changed the design of the source type selector.

diff --git a/Engine/Plugins/Experimental/DynamicMaterial/Source/DynamicMaterialEditor/Private/UI/PropertyGenerators/DMStagePropertyRowGenerator.cpp b/Engine/Plugins/Experimental/DynamicMaterial/Source/DynamicMaterialEditor/Private/UI/PropertyGenerators/DMStagePropertyRowGenerator.cpp
--- a/Engine/Plugins/Experimental/DynamicMaterial/Source/DynamicMaterialEditor/Private/UI/PropertyGenerators/DMStagePropertyRowGenerator.cpp
+++ b/Engine/Plugins/Experimental/DynamicMaterial/Source/DynamicMaterialEditor/Private/UI/PropertyGenerators/DMStagePropertyRowGenerator.cpp
@@ -3,10 +3,25 @@
 #include "UI/PropertyGenerators/DMStagePropertyRowGenerator.h"
 
 #include "Components/DMMaterialComponent.h"
+#include "Components/DMMaterialLayer.h"
 #include "Components/DMMaterialStage.h"
 #include "Components/DMMaterialStageSource.h"
+#include "DetailLayoutBuilder.h"
+#include "DMEDefs.h"
 #include "DynamicMaterialEditorModule.h"
+#include "Internationalization/Text.h"
+#include "Model/DynamicMaterialModel.h"
+#include "Model/DynamicMaterialModelBase.h"
+#include "UI/Menus/DMMaterialStageSourceMenus.h"
 #include "UI/Widgets/Editor/SDMMaterialComponentEditor.h"
+#include "UI/Widgets/Editor/SDMMaterialSlotEditor.h"
+#include "UI/Widgets/Editor/SlotEditor/SDMMaterialSlotLayerItem.h"
+#include "UI/Widgets/Editor/SlotEditor/SDMMaterialSlotLayerView.h"
+#include "UI/Widgets/Editor/SlotEditor/SDMMaterialStage.h"
+#include "UI/Widgets/SDMMaterialEditor.h"
+#include "Widgets/Input/SComboButton.h"
+
+#define LOCTEXT_NAMESPACE "DMStagePropertyRowGenerator"
 
 const TSharedRef<FDMStagePropertyRowGenerator>& FDMStagePropertyRowGenerator::Get()
 {
@@ -41,6 +56,115 @@ void FDMStagePropertyRowGenerator::AddComponentProperties(const TSharedRef<SDMMa
 		return;
 	}
 
+	if (Stage->CanChangeSource())
+	{
+		bool bIsDynamic = false;
+
+		if (TSharedPtr<SDMMaterialEditor> EditorWidget = InComponentEditorWidget->GetEditorWidget())
+		{
+			if (UDynamicMaterialModelBase* MaterialModelBase = EditorWidget->GetMaterialModelBase())
+			{
+				bIsDynamic = !MaterialModelBase->IsA<UDynamicMaterialModel>();
+			}
+		}
+
+		FDMPropertyHandle& SourceHandle = InOutPropertyRows.AddDefaulted_GetRef();
+		SourceHandle.ValueName = TEXT("SourceType");
+		SourceHandle.NameOverride = LOCTEXT("SourceType", "Source Type");
+		SourceHandle.bEnabled = !bIsDynamic;
+		SourceHandle.ValueWidget = CreateSourceTypeEditWidget(InComponentEditorWidget, Stage);
+		SourceHandle.ResetToDefaultOverride = FResetToDefaultOverride::Hide(true);
+		SourceHandle.Priority = EDMPropertyHandlePriority::High;
+	}
+
 	FDynamicMaterialEditorModule::GeneratorComponentPropertyRows(InComponentEditorWidget, Source, InOutPropertyRows, InOutProcessedObjects);
 	FDMComponentPropertyRowGenerator::AddComponentProperties(InComponentEditorWidget, Stage, InOutPropertyRows, InOutProcessedObjects);
-}
\ No newline at end of file
+}
+
+TSharedRef<SWidget> FDMStagePropertyRowGenerator::CreateSourceTypeEditWidget(const TSharedRef<SDMMaterialComponentEditor>& InComponentEditorWidget, 
+	UDMMaterialStage* InStage)
+{
+	constexpr float Padding = 2.f;
+
+	const FComboButtonStyle& ComboStyle = FAppStyle::Get().GetWidgetStyle<FComboButtonStyle>("ComboButton");
+
+	return SNew(SComboButton)
+		.HasDownArrow(false)
+		.IsFocusable(true)
+		.ContentPadding(FMargin(0.f, Padding, Padding, 0.f))
+		.ToolTipText(LOCTEXT("ChangeStageSourceToolTip", "Change Stage Source"))
+		.OnGetMenuContent_Static(&FDMStagePropertyRowGenerator::MakeSourceTypeEditWidgetMenuContent, InComponentEditorWidget.ToWeakPtr(), InStage)
+		.ButtonContent()
+		[
+			SNew(SHorizontalBox)
+			+ SHorizontalBox::Slot()
+			.AutoWidth()
+			.VAlign(EVerticalAlignment::VAlign_Center)
+			[
+				SNew(STextBlock)
+				.Font(IDetailLayoutBuilder::GetDetailFont())
+				.Text_Static(&FDMStagePropertyRowGenerator::GetSourceTypeEditWidgetText, InStage)
+			]
+			+ SHorizontalBox::Slot()
+			.AutoWidth()
+			.VAlign(EVerticalAlignment::VAlign_Center)
+			.Padding(2.f, 0.f, 0.f, 0.f)
+			[
+				SNew(SImage)
+				.Image(ComboStyle.DownArrowImage.IsSet() ? &ComboStyle.DownArrowImage : nullptr)
+			]
+		];
+}
+
+TSharedRef<SWidget> FDMStagePropertyRowGenerator::MakeSourceTypeEditWidgetMenuContent(TWeakPtr<SDMMaterialComponentEditor> InComponentEditorWidgetWeak, 
+	UDMMaterialStage* InStage)
+{
+	if (IsValid(InStage))
+	{
+		if (UDMMaterialLayerObject* Layer = InStage->GetLayer())
+		{
+			if (UDMMaterialSlot* Slot = Layer->GetSlot())
+			{
+				if (TSharedPtr<SDMMaterialComponentEditor> ComponentEditor = InComponentEditorWidgetWeak.Pin())
+				{
+					if (TSharedPtr<SDMMaterialEditor> EditorWidget = ComponentEditor->GetEditorWidget())
+					{
+						if (TSharedPtr<SDMMaterialSlotEditor> SlotEditorWidget = EditorWidget->GetSlotEditorWidget())
+						{
+							if (SlotEditorWidget->GetSlot() == Slot)
+							{
+								if (TSharedPtr<SDMMaterialSlotLayerView> SlotLayerView = SlotEditorWidget->GetLayerView())
+								{
+									if (TSharedPtr<SDMMaterialSlotLayerItem> SlotLayerItem = SlotLayerView->GetWidgetForLayer(Layer))
+									{
+										if (TSharedPtr<SDMMaterialStage> StageWidget = SlotLayerItem->GetWidgetForStage(InStage))
+										{
+											return FDMMaterialStageSourceMenus::MakeChangeSourceMenu(SlotEditorWidget, StageWidget);
+										}
+									}
+								}
+							}
+						}
+					}
+				}
+			}
+		}
+	}
+
+	return SNullWidget::NullWidget;
+}
+
+FText FDMStagePropertyRowGenerator::GetSourceTypeEditWidgetText(UDMMaterialStage* InStage)
+{
+	if (IsValid(InStage))
+	{
+		if (UDMMaterialStageSource* Source = InStage->GetSource())
+		{
+			return Source->GetStageDescription();
+		}
+	}
+
+	return FText::GetEmpty();
+}
+
+#undef LOCTEXT_NAMESPACE
diff --git a/Engine/Plugins/Experimental/DynamicMaterial/Source/DynamicMaterialEditor/Private/UI/PropertyGenerators/DMStagePropertyRowGenerator.h b/Engine/Plugins/Experimental/DynamicMaterial/Source/DynamicMaterialEditor/Private/UI/PropertyGenerators/DMStagePropertyRowGenerator.h
--- a/Engine/Plugins/Experimental/DynamicMaterial/Source/DynamicMaterialEditor/Private/UI/PropertyGenerators/DMStagePropertyRowGenerator.h
+++ b/Engine/Plugins/Experimental/DynamicMaterial/Source/DynamicMaterialEditor/Private/UI/PropertyGenerators/DMStagePropertyRowGenerator.h
@@ -4,6 +4,10 @@
 
 #include "UI/PropertyGenerators/DMComponentPropertyRowGenerator.h"
 
+class FText;
+class SWidget;
+class UDMMaterialStage;
+
 class FDMStagePropertyRowGenerator : public FDMComponentPropertyRowGenerator
 {
 public:
@@ -15,4 +19,13 @@ public:
 	virtual void AddComponentProperties(const TSharedRef<SDMMaterialComponentEditor>& InComponentEditorWidget, UDMMaterialComponent* InComponent,
 		TArray<FDMPropertyHandle>& InOutPropertyRows, TSet<UDMMaterialComponent*>& InOutProcessedObjects) override;
 	//~ End FDMComponentPropertyRowGenerator
+
+private:
+	static TSharedRef<SWidget> CreateSourceTypeEditWidget(const TSharedRef<SDMMaterialComponentEditor>& InComponentEditorWidget, 
+		UDMMaterialStage* InStage);
+
+	static TSharedRef<SWidget> MakeSourceTypeEditWidgetMenuContent(TWeakPtr<SDMMaterialComponentEditor> InComponentEditorWidgetWeak, 
+		UDMMaterialStage* InStage);
+
+	static FText GetSourceTypeEditWidgetText(UDMMaterialStage* InStage);
 };
diff --git a/Engine/Plugins/Experimental/DynamicMaterial/Source/DynamicMaterialEditor/Private/UI/Widgets/Editor/SDMMaterialComponentEditor.cpp b/Engine/Plugins/Experimental/DynamicMaterial/Source/DynamicMaterialEditor/Private/UI/Widgets/Editor/SDMMaterialComponentEditor.cpp
--- a/Engine/Plugins/Experimental/DynamicMaterial/Source/DynamicMaterialEditor/Private/UI/Widgets/Editor/SDMMaterialComponentEditor.cpp
+++ b/Engine/Plugins/Experimental/DynamicMaterial/Source/DynamicMaterialEditor/Private/UI/Widgets/Editor/SDMMaterialComponentEditor.cpp
@@ -2,21 +2,10 @@
 
 #include "UI/Widgets/Editor/SDMMaterialComponentEditor.h"
 
+#include "DynamicMaterialEditorModule.h"
 #include "Components/DMMaterialComponent.h"
 #include "Components/DMMaterialLayer.h"
-#include "Components/DMMaterialStage.h"
-#include "DetailLayoutBuilder.h"
-#include "DMEDefs.h"
-#include "DynamicMaterialEditorModule.h"
-#include "DynamicMaterialEditorStyle.h"
 #include "DynamicMaterialModule.h"
-#include "Model/DynamicMaterialModel.h"
-#include "Model/DynamicMaterialModelBase.h"
-#include "UI/Menus/DMMaterialStageSourceMenus.h"
-#include "UI/Widgets/Editor/SDMMaterialSlotEditor.h"
-#include "UI/Widgets/Editor/SlotEditor/SDMMaterialSlotLayerItem.h"
-#include "UI/Widgets/Editor/SlotEditor/SDMMaterialSlotLayerView.h"
-#include "UI/Widgets/Editor/SlotEditor/SDMMaterialStage.h"
 #include "UI/Widgets/SDMMaterialEditor.h"
 #include "Widgets/Input/SComboButton.h"
 
@@ -61,78 +50,6 @@ UDMMaterialComponent* SDMMaterialComponentEditor::GetComponent() const
 	return Cast<UDMMaterialComponent>(ObjectWeak.Get());
 }
 
-TSharedRef<SWidget> SDMMaterialComponentEditor::CreateSourceTypeEditWidget()
-{
-	TWeakPtr<SDMMaterialComponentEditor> ThisWeak = SharedThis(this);
-
-	return SNew(SBox)
-		.HAlign(EHorizontalAlignment::HAlign_Left)
-		.VAlign(EVerticalAlignment::VAlign_Center)
-		.Padding(0.f, 0.f, 0.f, 0.2f)
-		[
-			SNew(SComboButton)
-			.HasDownArrow(false)
-			.IsFocusable(true)
-			.ContentPadding(4.0f)
-			.ButtonStyle(FDynamicMaterialEditorStyle::Get(), "HoverHintOnly.Bordered.Dark")
-			.ToolTipText(LOCTEXT("ChangeLayer", "Change Stage Type"))
-			.OnGetMenuContent(this, &SDMMaterialComponentEditor::MakeSourceTypeEditWidgetMenuContent)
-			.ButtonContent()
-			[
-				SNew(STextBlock)
-				.Font(IDetailLayoutBuilder::GetDetailFont())
-				.Text(this, &SDMMaterialComponentEditor::GetSourceTypeEditWidgetText)
-			]
-		];
-}
-
-TSharedRef<SWidget> SDMMaterialComponentEditor::MakeSourceTypeEditWidgetMenuContent()
-{
-	if (UDMMaterialStage* Stage = Cast<UDMMaterialStage>(ObjectWeak.Get()))
-	{
-		if (UDMMaterialLayerObject* Layer = Stage->GetLayer())
-		{
-			if (UDMMaterialSlot* Slot = Layer->GetSlot())
-			{
-				if (TSharedPtr<SDMMaterialEditor> EditorWidget = GetEditorWidget())
-				{
-					if (TSharedPtr<SDMMaterialSlotEditor> SlotEditorWidget = EditorWidget->GetSlotEditorWidget())
-					{
-						if (SlotEditorWidget->GetSlot() == Slot)
-						{
-							if (TSharedPtr<SDMMaterialSlotLayerView> SlotLayerView = SlotEditorWidget->GetLayerView())
-							{
-								if (TSharedPtr<SDMMaterialSlotLayerItem> SlotLayerItem = SlotLayerView->GetWidgetForLayer(Layer))
-								{
-									if (TSharedPtr<SDMMaterialStage> StageWidget = SlotLayerItem->GetWidgetForStage(Stage))
-									{
-										return FDMMaterialStageSourceMenus::MakeChangeSourceMenu(SlotEditorWidget, StageWidget);
-									}
-								}
-							}
-						}
-					}
-				}
-			}
-		}
-	}
-
-	return SNullWidget::NullWidget;
-}
-
-FText SDMMaterialComponentEditor::GetSourceTypeEditWidgetText() const
-{
-	if (UDMMaterialStage* Stage = Cast<UDMMaterialStage>(ObjectWeak.Get()))
-	{
-		if (UDMMaterialStageSource* Source = Stage->GetSource())
-		{
-			return Source->GetStageDescription();
-		}
-	}
-
-	return FText::GetEmpty();
-}
-
 void SDMMaterialComponentEditor::OnComponentUpdated(UDMMaterialComponent* InComponent, EDMUpdateType InUpdateType)
 {
 	if (EnumHasAnyFlags(InUpdateType, EDMUpdateType::Structure))
@@ -149,31 +66,10 @@ TArray<FDMPropertyHandle> SDMMaterialComponentEditor::GetPropertyRows()
 	TArray<FDMPropertyHandle> PropertyRows;
 	TSet<UDMMaterialComponent*> ProcessedObjects;
 
-	bool bIsDynamic = false;
-
-	if (TSharedPtr<SDMMaterialEditor> EditorWidget = GetEditorWidget())
-	{
-		if (UDynamicMaterialModelBase* MaterialModelBase = EditorWidget->GetMaterialModelBase())
-		{
-			bIsDynamic = !MaterialModelBase->IsA<UDynamicMaterialModel>();
-		}
-	}
-
-	if (UDMMaterialStage* Stage = Cast<UDMMaterialStage>(ObjectWeak.Get()))
-	{
-		FDMPropertyHandle& SourceHandle = PropertyRows.AddDefaulted_GetRef();
-		SourceHandle.ValueName = TEXT("SourceType");
-		SourceHandle.NameOverride = LOCTEXT("SourceType", "Source Type");
-		SourceHandle.bEnabled = !bIsDynamic;
-		SourceHandle.ValueWidget = CreateSourceTypeEditWidget();
-		SourceHandle.ResetToDefaultOverride = FResetToDefaultOverride::Hide(true);
-		SourceHandle.Priority = EDMPropertyHandlePriority::High;
-	}
-
 	FDynamicMaterialEditorModule::GeneratorComponentPropertyRows(
-		SharedThis(this), 
-		GetComponent(), 
-		PropertyRows, 
+		SharedThis(this),
+		GetComponent(),
+		PropertyRows,
 		ProcessedObjects
 	);
 
diff --git a/Engine/Plugins/Experimental/DynamicMaterial/Source/DynamicMaterialEditor/Private/UI/Widgets/Editor/SDMMaterialComponentEditor.h b/Engine/Plugins/Experimental/DynamicMaterial/Source/DynamicMaterialEditor/Private/UI/Widgets/Editor/SDMMaterialComponentEditor.h
--- a/Engine/Plugins/Experimental/DynamicMaterial/Source/DynamicMaterialEditor/Private/UI/Widgets/Editor/SDMMaterialComponentEditor.h
+++ b/Engine/Plugins/Experimental/DynamicMaterial/Source/DynamicMaterialEditor/Private/UI/Widgets/Editor/SDMMaterialComponentEditor.h
@@ -23,12 +23,6 @@ public:
 	UDMMaterialComponent* GetComponent() const;
 
 protected:
-	TSharedRef<SWidget> CreateSourceTypeEditWidget();
-
-	TSharedRef<SWidget> MakeSourceTypeEditWidgetMenuContent();
-
-	FText GetSourceTypeEditWidgetText() const;
-
 	void OnComponentUpdated(UDMMaterialComponent* InComponent, EDMUpdateType InUpdateType);
 
 	//~ Begin SDMObjectEditorWidgetBase
