Material Designer
- Layers can now be dragged onto the trash button.
- Effects can now be dragged onto the trash button.

diff --git a/Engine/Plugins/Experimental/DynamicMaterial/Source/DynamicMaterialEditor/Private/UI/DragDrop/DMLayerEffectsDragDropOperation.cpp b/Engine/Plugins/Experimental/DynamicMaterial/Source/DynamicMaterialEditor/Private/UI/DragDrop/DMLayerEffectsDragDropOperation.cpp
--- a/Engine/Plugins/Experimental/DynamicMaterial/Source/DynamicMaterialEditor/Private/UI/DragDrop/DMLayerEffectsDragDropOperation.cpp
+++ b/Engine/Plugins/Experimental/DynamicMaterial/Source/DynamicMaterialEditor/Private/UI/DragDrop/DMLayerEffectsDragDropOperation.cpp
@@ -29,6 +29,16 @@ FDMLayerEffectsDragDropOperation::FDMLayerEffectsDragDropOperation(const TShared
 	Construct();
 }
 
+UDMMaterialEffect* FDMLayerEffectsDragDropOperation::GetMaterialEffect() const
+{
+	if (TSharedPtr<SDMMaterialSlotLayerEffectItem> LayerItemWidget = GetLayerItemWidget())
+	{
+		return LayerItemWidget->GetMaterialEffect();
+	}
+
+	return nullptr;
+}
+
 TSharedPtr<SWidget> FDMLayerEffectsDragDropOperation::GetDefaultDecorator() const
 {
 	static const FLinearColor InvalidLocationColor = FLinearColor(1.0f, 0.0f, 0.0f, 0.5f);
diff --git a/Engine/Plugins/Experimental/DynamicMaterial/Source/DynamicMaterialEditor/Private/UI/DragDrop/DMLayerEffectsDragDropOperation.h b/Engine/Plugins/Experimental/DynamicMaterial/Source/DynamicMaterialEditor/Private/UI/DragDrop/DMLayerEffectsDragDropOperation.h
--- a/Engine/Plugins/Experimental/DynamicMaterial/Source/DynamicMaterialEditor/Private/UI/DragDrop/DMLayerEffectsDragDropOperation.h
+++ b/Engine/Plugins/Experimental/DynamicMaterial/Source/DynamicMaterialEditor/Private/UI/DragDrop/DMLayerEffectsDragDropOperation.h
@@ -11,6 +11,7 @@
 class FWidgetRenderer;
 class SDMMaterialSlotLayerEffectItem;
 class SWidget;
+class UDMMaterialEffect;
 class UDMMaterialStage;
 class UTextureRenderTarget2D;
 struct EVisibility;
@@ -29,6 +30,8 @@ public:
 
 	TSharedPtr<SDMMaterialSlotLayerEffectItem> GetLayerItemWidget() const { return LayerItemWidgetWeak.Pin(); }
 
+	UDMMaterialEffect* GetMaterialEffect() const;
+
 	//~ Begin FDragDropOperation
 	virtual TSharedPtr<SWidget> GetDefaultDecorator() const override;
 	virtual FCursorReply OnCursorQuery() override;
diff --git a/Engine/Plugins/Experimental/DynamicMaterial/Source/DynamicMaterialEditor/Private/UI/Widgets/Editor/SDMMaterialSlotEditor.cpp b/Engine/Plugins/Experimental/DynamicMaterial/Source/DynamicMaterialEditor/Private/UI/Widgets/Editor/SDMMaterialSlotEditor.cpp
--- a/Engine/Plugins/Experimental/DynamicMaterial/Source/DynamicMaterialEditor/Private/UI/Widgets/Editor/SDMMaterialSlotEditor.cpp
+++ b/Engine/Plugins/Experimental/DynamicMaterial/Source/DynamicMaterialEditor/Private/UI/Widgets/Editor/SDMMaterialSlotEditor.cpp
@@ -2,6 +2,8 @@
 
 #include "UI/Widgets/Editor/SDMMaterialSlotEditor.h"
 
+#include "Components/DMMaterialEffect.h"
+#include "Components/DMMaterialEffectStack.h"
 #include "Components/DMMaterialLayer.h"
 #include "Components/DMMaterialSlot.h"
 #include "Components/DMMaterialStage.h"
@@ -34,6 +36,8 @@
 #include "SAssetDropTarget.h"
 #include "ToolMenu.h"
 #include "ToolMenus.h"
+#include "UI/DragDrop/DMLayerEffectsDragDropOperation.h"
+#include "UI/DragDrop/DMSlotLayerDragDropOperation.h"
 #include "UI/Menus/DMMaterialSlotLayerAddEffectMenus.h"
 #include "UI/Menus/DMMaterialSlotLayerMenus.h"
 #include "UI/Utils/DMWidgetStatics.h"
@@ -745,6 +749,8 @@ TSharedRef<SDMMaterialSlotLayerView> SDMMaterialSlotEditor::CreateSlot_LayerView
 
 TSharedRef<SWidget> SDMMaterialSlotEditor::CreateSlot_LayerSettings()
 {
+	TSharedPtr<SDropTarget> DropTarget;
+
 	TSharedRef<SHorizontalBox> NewLayerSettings = SNew(SHorizontalBox)
 		.IsEnabled(!bIsDynamic)
 
@@ -826,16 +832,22 @@ TSharedRef<SWidget> SDMMaterialSlotEditor::CreateSlot_LayerSettings()
 		.VAlign(VAlign_Center)
 		.Padding(5.0f, 2.0f, 0.0f, 2.0f)
 		[
-			SNew(SButton)
-			.ContentPadding(4.0f)
-			.ButtonStyle(FDynamicMaterialEditorStyle::Get(), "HoverHintOnly.Bordered.Dark")
-			.ToolTipText(LOCTEXT("RemoveLayerTooltip", "Remove Selected Layer\n\nThe last layer cannot be removed."))
-			.IsEnabled(this, &SDMMaterialSlotEditor::GetLayerRowsButtonsCanRemove)
-			.OnClicked(this, &SDMMaterialSlotEditor::OnLayerRowButtonsRemoveClicked)
+			SAssignNew(DropTarget, SDropTarget)
+			.OnIsRecognized(this, &SDMMaterialSlotEditor::IsValidLayerDropForDelete)
+			.OnAllowDrop(this, &SDMMaterialSlotEditor::CanDropLayerForDelete)
+			.OnDropped(this, &SDMMaterialSlotEditor::OnLayerDroppedForDelete)
 			[
-				SNew(SImage)
-				.Image(FAppStyle::Get().GetBrush("Icons.Delete"))
-				.DesiredSizeOverride(FVector2D(16.0f))
+				SNew(SButton)
+				.ContentPadding(4.0f)
+				.ButtonStyle(FDynamicMaterialEditorStyle::Get(), "HoverHintOnly.Bordered.Dark")
+				.ToolTipText(LOCTEXT("RemoveLayerTooltip", "Remove Selected Layer\n\nThe last layer cannot be removed."))
+				.IsEnabled(this, &SDMMaterialSlotEditor::GetLayerRowsButtonsCanRemove)
+				.OnClicked(this, &SDMMaterialSlotEditor::OnLayerRowButtonsRemoveClicked)
+				[
+					SNew(SImage)
+					.Image(FAppStyle::Get().GetBrush("Icons.Delete"))
+					.DesiredSizeOverride(FVector2D(16.0f))
+				]
 			]
 		];
 
@@ -1199,4 +1211,57 @@ void SDMMaterialSlotEditor::HandleDrop_MaterialFunction(UMaterialFunctionInterfa
 	}
 }
 
-#undef LOCTEXT_NAMESPACE
+bool SDMMaterialSlotEditor::IsValidLayerDropForDelete(TSharedPtr<FDragDropOperation> InDragDropOperation)
+{
+	return InDragDropOperation.IsValid()
+		&& (InDragDropOperation->IsOfType<FDMSlotLayerDragDropOperation>()
+			|| InDragDropOperation->IsOfType<FDMLayerEffectsDragDropOperation>());
+}
+
+bool SDMMaterialSlotEditor::CanDropLayerForDelete(TSharedPtr<FDragDropOperation> InDragDropOperation)
+{
+	if (InDragDropOperation->IsOfType<FDMSlotLayerDragDropOperation>())
+	{
+		if (UDMMaterialLayerObject* Layer = StaticCastSharedPtr<FDMSlotLayerDragDropOperation>(InDragDropOperation)->GetLayer())
+		{
+			if (UDMMaterialSlot* Slot = Layer->GetSlot())
+			{
+				return Slot->CanRemoveLayer(Layer);
+			}
+		}
+	}
+	else if (InDragDropOperation->IsOfType<FDMLayerEffectsDragDropOperation>())
+	{
+		return IsValid(StaticCastSharedPtr<FDMLayerEffectsDragDropOperation>(InDragDropOperation)->GetMaterialEffect());
+	}
+
+	return false;	
+}
+
+FReply SDMMaterialSlotEditor::OnLayerDroppedForDelete(const FGeometry& InGeometry, const FDragDropEvent& InDragDropEvent)
+{
+	if (TSharedPtr<FDMSlotLayerDragDropOperation> LayerDragDropOperation = InDragDropEvent.GetOperationAs<FDMSlotLayerDragDropOperation>())
+	{
+		if (UDMMaterialLayerObject* Layer = LayerDragDropOperation->GetLayer())
+		{
+			if (UDMMaterialSlot* Slot = Layer->GetSlot())
+			{
+				Slot->RemoveLayer(Layer);
+			}
+		}
+	}
+	else if (TSharedPtr<FDMLayerEffectsDragDropOperation> EffectDragDropOperation = InDragDropEvent.GetOperationAs<FDMLayerEffectsDragDropOperation>())
+	{
+		if (UDMMaterialEffect* Effect = EffectDragDropOperation->GetMaterialEffect())
+		{
+			if (UDMMaterialEffectStack* EffectStack = Effect->GetEffectStack())
+			{
+				EffectStack->RemoveEffect(Effect);
+			}
+		}
+	}
+
+	return FReply::Handled();
+}
+
+#undef LOCTEXT_NAMESPACE
\ No newline at end of file
diff --git a/Engine/Plugins/Experimental/DynamicMaterial/Source/DynamicMaterialEditor/Private/UI/Widgets/Editor/SDMMaterialSlotEditor.h b/Engine/Plugins/Experimental/DynamicMaterial/Source/DynamicMaterialEditor/Private/UI/Widgets/Editor/SDMMaterialSlotEditor.h
--- a/Engine/Plugins/Experimental/DynamicMaterial/Source/DynamicMaterialEditor/Private/UI/Widgets/Editor/SDMMaterialSlotEditor.h
+++ b/Engine/Plugins/Experimental/DynamicMaterial/Source/DynamicMaterialEditor/Private/UI/Widgets/Editor/SDMMaterialSlotEditor.h
@@ -124,4 +124,10 @@ protected:
 	void HandleDrop_TextureSet(UDMTextureSet* InTextureSet);
 
 	void HandleDrop_MaterialFunction(UMaterialFunctionInterface* InMaterialFunction);
+
+	bool IsValidLayerDropForDelete(TSharedPtr<FDragDropOperation> InDragDropOperation);
+
+	bool CanDropLayerForDelete(TSharedPtr<FDragDropOperation> InDragDropOperation);
+
+	FReply OnLayerDroppedForDelete(const FGeometry& InGeometry, const FDragDropEvent& InDragDropEvent);
 };
