Material Designer
- Layout options

diff --git a/Engine/Plugins/Experimental/DynamicMaterial/Source/DynamicMaterial/Public/DMDefs.h b/Engine/Plugins/Experimental/DynamicMaterial/Source/DynamicMaterial/Public/DMDefs.h
--- a/Engine/Plugins/Experimental/DynamicMaterial/Source/DynamicMaterial/Public/DMDefs.h
+++ b/Engine/Plugins/Experimental/DynamicMaterial/Source/DynamicMaterial/Public/DMDefs.h
@@ -16,27 +16,27 @@ UENUM(BlueprintType)
 enum class EDMMaterialPropertyType : uint8
 {
 	None = 0,
-	BaseColor,
-	EmissiveColor UMETA(ShortName = "Emissive"),
+	BaseColor           UMETA(ShortName = "Base"),
+	EmissiveColor       UMETA(ShortName = "Emis"),
 	Opacity,
-	OpacityMask UMETA(ShortName = "Mask"),
-	Roughness,
-	Specular,
-	Metallic,
-	Normal,
-	PixelDepthOffset UMETA(ShortName = "PDO"),
+	OpacityMask         UMETA(ShortName = "Mask"),
+	Roughness           UMETA(ShortName = "Rough"),
+	Specular            UMETA(ShortName = "Spec"),
+	Metallic            UMETA(ShortName = "Metal"),
+	Normal              UMETA(ShortName = "Norm"),
+	PixelDepthOffset    UMETA(ShortName = "PDO"),
 	WorldPositionOffset UMETA(ShortName = "WPO"),
-	AmbientOcclusion UMETA(ShortName = "AO"),
-	Anisotropy,
-	Refraction,
-	Tangent,
+	AmbientOcclusion    UMETA(ShortName = "AO"),
+	Anisotropy          UMETA(ShortName = "Aniso"),
+	Refraction          UMETA(ShortName = "Refra"),
+	Tangent             UMETA(ShortName = "Tang"),
 	Custom1,
 	Custom2,
 	Custom3,
 	Custom4,
-	Displacement,
-	SubsurfaceColor,
-	SurfaceThickness,
+	Displacement        UMETA(ShortName = "Disp"),
+	SubsurfaceColor     UMETA(ShortName = "Surf"),
+	SurfaceThickness    UMETA(ShortName = "Thick"),
 	Any
 };
 
diff --git a/Engine/Plugins/Experimental/DynamicMaterial/Source/DynamicMaterialEditor/Private/Components/DMMaterialSlot.cpp b/Engine/Plugins/Experimental/DynamicMaterial/Source/DynamicMaterialEditor/Private/Components/DMMaterialSlot.cpp
--- a/Engine/Plugins/Experimental/DynamicMaterial/Source/DynamicMaterialEditor/Private/Components/DMMaterialSlot.cpp
+++ b/Engine/Plugins/Experimental/DynamicMaterial/Source/DynamicMaterialEditor/Private/Components/DMMaterialSlot.cpp
@@ -23,6 +23,7 @@
 #include "Model/DynamicMaterialModelEditorOnlyData.h"
 #include "UObject/Package.h"
 #include "UObject/UObjectGlobals.h"
+#include "Utils/DMPrivate.h"
 #include "Utils/DMUtils.h"
 
 #define LOCTEXT_NAMESPACE "DMMaterialStage"
@@ -1382,16 +1383,13 @@ FString UDMMaterialSlot::GetComponentPathComponent() const
 
 		if (SlotProperties.Num() == 1)
 		{
-			UEnum* MaterialPropertyEnum = StaticEnum<EDMMaterialPropertyType>();
-			constexpr const TCHAR* ShortNameName = TEXT("ShortName");
-			const FString ShortName = MaterialPropertyEnum->GetMetaData(ShortNameName, MaterialPropertyEnum->GetIndexByValue(static_cast<int64>(SlotProperties[0])));
-			const FString Token = !ShortName.IsEmpty() ? ShortName : MaterialPropertyEnum->GetNameStringByValue(static_cast<int64>(SlotProperties[0]));
+			const FText Token = UE::DynamicMaterialEditor::Private::GetMaterialPropertyShortDisplayName(SlotProperties[0]);
 
 			return FString::Printf(
 				TEXT("%s%hc%s%hc"),
 				*UDynamicMaterialModelEditorOnlyData::SlotsPathToken,
 				FDMComponentPath::ParameterOpen,
-				*Token,
+				*Token.ToString(),
 				FDMComponentPath::ParameterClose
 			);
 		}
diff --git a/Engine/Plugins/Experimental/DynamicMaterial/Source/DynamicMaterialEditor/Private/DynamicMaterialEditorSettings.cpp b/Engine/Plugins/Experimental/DynamicMaterial/Source/DynamicMaterialEditor/Private/DynamicMaterialEditorSettings.cpp
--- a/Engine/Plugins/Experimental/DynamicMaterial/Source/DynamicMaterialEditor/Private/DynamicMaterialEditorSettings.cpp
+++ b/Engine/Plugins/Experimental/DynamicMaterial/Source/DynamicMaterialEditor/Private/DynamicMaterialEditorSettings.cpp
@@ -267,7 +267,8 @@ void UDynamicMaterialEditorSettings::PostEditChangeProperty(FPropertyChangedEven
 
 	const FName PropertyName = InPropertyChangedEvent.GetMemberPropertyName();
 
-	if (PropertyName == GET_MEMBER_NAME_CHECKED(UDynamicMaterialEditorSettings, SplitterLocation)
+	if (PropertyName == GET_MEMBER_NAME_CHECKED(UDynamicMaterialEditorSettings, Layout)
+		|| PropertyName == GET_MEMBER_NAME_CHECKED(UDynamicMaterialEditorSettings, SplitterLocation)
 		|| PropertyName == GET_MEMBER_NAME_CHECKED(UDynamicMaterialEditorSettings, bPreviewImagesUseTextureUVs)
 		|| PropertyName == GET_MEMBER_NAME_CHECKED(UDynamicMaterialEditorSettings, PreviewMesh)
 		|| PropertyName == GET_MEMBER_NAME_CHECKED(UDynamicMaterialEditorSettings, bShowPreviewBackground))
@@ -297,6 +298,7 @@ void UDynamicMaterialEditorSettings::OpenEditorSettingsWindow() const
 
 void UDynamicMaterialEditorSettings::ResetAllLayoutSettings()
 {
+	Layout = EDMMaterialEditorLayout::Left;
 	SplitterLocation = 240;
 	ThumbnailSize = 256.f;
 	PreviewMesh = EDMMaterialPreviewMesh::Plane;
diff --git a/Engine/Plugins/Experimental/DynamicMaterial/Source/DynamicMaterialEditor/Private/DynamicMaterialEditorSettings.h b/Engine/Plugins/Experimental/DynamicMaterial/Source/DynamicMaterialEditor/Private/DynamicMaterialEditorSettings.h
--- a/Engine/Plugins/Experimental/DynamicMaterial/Source/DynamicMaterialEditor/Private/DynamicMaterialEditorSettings.h
+++ b/Engine/Plugins/Experimental/DynamicMaterial/Source/DynamicMaterialEditor/Private/DynamicMaterialEditorSettings.h
@@ -137,6 +137,22 @@ enum class EDMMaterialPreviewMesh : uint8
 	ShaderBall
 };
 
+UENUM(BlueprintType)
+enum class EDMMaterialEditorLayout : uint8
+{
+	Left,
+	LeftAutoHide,
+	LeftSlim,
+	TopVertical,
+	TopVerticalAutoHide,
+	TopHorizontal,
+	TopHorizontalAutoHide,
+	TopSlim,
+
+	First = Left,
+	Last = TopSlim
+};
+
 /**
  * Material Designer Settings
  */
@@ -155,6 +171,9 @@ public:
 	UPROPERTY(Config, EditAnywhere, BlueprintReadWrite, Category = "Editor")
 	bool bFollowSelection;
 
+	UPROPERTY(Config, EditAnywhere, BlueprintReadWrite, Category = "Layout")
+	EDMMaterialEditorLayout Layout;
+
 	/** Adjusts the vertical size of the material layer view. */
 	UPROPERTY(Config, EditAnywhere, BlueprintReadWrite, Category = "Layout", meta = (
 		ClampMin = "0.05", UIMin = "0.05", ClampMax = "0.95", UIMax = "0.95"))
diff --git a/Engine/Plugins/Experimental/DynamicMaterial/Source/DynamicMaterialEditor/Private/UI/Menus/DMToolBarMenus.cpp b/Engine/Plugins/Experimental/DynamicMaterial/Source/DynamicMaterialEditor/Private/UI/Menus/DMToolBarMenus.cpp
--- a/Engine/Plugins/Experimental/DynamicMaterial/Source/DynamicMaterialEditor/Private/UI/Menus/DMToolBarMenus.cpp
+++ b/Engine/Plugins/Experimental/DynamicMaterial/Source/DynamicMaterialEditor/Private/UI/Menus/DMToolBarMenus.cpp
@@ -55,7 +55,7 @@ TSharedRef<SWidget> FDMToolBarMenus::MakeEditorLayoutMenu(const TSharedPtr<SDMMa
 
 		NewToolMenu->AddDynamicSection(
 			TEXT("MaterialDesignerSettings"),
-			FNewToolMenuDelegate::CreateStatic(&AddToolBarEditorLayoutMenu)
+			FNewToolMenuDelegate::CreateStatic(&AddMenu)
 		);
 	}
 
@@ -68,7 +68,13 @@ TSharedRef<SWidget> FDMToolBarMenus::MakeEditorLayoutMenu(const TSharedPtr<SDMMa
 	return ToolMenus->GenerateWidget(ToolBarEditorLayoutMenuName, MenuContext);
 }
 
-void FDMToolBarMenus::AddToolBarExportMenu(UToolMenu* InMenu)
+void FDMToolBarMenus::AddMenu(UToolMenu* InMenu)
+{
+	AddExportMenu(InMenu);
+	AddSettingsMenu(InMenu);
+}
+
+void FDMToolBarMenus::AddExportMenu(UToolMenu* InMenu)
 {
 	using namespace UE::DynamicMaterialEditor::Private;
 
@@ -157,23 +163,7 @@ void FDMToolBarMenus::AddToolBarExportMenu(UToolMenu* InMenu)
 	);
 }
 
-void FDMToolBarMenus::AddToolBarAdvancedSection(UToolMenu* InMenu)
-{
-	FToolMenuSection& NewSection = InMenu->AddSection(TEXT("AdvancedSettings"), LOCTEXT("AdvancedSettingsSection", "Advanced Settings"));
-
-	NewSection.AddMenuEntry(
-		NAME_None,
-		LOCTEXT("ResetAllSettingsToDefaults", "Reset All To Defaults"),
-		LOCTEXT("ResetAllSettingsToDefaultsTooltip", "Resets all the Material Designer settings to their default values."),
-		FSlateIcon(),
-		FUIAction(FExecuteAction::CreateUObject(
-			UDynamicMaterialEditorSettings::Get(), 
-			&UDynamicMaterialEditorSettings::ResetAllLayoutSettings)
-		)
-	);
-}
-
-void FDMToolBarMenus::AddToolBarSettingsMenu(UToolMenu* InMenu)
+void FDMToolBarMenus::AddSettingsMenu(UToolMenu* InMenu)
 {
 	using namespace UE::DynamicMaterialEditor::Private;
 
@@ -188,7 +178,14 @@ void FDMToolBarMenus::AddToolBarSettingsMenu(UToolMenu* InMenu)
 		"AdvancedSettings",
 		LOCTEXT("AdvancedSettingsSubMenu", "Advanced Settings"),
 		LOCTEXT("AdvancedSettingsSubMenu_ToolTip", "Display advanced Material Designer settings"),
-		FNewToolMenuDelegate::CreateStatic(&AddToolBarAdvancedSection)
+		FNewToolMenuDelegate::CreateStatic(&AddAdvancedSection)
+	);
+
+	NewSection.AddSubMenu(
+		"EditorLayout",
+		LOCTEXT("EditorLayoutSubMenu", "Editor Layout"),
+		LOCTEXT("EditorLayoutSubMenu_ToolTip", "Change the layout of the Material Designer Editor"),
+		FNewToolMenuDelegate::CreateStatic(&AddEditorLayoutSection)
 	);
 
 	NewSection.AddMenuEntry(
@@ -203,10 +200,93 @@ void FDMToolBarMenus::AddToolBarSettingsMenu(UToolMenu* InMenu)
 	);
 }
 
-void FDMToolBarMenus::AddToolBarEditorLayoutMenu(UToolMenu* InMenu)
+void FDMToolBarMenus::AddEditorLayoutSection(UToolMenu* InMenu)
 {
-	AddToolBarExportMenu(InMenu);
-	AddToolBarSettingsMenu(InMenu);
+	FToolMenuSection& NewSection = InMenu->AddSection(TEXT("EditorLayout"), LOCTEXT("EditorLayoutSection", "EditorLayout"));
+
+	UEnum* LayoutEnum = StaticEnum<EDMMaterialEditorLayout>();
+
+	for (EDMMaterialEditorLayout Layout = EDMMaterialEditorLayout::First;
+		Layout <= EDMMaterialEditorLayout::Last;
+		Layout = static_cast<EDMMaterialEditorLayout>(static_cast<uint8>(Layout)+1))
+	{
+		FUIAction Action;
+
+		switch (Layout)
+		{
+			case EDMMaterialEditorLayout::LeftAutoHide:
+			case EDMMaterialEditorLayout::TopHorizontalAutoHide:
+			case EDMMaterialEditorLayout::TopVerticalAutoHide:
+				Action.CanExecuteAction = FCanExecuteAction::CreateLambda([]() { return false; });
+				break;
+
+			default:
+				break;
+		}
+
+		Action.GetActionCheckState = FGetActionCheckState::CreateLambda(
+			[Layout]()
+			{
+				if (UDynamicMaterialEditorSettings* Settings = GetMutableDefault<UDynamicMaterialEditorSettings>())
+				{
+					if (Settings->Layout == Layout)
+					{
+						return ECheckBoxState::Checked;
+					}
+				}
+
+				return ECheckBoxState::Unchecked;
+			}
+		);
+
+		Action.ExecuteAction = FExecuteAction::CreateLambda(
+			[Layout]()
+			{
+				UDynamicMaterialEditorSettings* Settings = GetMutableDefault<UDynamicMaterialEditorSettings>();
+
+				if (!Settings || Settings->Layout == Layout)
+				{
+					return;
+				}
+
+				Settings->Layout = Layout;
+				TArray<UObject*> TopLeftObjects = {Settings};
+
+				FPropertyChangedEvent PropertyChangedEvent(
+					UDynamicMaterialEditorSettings::StaticClass()->FindPropertyByName(GET_MEMBER_NAME_CHECKED(UDynamicMaterialEditorSettings, Layout)),
+					EPropertyChangeType::Interactive,
+					TopLeftObjects
+				);
+
+				Settings->PostEditChangeProperty(PropertyChangedEvent);
+			}
+		);
+
+		NewSection.AddMenuEntry(
+			NAME_None,
+			LayoutEnum->GetDisplayNameTextByValue(static_cast<int64>(Layout)),
+			FText::GetEmpty(),
+			FSlateIcon(),
+			Action,
+			EUserInterfaceActionType::RadioButton
+		);
+	}
+}
+
+void FDMToolBarMenus::AddAdvancedSection(UToolMenu* InMenu)
+{
+	FToolMenuSection& NewSection = InMenu->AddSection(TEXT("AdvancedSettings"), LOCTEXT("AdvancedSettingsSection", "Advanced Settings"));
+
+	NewSection.AddMenuEntry(
+		NAME_None,
+		LOCTEXT("ResetAllSettingsToDefaults", "Reset All To Defaults"),
+		LOCTEXT("ResetAllSettingsToDefaultsTooltip", "Resets all the Material Designer settings to their default values."),
+		FSlateIcon(),
+		FUIAction(FExecuteAction::CreateUObject(
+			UDynamicMaterialEditorSettings::Get(), 
+			&UDynamicMaterialEditorSettings::ResetAllLayoutSettings)
+		)
+	);
 }
 
 void FDMToolBarMenus::OpenMaterialEditorFromContext(UDMMenuContext* InMenuContext)
@@ -357,7 +437,7 @@ void FDMToolBarMenus::SnapshotMaterial(TWeakObjectPtr<UDynamicMaterialModelBase>
 	}
 }
 
-void FDMToolBarMenus::AddToolBarBoolOptionMenuEntry(FToolMenuSection& InSection, const FName& InPropertyName, const FUIAction InAction)
+void FDMToolBarMenus::AddBoolOptionMenuEntry(FToolMenuSection& InSection, const FName& InPropertyName, const FUIAction InAction)
 {
 	const FProperty* const OptionProperty = UDynamicMaterialEditorSettings::StaticClass()->FindPropertyByName(InPropertyName);
 
@@ -373,7 +453,7 @@ void FDMToolBarMenus::AddToolBarBoolOptionMenuEntry(FToolMenuSection& InSection,
 	}
 }
 
-void FDMToolBarMenus::AddToolBarIntOptionMenuEntry(FToolMenuSection& InSection, const FName& InPropertyName,
+void FDMToolBarMenus::AddIntOptionMenuEntry(FToolMenuSection& InSection, const FName& InPropertyName,
 	TAttribute<bool> InIsEnabledAttribute, TAttribute<EVisibility> InVisibilityAttribute)
 {
 	InSection.AddDynamicEntry(NAME_None,
diff --git a/Engine/Plugins/Experimental/DynamicMaterial/Source/DynamicMaterialEditor/Private/UI/Menus/DMToolBarMenus.h b/Engine/Plugins/Experimental/DynamicMaterial/Source/DynamicMaterialEditor/Private/UI/Menus/DMToolBarMenus.h
--- a/Engine/Plugins/Experimental/DynamicMaterial/Source/DynamicMaterialEditor/Private/UI/Menus/DMToolBarMenus.h
+++ b/Engine/Plugins/Experimental/DynamicMaterial/Source/DynamicMaterialEditor/Private/UI/Menus/DMToolBarMenus.h
@@ -14,6 +14,7 @@ class UDMMenuContext;
 class UDynamicMaterialInstance;
 class UDynamicMaterialModelBase;
 class UToolMenu;
+enum class EDMMaterialEditorLayout : uint8;
 struct FToolMenuSection;
 struct FUIAction;
 
@@ -23,17 +24,19 @@ public:
 	static TSharedRef<SWidget> MakeEditorLayoutMenu(const TSharedPtr<SDMMaterialEditor>& InEditorWidget = nullptr);
 
 private:
-	static void AddToolBarExportMenu(UToolMenu* InMenu);
+	static void AddMenu(UToolMenu* InMenu);
 
-	static void AddToolBarAdvancedSection(UToolMenu* InMenu);
+	static void AddExportMenu(UToolMenu* InMenu);
 
-	static void AddToolBarSettingsMenu(UToolMenu* InMenu);
+	static void AddSettingsMenu(UToolMenu* InMenu);
 
-	static void AddToolBarEditorLayoutMenu(UToolMenu* InMenu);
+	static void AddEditorLayoutSection(UToolMenu* InMenu);
 
-	static void AddToolBarBoolOptionMenuEntry(FToolMenuSection& InSection, const FName& InPropertyName, const FUIAction InAction);
+	static void AddAdvancedSection(UToolMenu* InMenu);
 
-	static void AddToolBarIntOptionMenuEntry(FToolMenuSection& InSection, const FName& InPropertyName,
+	static void AddBoolOptionMenuEntry(FToolMenuSection& InSection, const FName& InPropertyName, const FUIAction InAction);
+
+	static void AddIntOptionMenuEntry(FToolMenuSection& InSection, const FName& InPropertyName,
 		TAttribute<bool> InIsEnabledAttribute = TAttribute<bool>(),
 		TAttribute<EVisibility> InVisibilityAttribute = TAttribute<EVisibility>());
 
diff --git a/Engine/Plugins/Experimental/DynamicMaterial/Source/DynamicMaterialEditor/Private/UI/Widgets/Editor/EditorLayouts/SDMMaterialEditor_Left.cpp b/Engine/Plugins/Experimental/DynamicMaterial/Source/DynamicMaterialEditor/Private/UI/Widgets/Editor/EditorLayouts/SDMMaterialEditor_Left.cpp
new file mode 100644
--- /dev/null
+++ b/Engine/Plugins/Experimental/DynamicMaterial/Source/DynamicMaterialEditor/Private/UI/Widgets/Editor/EditorLayouts/SDMMaterialEditor_Left.cpp
@@ -0,0 +1,63 @@
+// Copyright Epic Games, Inc. All Rights Reserved.
+
+#include "UI/Widgets/Editor/EditorLayouts/SDMMaterialEditor_Left.h"
+
+#include "Styling/AppStyle.h"
+#include "UI/Widgets/Editor/PropertySelectorLayouts/SDMMaterialPropertySelector_Vertical.h"
+#include "UI/Widgets/Editor/SDMMaterialPreview.h"
+#include "UI/Widgets/Editor/SDMMaterialPropertySelector.h"
+#include "UI/Widgets/Editor/SDMMaterialSlotEditor.h"
+#include "Widgets/Layout/SBorder.h"
+#include "Widgets/Layout/SSplitter.h"
+#include "Widgets/SBoxPanel.h"
+#include "Widgets/SNullWidget.h"
+
+void SDMMaterialEditor_Left::Construct(const FArguments& InArgs, const TSharedRef<SDMMaterialDesigner>& InDesignerWidget)
+{
+	SDMMaterialEditor_LeftBase::Construct(
+		SDMMaterialEditor_LeftBase::FArguments()
+		.MaterialModelBase(InArgs._MaterialModelBase)
+		.MaterialProperty(InArgs._MaterialProperty),
+		InDesignerWidget
+	);
+}
+
+TSharedRef<SWidget> SDMMaterialEditor_Left::CreateSlot_Left()
+{
+	using namespace UE::DynamicMaterialEditor::Private;
+
+	SVerticalBox::FSlot* MaterialPreviewSlotPtr = nullptr;
+	SVerticalBox::FSlot* PropertySelectorSlotPtr = nullptr;
+
+	TSharedRef<SWidget> NewLeft = SNew(SBorder)
+		.BorderImage(FAppStyle::GetBrush(EditorDarkBackground))
+		.Padding(5.f)
+		[
+			SNew(SVerticalBox)
+			+ SVerticalBox::Slot()
+			.Expose(MaterialPreviewSlotPtr)
+			.AutoHeight()
+			.Padding(0.f)
+			[
+				SNullWidget::NullWidget
+			]
+
+			+ SVerticalBox::Slot()
+			.Expose(PropertySelectorSlotPtr)
+			.FillHeight(1.0f)
+			.Padding(0.f, 5.f, 0.f, 0.f)
+			[
+				SNullWidget::NullWidget
+			]
+		];
+
+	MaterialPreviewSlot = TDMWidgetSlot<SDMMaterialPreview>(MaterialPreviewSlotPtr, CreateSlot_Preview());
+	PropertySelectorSlot = TDMWidgetSlot<SDMMaterialPropertySelector>(PropertySelectorSlotPtr, CreateSlot_PropertySelector());
+
+	return NewLeft;
+}
+
+TSharedRef<SDMMaterialPropertySelector> SDMMaterialEditor_Left::CreateSlot_PropertySelector_Impl()
+{
+	return SNew(SDMMaterialPropertySelector_Vertical, SharedThis(this));
+}
diff --git a/Engine/Plugins/Experimental/DynamicMaterial/Source/DynamicMaterialEditor/Private/UI/Widgets/Editor/EditorLayouts/SDMMaterialEditor_Left.h b/Engine/Plugins/Experimental/DynamicMaterial/Source/DynamicMaterialEditor/Private/UI/Widgets/Editor/EditorLayouts/SDMMaterialEditor_Left.h
new file mode 100644
--- /dev/null
+++ b/Engine/Plugins/Experimental/DynamicMaterial/Source/DynamicMaterialEditor/Private/UI/Widgets/Editor/EditorLayouts/SDMMaterialEditor_Left.h
@@ -0,0 +1,30 @@
+// Copyright Epic Games, Inc. All Rights Reserved.
+
+#pragma once
+
+#include "UI/Widgets/Editor/EditorLayouts/SDMMaterialEditor_LeftBase.h"
+
+class SDMMaterialEditor_Left : public SDMMaterialEditor_LeftBase
+{
+	SLATE_BEGIN_ARGS(SDMMaterialEditor_Left)
+		: _MaterialModelBase(nullptr)
+		, _MaterialProperty(TOptional<FDMObjectMaterialProperty>())
+		{}
+		SLATE_ARGUMENT(UDynamicMaterialModelBase*, MaterialModelBase)
+		SLATE_ARGUMENT(TOptional<FDMObjectMaterialProperty>, MaterialProperty)
+	SLATE_END_ARGS()
+
+public:
+	virtual ~SDMMaterialEditor_Left() override = default;
+
+	void Construct(const FArguments& InArgs, const TSharedRef<SDMMaterialDesigner>& InDesignerWidget);
+
+protected:
+	//~ Begin SDMMaterialEditor_LeftBase
+	virtual TSharedRef<SWidget> CreateSlot_Left() override;
+	//~ End SDMMaterialEditor_LeftBase
+
+	//~ Begin SDMMaterialEditor
+	virtual TSharedRef<SDMMaterialPropertySelector> CreateSlot_PropertySelector_Impl() override;
+	//~ End SDMMaterialEditor
+};
diff --git a/Engine/Plugins/Experimental/DynamicMaterial/Source/DynamicMaterialEditor/Private/UI/Widgets/Editor/EditorLayouts/SDMMaterialEditor_LeftBase.cpp b/Engine/Plugins/Experimental/DynamicMaterial/Source/DynamicMaterialEditor/Private/UI/Widgets/Editor/EditorLayouts/SDMMaterialEditor_LeftBase.cpp
new file mode 100644
--- /dev/null
+++ b/Engine/Plugins/Experimental/DynamicMaterial/Source/DynamicMaterialEditor/Private/UI/Widgets/Editor/EditorLayouts/SDMMaterialEditor_LeftBase.cpp
@@ -0,0 +1,254 @@
+// Copyright Epic Games, Inc. All Rights Reserved.
+
+#include "UI/Widgets/Editor/EditorLayouts/SDMMaterialEditor_LeftBase.h"
+
+#include "DynamicMaterialEditorSettings.h"
+#include "Styling/AppStyle.h"
+#include "UI/Widgets/Editor/SDMMaterialComponentEditor.h"
+#include "UI/Widgets/Editor/SDMMaterialGlobalSettingsEditor.h"
+#include "UI/Widgets/Editor/SDMMaterialPreview.h"
+#include "UI/Widgets/Editor/SDMMaterialPropertySelector.h"
+#include "UI/Widgets/Editor/SDMMaterialSlotEditor.h"
+#include "Widgets/Layout/SBorder.h"
+#include "Widgets/Layout/SScrollBox.h"
+#include "Widgets/Layout/SSplitter.h"
+#include "Widgets/SBoxPanel.h"
+
+void SDMMaterialEditor_LeftBase::Construct(const FArguments& InArgs, const TSharedRef<SDMMaterialDesigner>& InDesignerWidget)
+{
+	SDMMaterialEditor::Construct(
+		SDMMaterialEditor::FArguments()
+		.MaterialModelBase(InArgs._MaterialModelBase)
+		.MaterialProperty(InArgs._MaterialProperty),
+		InDesignerWidget
+	);
+}
+
+void SDMMaterialEditor_LeftBase::EditSlot(UDMMaterialSlot* InSlot, bool bInForceRefresh)
+{
+	if (!bInForceRefresh && SlotEditorSlot.IsValid() && SlotEditorSlot->GetSlot() == InSlot)
+	{
+		return;
+	}
+
+	RightSlot.Invalidate();
+
+	SDMMaterialEditor::EditSlot(InSlot, bInForceRefresh);
+}
+
+void SDMMaterialEditor_LeftBase::EditComponent(UDMMaterialComponent* InComponent, bool bInForceRefresh)
+{
+	if (!bInForceRefresh && ComponentEditorSlot.IsValid() && ComponentEditorSlot->GetComponent() == InComponent)
+	{
+		return;
+	}
+
+	if (bGlobalSettingsMode)
+	{
+		RightSlot.Invalidate();
+	}
+
+	SDMMaterialEditor::EditComponent(InComponent, bInForceRefresh);
+}
+
+void SDMMaterialEditor_LeftBase::EditGlobalSettings(bool bInForceRefresh)
+{
+	if (bGlobalSettingsMode && !bInForceRefresh)
+	{
+		return;
+	}
+
+	if (!bGlobalSettingsMode)
+	{
+		RightSlot.Invalidate();
+	}
+
+	SDMMaterialEditor::EditGlobalSettings(bInForceRefresh);
+}
+
+void SDMMaterialEditor_LeftBase::ValidateSlots_Main()
+{
+	if (LeftSlot.HasBeenInvalidated())
+	{
+		LeftSlot << CreateSlot_Left();
+	}
+
+	if (RightSlot.HasBeenInvalidated())
+	{
+		RightSlot << CreateSlot_Right();
+	}
+}
+
+void SDMMaterialEditor_LeftBase::ClearSlots_Main()
+{
+	LeftSlot.ClearWidget();
+	RightSlot.ClearWidget();
+}
+
+TSharedRef<SWidget> SDMMaterialEditor_LeftBase::CreateSlot_Main()
+{
+	SHorizontalBox::FSlot* LeftSlotPtr = nullptr;
+	SHorizontalBox::FSlot* RightSlotPtr = nullptr;
+
+	TSharedRef<SHorizontalBox> NewMain = SNew(SHorizontalBox)
+		+ SHorizontalBox::Slot()
+		.Expose(LeftSlotPtr)
+		.AutoWidth()
+		[
+			SNullWidget::NullWidget
+		]
+
+		+ SHorizontalBox::Slot()
+		.Expose(RightSlotPtr)
+		.FillWidth(1.0f)
+		[
+			SNullWidget::NullWidget
+		];
+
+	LeftSlot = TDMWidgetSlot<SWidget>(LeftSlotPtr, CreateSlot_Left());
+	RightSlot = TDMWidgetSlot<SWidget>(RightSlotPtr, CreateSlot_Right());
+
+	return NewMain;
+}
+
+TSharedRef<SWidget> SDMMaterialEditor_LeftBase::CreateSlot_Left()
+{
+	using namespace UE::DynamicMaterialEditor::Private;
+
+	SVerticalBox::FSlot* MaterialPreviewSlotPtr = nullptr;
+	SVerticalBox::FSlot* PropertySelectorSlotPtr = nullptr;
+
+	TSharedRef<SWidget> NewLeft = SNew(SBorder)
+		.BorderImage(FAppStyle::GetBrush(EditorDarkBackground))
+		.Padding(5.f)
+		[
+			SNew(SVerticalBox)
+			+ SVerticalBox::Slot()
+			.Expose(MaterialPreviewSlotPtr)
+			.AutoHeight()
+			.Padding(0.f)
+			[
+				SNullWidget::NullWidget
+			]
+
+			+ SVerticalBox::Slot()
+			.Expose(PropertySelectorSlotPtr)
+			.FillHeight(1.0f)
+			.Padding(0.f, 5.f, 0.f, 0.f)
+			[
+				SNullWidget::NullWidget
+			]
+		];
+
+	MaterialPreviewSlot = TDMWidgetSlot<SDMMaterialPreview>(MaterialPreviewSlotPtr, CreateSlot_Preview());
+	PropertySelectorSlot = TDMWidgetSlot<SDMMaterialPropertySelector>(PropertySelectorSlotPtr, CreateSlot_PropertySelector());
+
+	return NewLeft;
+}
+
+TSharedRef<SWidget> SDMMaterialEditor_LeftBase::CreateSlot_Right()
+{
+	using namespace UE::DynamicMaterialEditor::Private;
+
+	const bool bHasSlotToEdit = SlotToEdit.IsValid();
+
+	if (!bGlobalSettingsMode && !bHasSlotToEdit)
+	{
+		bGlobalSettingsMode = true;
+	}
+	else if (bHasSlotToEdit)
+	{
+		bGlobalSettingsMode = false;
+	}
+
+	return SNew(SBorder)
+		.BorderImage(FAppStyle::GetBrush(EditorDarkBackground))
+		.Padding(FMargin(0.f, 5.f))
+		[
+			bGlobalSettingsMode
+				? CreateSlot_Right_GlobalSettings()
+				: CreateSlot_Right_Slot()
+		];
+}
+
+TSharedRef<SWidget> SDMMaterialEditor_LeftBase::CreateSlot_Right_GlobalSettings()
+{
+	using namespace UE::DynamicMaterialEditor::Private;
+
+	SScrollBox::FSlot* GlobalSettingsSlotPtr = nullptr;
+
+	TSharedRef<SBorder> NewRight = SNew(SBorder)
+		.BorderImage(FAppStyle::GetBrush(EditorLightBackground))
+		.Padding(0.f)
+		[
+			SNew(SScrollBox)
+			+ SScrollBox::Slot()
+			.Expose(GlobalSettingsSlotPtr)
+			.VAlign(EVerticalAlignment::VAlign_Fill)
+			[
+				SNullWidget::NullWidget
+			]
+		];
+
+	GlobalSettingsEditorSlot = TDMWidgetSlot<SDMMaterialGlobalSettingsEditor>(GlobalSettingsSlotPtr, CreateSlot_GlobalSettingsEditor());
+
+	return NewRight;
+}
+
+TSharedRef<SWidget> SDMMaterialEditor_LeftBase::CreateSlot_Right_Slot()
+{
+	using namespace UE::DynamicMaterialEditor::Private;
+
+	float SplitterValue = 0.5;
+
+	if (UDynamicMaterialEditorSettings* Settings = UDynamicMaterialEditorSettings::Get())
+	{
+		SplitterValue = Settings->SplitterLocation;
+	}
+
+	TSharedPtr<SBorder> TopBox;
+	TSharedPtr<SBorder> BottomBox;
+
+	SSplitter::FSlot* ExplosedSlot = nullptr;
+
+	TSharedRef<SSplitter> NewRight = SNew(SSplitter)
+		.Style(FAppStyle::Get(), "DetailsView.Splitter")
+		.Orientation(Orient_Vertical)
+		.ResizeMode(ESplitterResizeMode::Fill)
+		.PhysicalSplitterHandleSize(5.0f)
+		.HitDetectionSplitterHandleSize(5.0f)
+		.OnSplitterFinishedResizing(this, &SDMMaterialEditor_LeftBase::OnEditorSplitterResized)
+
+		+ SSplitter::Slot()
+		.Expose(ExplosedSlot)
+		.Resizable(true)
+		.SizeRule(SSplitter::ESizeRule::FractionOfParent)
+		.MinSize(165)
+		.Value(SplitterValue)
+		[
+			SAssignNew(TopBox, SBorder)
+			.BorderImage(FAppStyle::GetBrush(EditorLightBackground))
+			[
+				SNullWidget::NullWidget
+			]
+		]
+
+		+ SSplitter::Slot()
+		.Resizable(true)
+		.SizeRule(SSplitter::ESizeRule::FractionOfParent)
+		.MinSize(60)
+		.Value(1.f - SplitterValue)
+		[
+			SAssignNew(BottomBox, SBorder)
+			.BorderImage(FAppStyle::GetBrush(EditorLightBackground))
+			[
+				SNullWidget::NullWidget
+			]
+		];
+
+	SplitterSlot = ExplosedSlot;
+	SlotEditorSlot = TDMWidgetSlot<SDMMaterialSlotEditor>(TopBox.ToSharedRef(), 0, CreateSlot_SlotEditor());
+	ComponentEditorSlot = TDMWidgetSlot<SDMMaterialComponentEditor>(BottomBox.ToSharedRef(), 0, CreateSlot_ComponentEditor());
+
+	return NewRight;
+}
diff --git a/Engine/Plugins/Experimental/DynamicMaterial/Source/DynamicMaterialEditor/Private/UI/Widgets/Editor/EditorLayouts/SDMMaterialEditor_LeftBase.h b/Engine/Plugins/Experimental/DynamicMaterial/Source/DynamicMaterialEditor/Private/UI/Widgets/Editor/EditorLayouts/SDMMaterialEditor_LeftBase.h
new file mode 100644
--- /dev/null
+++ b/Engine/Plugins/Experimental/DynamicMaterial/Source/DynamicMaterialEditor/Private/UI/Widgets/Editor/EditorLayouts/SDMMaterialEditor_LeftBase.h
@@ -0,0 +1,45 @@
+// Copyright Epic Games, Inc. All Rights Reserved.
+
+#pragma once
+
+#include "UI/Widgets/SDMMaterialEditor.h"
+
+class SDMMaterialEditor_LeftBase : public SDMMaterialEditor
+{
+	SLATE_BEGIN_ARGS(SDMMaterialEditor_LeftBase)
+		: _MaterialModelBase(nullptr)
+		, _MaterialProperty(TOptional<FDMObjectMaterialProperty>())
+		{}
+		SLATE_ARGUMENT(UDynamicMaterialModelBase*, MaterialModelBase)
+		SLATE_ARGUMENT(TOptional<FDMObjectMaterialProperty>, MaterialProperty)
+	SLATE_END_ARGS()
+
+public:
+	virtual ~SDMMaterialEditor_LeftBase() override = default;
+
+	void Construct(const FArguments& InArgs, const TSharedRef<SDMMaterialDesigner>& InDesignerWidget);
+
+	//~ Begin SDMMaterialEditor
+	virtual void EditSlot(UDMMaterialSlot* InSlot, bool bInForceRefresh = false) override;
+	virtual void EditComponent(UDMMaterialComponent* InComponent, bool bInForceRefresh = false) override;
+	virtual void EditGlobalSettings(bool bInForceRefresh = false) override;
+	//~ End SDMMaterialEditor
+
+protected:
+	TDMWidgetSlot<SWidget> LeftSlot;
+	TDMWidgetSlot<SWidget> RightSlot;
+
+	virtual TSharedRef<SWidget> CreateSlot_Left() = 0;
+
+	TSharedRef<SWidget> CreateSlot_Right();
+
+	TSharedRef<SWidget> CreateSlot_Right_GlobalSettings();
+
+	TSharedRef<SWidget> CreateSlot_Right_Slot();
+
+	//~ Begin SDMMaterialEditor
+	virtual void ValidateSlots_Main() override;
+	virtual void ClearSlots_Main() override;
+	virtual TSharedRef<SWidget> CreateSlot_Main() override;
+	//~ End SDMMaterialEditor
+};
diff --git a/Engine/Plugins/Experimental/DynamicMaterial/Source/DynamicMaterialEditor/Private/UI/Widgets/Editor/EditorLayouts/SDMMaterialEditor_LeftSlim.cpp b/Engine/Plugins/Experimental/DynamicMaterial/Source/DynamicMaterialEditor/Private/UI/Widgets/Editor/EditorLayouts/SDMMaterialEditor_LeftSlim.cpp
new file mode 100644
--- /dev/null
+++ b/Engine/Plugins/Experimental/DynamicMaterial/Source/DynamicMaterialEditor/Private/UI/Widgets/Editor/EditorLayouts/SDMMaterialEditor_LeftSlim.cpp
@@ -0,0 +1,49 @@
+// Copyright Epic Games, Inc. All Rights Reserved.
+
+#include "UI/Widgets/Editor/EditorLayouts/SDMMaterialEditor_LeftSlim.h"
+
+#include "Styling/AppStyle.h"
+#include "UI/Widgets/Editor/PropertySelectorLayouts/SDMMaterialPropertySelector_VerticalSlim.h"
+#include "UI/Widgets/Editor/SDMMaterialPropertySelector.h"
+#include "UI/Widgets/Editor/SDMMaterialSlotEditor.h"
+#include "Widgets/Layout/SBorder.h"
+#include "Widgets/SBoxPanel.h"
+#include "Widgets/SNullWidget.h"
+
+void SDMMaterialEditor_LeftSlim::Construct(const FArguments& InArgs, const TSharedRef<SDMMaterialDesigner>& InDesignerWidget)
+{
+	SDMMaterialEditor_LeftBase::Construct(
+		SDMMaterialEditor_LeftBase::FArguments()
+			.MaterialModelBase(InArgs._MaterialModelBase)
+			.MaterialProperty(InArgs._MaterialProperty),
+		InDesignerWidget
+	);
+}
+
+TSharedRef<SWidget> SDMMaterialEditor_LeftSlim::CreateSlot_Left()
+{
+	using namespace UE::DynamicMaterialEditor::Private;
+
+	SVerticalBox::FSlot* PropertySelectorSlotPtr = nullptr;
+
+	TSharedRef<SWidget> NewLeft = SNew(SBorder)
+		.BorderImage(FAppStyle::GetBrush(EditorDarkBackground))
+		.Padding(5.f)
+		[
+			SNew(SVerticalBox)
+			+ SVerticalBox::Slot()
+			.Expose(PropertySelectorSlotPtr)
+			[
+				SNullWidget::NullWidget
+			]
+		];
+
+	PropertySelectorSlot = TDMWidgetSlot<SDMMaterialPropertySelector>(PropertySelectorSlotPtr, CreateSlot_PropertySelector());
+
+	return NewLeft;
+}
+
+TSharedRef<SDMMaterialPropertySelector> SDMMaterialEditor_LeftSlim::CreateSlot_PropertySelector_Impl()
+{
+	return SNew(SDMMaterialPropertySelector_VerticalSlim, SharedThis(this));
+}
diff --git a/Engine/Plugins/Experimental/DynamicMaterial/Source/DynamicMaterialEditor/Private/UI/Widgets/Editor/EditorLayouts/SDMMaterialEditor_LeftSlim.h b/Engine/Plugins/Experimental/DynamicMaterial/Source/DynamicMaterialEditor/Private/UI/Widgets/Editor/EditorLayouts/SDMMaterialEditor_LeftSlim.h
new file mode 100644
--- /dev/null
+++ b/Engine/Plugins/Experimental/DynamicMaterial/Source/DynamicMaterialEditor/Private/UI/Widgets/Editor/EditorLayouts/SDMMaterialEditor_LeftSlim.h
@@ -0,0 +1,30 @@
+// Copyright Epic Games, Inc. All Rights Reserved.
+
+#pragma once
+
+#include "UI/Widgets/Editor/EditorLayouts/SDMMaterialEditor_LeftBase.h"
+
+class SDMMaterialEditor_LeftSlim : public SDMMaterialEditor_LeftBase
+{
+	SLATE_BEGIN_ARGS(SDMMaterialEditor_LeftSlim)
+		: _MaterialModelBase(nullptr)
+		, _MaterialProperty(TOptional<FDMObjectMaterialProperty>())
+		{}
+		SLATE_ARGUMENT(UDynamicMaterialModelBase*, MaterialModelBase)
+		SLATE_ARGUMENT(TOptional<FDMObjectMaterialProperty>, MaterialProperty)
+	SLATE_END_ARGS()
+
+public:
+	virtual ~SDMMaterialEditor_LeftSlim() override = default;
+
+	void Construct(const FArguments& InArgs, const TSharedRef<SDMMaterialDesigner>& InDesignerWidget);
+
+protected:
+	//~ Begin SDMMaterialEditor_LeftBase
+	virtual TSharedRef<SWidget> CreateSlot_Left() override;
+	//~ End SDMMaterialEditor_LeftBase
+
+	//~ Begin SDMMaterialEditor
+	virtual TSharedRef<SDMMaterialPropertySelector> CreateSlot_PropertySelector_Impl() override;
+	//~ End SDMMaterialEditor
+};
diff --git a/Engine/Plugins/Experimental/DynamicMaterial/Source/DynamicMaterialEditor/Private/UI/Widgets/Editor/EditorLayouts/SDMMaterialEditor_TopBase.cpp b/Engine/Plugins/Experimental/DynamicMaterial/Source/DynamicMaterialEditor/Private/UI/Widgets/Editor/EditorLayouts/SDMMaterialEditor_TopBase.cpp
new file mode 100644
--- /dev/null
+++ b/Engine/Plugins/Experimental/DynamicMaterial/Source/DynamicMaterialEditor/Private/UI/Widgets/Editor/EditorLayouts/SDMMaterialEditor_TopBase.cpp
@@ -0,0 +1,252 @@
+// Copyright Epic Games, Inc. All Bottoms Reserved.
+
+#include "UI/Widgets/Editor/EditorLayouts/SDMMaterialEditor_TopBase.h"
+
+#include "DynamicMaterialEditorSettings.h"
+#include "UI/Widgets/Editor/SDMMaterialComponentEditor.h"
+#include "UI/Widgets/Editor/SDMMaterialGlobalSettingsEditor.h"
+#include "UI/Widgets/Editor/SDMMaterialPreview.h"
+#include "UI/Widgets/Editor/SDMMaterialPropertySelector.h"
+#include "UI/Widgets/Editor/SDMMaterialSlotEditor.h"
+#include "Widgets/Layout/SBorder.h"
+#include "Widgets/Layout/SScrollBox.h"
+#include "Widgets/Layout/SSplitter.h"
+
+void SDMMaterialEditor_TopBase::Construct(const FArguments& InArgs, const TSharedRef<SDMMaterialDesigner>& InDesignerWidget)
+{
+	SDMMaterialEditor::Construct(
+		SDMMaterialEditor::FArguments()
+		.MaterialModelBase(InArgs._MaterialModelBase)
+		.MaterialProperty(InArgs._MaterialProperty),
+		InDesignerWidget
+	);
+}
+
+void SDMMaterialEditor_TopBase::EditSlot(UDMMaterialSlot* InSlot, bool bInForceRefresh)
+{
+	if (!bInForceRefresh && SlotEditorSlot.IsValid() && SlotEditorSlot->GetSlot() == InSlot)
+	{
+		return;
+	}
+
+	BottomSlot.Invalidate();
+
+	SDMMaterialEditor::EditSlot(InSlot, bInForceRefresh);
+}
+
+void SDMMaterialEditor_TopBase::EditComponent(UDMMaterialComponent* InComponent, bool bInForceRefresh)
+{
+	if (!bInForceRefresh && ComponentEditorSlot.IsValid() && ComponentEditorSlot->GetComponent() == InComponent)
+	{
+		return;
+	}
+
+	if (bGlobalSettingsMode)
+	{
+		BottomSlot.Invalidate();
+	}
+
+	SDMMaterialEditor::EditComponent(InComponent, bInForceRefresh);
+}
+
+void SDMMaterialEditor_TopBase::EditGlobalSettings(bool bInForceRefresh)
+{
+	if (bGlobalSettingsMode && !bInForceRefresh)
+	{
+		return;
+	}
+
+	if (!bGlobalSettingsMode)
+	{
+		BottomSlot.Invalidate();
+	}
+
+	SDMMaterialEditor::EditGlobalSettings(bInForceRefresh);
+}
+
+void SDMMaterialEditor_TopBase::ValidateSlots_Main()
+{
+	if (TopSlot.HasBeenInvalidated())
+	{
+		TopSlot << CreateSlot_Top();
+	}
+
+	if (BottomSlot.HasBeenInvalidated())
+	{
+		BottomSlot << CreateSlot_Bottom();
+	}
+}
+
+void SDMMaterialEditor_TopBase::ClearSlots_Main()
+{
+	TopSlot.ClearWidget();
+	BottomSlot.ClearWidget();
+}
+
+TSharedRef<SWidget> SDMMaterialEditor_TopBase::CreateSlot_Main()
+{
+	SVerticalBox::FSlot* TopSlotPtr = nullptr;
+	SVerticalBox::FSlot* BottomSlotPtr = nullptr;
+
+	TSharedRef<SVerticalBox> NewMain = SNew(SVerticalBox)
+		+ SVerticalBox::Slot()
+		.Expose(TopSlotPtr)
+		.AutoHeight()
+		[
+			SNullWidget::NullWidget
+		]
+
+		+ SVerticalBox::Slot()
+		.Expose(BottomSlotPtr)
+		.FillHeight(1.0f)
+		[
+			SNullWidget::NullWidget
+		];
+
+	TopSlot = TDMWidgetSlot<SWidget>(TopSlotPtr, CreateSlot_Top());
+	BottomSlot = TDMWidgetSlot<SWidget>(BottomSlotPtr, CreateSlot_Bottom());
+
+	return NewMain;
+}
+
+TSharedRef<SWidget> SDMMaterialEditor_TopBase::CreateSlot_Top()
+{
+	using namespace UE::DynamicMaterialEditor::Private;
+
+	SVerticalBox::FSlot* MaterialPreviewSlotPtr = nullptr;
+	SVerticalBox::FSlot* PropertySelectorSlotPtr = nullptr;
+
+	TSharedRef<SWidget> NewTop = SNew(SBorder)
+		.BorderImage(FAppStyle::GetBrush(EditorDarkBackground))
+		.Padding(5.f)
+		[
+			SNew(SVerticalBox)
+			+ SVerticalBox::Slot()
+			.Expose(MaterialPreviewSlotPtr)
+			.AutoHeight()
+			.Padding(0.f)
+			[
+				SNullWidget::NullWidget
+			]
+
+			+ SVerticalBox::Slot()
+			.Expose(PropertySelectorSlotPtr)
+			.FillHeight(1.0f)
+			.Padding(0.f, 5.f, 0.f, 0.f)
+			[
+				SNullWidget::NullWidget
+			]
+		];
+
+	MaterialPreviewSlot = TDMWidgetSlot<SDMMaterialPreview>(MaterialPreviewSlotPtr, CreateSlot_Preview());
+	PropertySelectorSlot = TDMWidgetSlot<SDMMaterialPropertySelector>(PropertySelectorSlotPtr, CreateSlot_PropertySelector());
+
+	return NewTop;
+}
+
+TSharedRef<SWidget> SDMMaterialEditor_TopBase::CreateSlot_Bottom()
+{
+	using namespace UE::DynamicMaterialEditor::Private;
+
+	const bool bHasSlotToEdit = SlotToEdit.IsValid();
+
+	if (!bGlobalSettingsMode && !bHasSlotToEdit)
+	{
+		bGlobalSettingsMode = true;
+	}
+	else if (bHasSlotToEdit)
+	{
+		bGlobalSettingsMode = false;
+	}
+
+	return SNew(SBorder)
+		.BorderImage(FAppStyle::GetBrush(EditorDarkBackground))
+		.Padding(FMargin(0.f, 5.f))
+		[
+			bGlobalSettingsMode
+				? CreateSlot_Bottom_GlobalSettings()
+				: CreateSlot_Bottom_Slot()
+		];
+}
+
+TSharedRef<SWidget> SDMMaterialEditor_TopBase::CreateSlot_Bottom_GlobalSettings()
+{
+	using namespace UE::DynamicMaterialEditor::Private;
+
+	SScrollBox::FSlot* GlobalSettingsSlotPtr = nullptr;
+
+	TSharedRef<SBorder> NewBottom = SNew(SBorder)
+		.BorderImage(FAppStyle::GetBrush(EditorLightBackground))
+		.Padding(0.f)
+		[
+			SNew(SScrollBox)
+			+ SScrollBox::Slot()
+			.Expose(GlobalSettingsSlotPtr)
+			.VAlign(EVerticalAlignment::VAlign_Fill)
+			[
+				SNullWidget::NullWidget
+			]
+		];
+
+	GlobalSettingsEditorSlot = TDMWidgetSlot<SDMMaterialGlobalSettingsEditor>(GlobalSettingsSlotPtr, CreateSlot_GlobalSettingsEditor());
+
+	return NewBottom;
+}
+
+TSharedRef<SWidget> SDMMaterialEditor_TopBase::CreateSlot_Bottom_Slot()
+{
+	using namespace UE::DynamicMaterialEditor::Private;
+
+	float SplitterValue = 0.5;
+
+	if (UDynamicMaterialEditorSettings* Settings = UDynamicMaterialEditorSettings::Get())
+	{
+		SplitterValue = Settings->SplitterLocation;
+	}
+
+	TSharedPtr<SBorder> TopBox;
+	TSharedPtr<SBorder> BottomBox;
+
+	SSplitter::FSlot* ExplosedSlot = nullptr;
+
+	TSharedRef<SSplitter> NewBottom = SNew(SSplitter)
+		.Style(FAppStyle::Get(), "DetailsView.Splitter")
+		.Orientation(Orient_Vertical)
+		.ResizeMode(ESplitterResizeMode::Fill)
+		.PhysicalSplitterHandleSize(5.0f)
+		.HitDetectionSplitterHandleSize(5.0f)
+		.OnSplitterFinishedResizing(this, &SDMMaterialEditor_TopBase::OnEditorSplitterResized)
+
+		+ SSplitter::Slot()
+		.Expose(ExplosedSlot)
+		.Resizable(true)
+		.SizeRule(SSplitter::ESizeRule::FractionOfParent)
+		.MinSize(165)
+		.Value(SplitterValue)
+		[
+			SAssignNew(TopBox, SBorder)
+			.BorderImage(FAppStyle::GetBrush(EditorLightBackground))
+			[
+				SNullWidget::NullWidget
+			]
+		]
+
+		+ SSplitter::Slot()
+		.Resizable(true)
+		.SizeRule(SSplitter::ESizeRule::FractionOfParent)
+		.MinSize(60)
+		.Value(1.f - SplitterValue)
+		[
+			SAssignNew(BottomBox, SBorder)
+			.BorderImage(FAppStyle::GetBrush(EditorLightBackground))
+			[
+				SNullWidget::NullWidget
+			]
+		];
+
+	SplitterSlot = ExplosedSlot;
+	SlotEditorSlot = TDMWidgetSlot<SDMMaterialSlotEditor>(TopBox.ToSharedRef(), 0, CreateSlot_SlotEditor());
+	ComponentEditorSlot = TDMWidgetSlot<SDMMaterialComponentEditor>(BottomBox.ToSharedRef(), 0, CreateSlot_ComponentEditor());
+
+	return NewBottom;
+}
diff --git a/Engine/Plugins/Experimental/DynamicMaterial/Source/DynamicMaterialEditor/Private/UI/Widgets/Editor/EditorLayouts/SDMMaterialEditor_TopBase.h b/Engine/Plugins/Experimental/DynamicMaterial/Source/DynamicMaterialEditor/Private/UI/Widgets/Editor/EditorLayouts/SDMMaterialEditor_TopBase.h
new file mode 100644
--- /dev/null
+++ b/Engine/Plugins/Experimental/DynamicMaterial/Source/DynamicMaterialEditor/Private/UI/Widgets/Editor/EditorLayouts/SDMMaterialEditor_TopBase.h
@@ -0,0 +1,45 @@
+// Copyright Epic Games, Inc. All Rights Reserved.
+
+#pragma once
+
+#include "UI/Widgets/SDMMaterialEditor.h"
+
+class SDMMaterialEditor_TopBase : public SDMMaterialEditor
+{
+	SLATE_BEGIN_ARGS(SDMMaterialEditor_TopBase)
+		: _MaterialModelBase(nullptr)
+		, _MaterialProperty(TOptional<FDMObjectMaterialProperty>())
+		{}
+		SLATE_ARGUMENT(UDynamicMaterialModelBase*, MaterialModelBase)
+		SLATE_ARGUMENT(TOptional<FDMObjectMaterialProperty>, MaterialProperty)
+	SLATE_END_ARGS()
+
+public:
+	virtual ~SDMMaterialEditor_TopBase() override = default;
+
+	void Construct(const FArguments& InArgs, const TSharedRef<SDMMaterialDesigner>& InDesignerWidget);
+
+	//~ Begin SDMMaterialEditor
+	virtual void EditSlot(UDMMaterialSlot* InSlot, bool bInForceRefresh = false) override;
+	virtual void EditComponent(UDMMaterialComponent* InComponent, bool bInForceRefresh = false) override;
+	virtual void EditGlobalSettings(bool bInForceRefresh = false) override;
+	//~ End SDMMaterialEditor
+
+protected:
+	TDMWidgetSlot<SWidget> TopSlot;
+	TDMWidgetSlot<SWidget> BottomSlot;
+
+	virtual TSharedRef<SWidget> CreateSlot_Top() = 0;
+
+	TSharedRef<SWidget> CreateSlot_Bottom();
+
+	TSharedRef<SWidget> CreateSlot_Bottom_GlobalSettings();
+
+	TSharedRef<SWidget> CreateSlot_Bottom_Slot();
+
+	//~ Begin SDMMaterialEditor
+	virtual void ValidateSlots_Main() override;
+	virtual void ClearSlots_Main() override;
+	virtual TSharedRef<SWidget> CreateSlot_Main() override;
+	//~ End SDMMaterialEditor
+};
diff --git a/Engine/Plugins/Experimental/DynamicMaterial/Source/DynamicMaterialEditor/Private/UI/Widgets/Editor/EditorLayouts/SDMMaterialEditor_TopHorizontal.cpp b/Engine/Plugins/Experimental/DynamicMaterial/Source/DynamicMaterialEditor/Private/UI/Widgets/Editor/EditorLayouts/SDMMaterialEditor_TopHorizontal.cpp
new file mode 100644
--- /dev/null
+++ b/Engine/Plugins/Experimental/DynamicMaterial/Source/DynamicMaterialEditor/Private/UI/Widgets/Editor/EditorLayouts/SDMMaterialEditor_TopHorizontal.cpp
@@ -0,0 +1,62 @@
+// Copyright Epic Games, Inc. All Rights Reserved.
+
+#include "UI/Widgets/Editor/EditorLayouts/SDMMaterialEditor_TopHorizontal.h"
+
+#include "Styling/AppStyle.h"
+#include "UI/Widgets/Editor/PropertySelectorLayouts/SDMMaterialPropertySelector_Wrap.h"
+#include "UI/Widgets/Editor/SDMMaterialPreview.h"
+#include "UI/Widgets/Editor/SDMMaterialPropertySelector.h"
+#include "UI/Widgets/Editor/SDMMaterialSlotEditor.h"
+#include "Widgets/Layout/SBorder.h"
+#include "Widgets/SBoxPanel.h"
+#include "Widgets/SNullWidget.h"
+
+void SDMMaterialEditor_TopHorizontal::Construct(const FArguments& InArgs, const TSharedRef<SDMMaterialDesigner>& InDesignerWidget)
+{
+	SDMMaterialEditor_TopBase::Construct(
+		SDMMaterialEditor_TopBase::FArguments()
+		.MaterialModelBase(InArgs._MaterialModelBase)
+		.MaterialProperty(InArgs._MaterialProperty),
+		InDesignerWidget
+	);
+}
+
+TSharedRef<SWidget> SDMMaterialEditor_TopHorizontal::CreateSlot_Top()
+{
+	using namespace UE::DynamicMaterialEditor::Private;
+
+	SHorizontalBox::FSlot* MaterialPreviewSlotPtr = nullptr;
+	SHorizontalBox::FSlot* PropertySelectorSlotPtr = nullptr;
+
+	TSharedRef<SWidget> NewLeft = SNew(SBorder)
+		.BorderImage(FAppStyle::GetBrush(EditorDarkBackground))
+		.Padding(5.f, 5.f, 5.f, 0.f)
+		[
+			SNew(SHorizontalBox)
+			+ SHorizontalBox::Slot()
+			.Expose(MaterialPreviewSlotPtr)
+			.AutoWidth()
+			.Padding(0.f)
+			[
+				SNullWidget::NullWidget
+			]
+
+			+ SHorizontalBox::Slot()
+			.Expose(PropertySelectorSlotPtr)
+			.FillWidth(1.0f)
+			.Padding(5.f, 0.f, 0.f, 0.f)
+			[
+				SNullWidget::NullWidget
+			]
+		];
+
+	MaterialPreviewSlot = TDMWidgetSlot<SDMMaterialPreview>(MaterialPreviewSlotPtr, CreateSlot_Preview());
+	PropertySelectorSlot = TDMWidgetSlot<SDMMaterialPropertySelector>(PropertySelectorSlotPtr, CreateSlot_PropertySelector());
+
+	return NewLeft;
+}
+
+TSharedRef<SDMMaterialPropertySelector> SDMMaterialEditor_TopHorizontal::CreateSlot_PropertySelector_Impl()
+{
+	return SNew(SDMMaterialPropertySelector_Wrap, SharedThis(this));
+}
diff --git a/Engine/Plugins/Experimental/DynamicMaterial/Source/DynamicMaterialEditor/Private/UI/Widgets/Editor/EditorLayouts/SDMMaterialEditor_TopHorizontal.h b/Engine/Plugins/Experimental/DynamicMaterial/Source/DynamicMaterialEditor/Private/UI/Widgets/Editor/EditorLayouts/SDMMaterialEditor_TopHorizontal.h
new file mode 100644
--- /dev/null
+++ b/Engine/Plugins/Experimental/DynamicMaterial/Source/DynamicMaterialEditor/Private/UI/Widgets/Editor/EditorLayouts/SDMMaterialEditor_TopHorizontal.h
@@ -0,0 +1,30 @@
+// Copyright Epic Games, Inc. All Rights Reserved.
+
+#pragma once
+
+#include "UI/Widgets/Editor/EditorLayouts/SDMMaterialEditor_TopBase.h"
+
+class SDMMaterialEditor_TopHorizontal : public SDMMaterialEditor_TopBase
+{
+	SLATE_BEGIN_ARGS(SDMMaterialEditor_TopHorizontal)
+		: _MaterialModelBase(nullptr)
+		, _MaterialProperty(TOptional<FDMObjectMaterialProperty>())
+		{}
+		SLATE_ARGUMENT(UDynamicMaterialModelBase*, MaterialModelBase)
+		SLATE_ARGUMENT(TOptional<FDMObjectMaterialProperty>, MaterialProperty)
+	SLATE_END_ARGS()
+
+public:
+	virtual ~SDMMaterialEditor_TopHorizontal() override = default;
+
+	void Construct(const FArguments& InArgs, const TSharedRef<SDMMaterialDesigner>& InDesignerWidget);
+
+protected:
+	//~ Begin SDMMaterialEditor_TopBase
+	virtual TSharedRef<SWidget> CreateSlot_Top() override;
+	//~ End SDMMaterialEditor_TopBase
+
+	//~ Begin SDMMaterialEditor
+	virtual TSharedRef<SDMMaterialPropertySelector> CreateSlot_PropertySelector_Impl() override;
+	//~ End SDMMaterialEditor
+};
diff --git a/Engine/Plugins/Experimental/DynamicMaterial/Source/DynamicMaterialEditor/Private/UI/Widgets/Editor/EditorLayouts/SDMMaterialEditor_TopSlim.cpp b/Engine/Plugins/Experimental/DynamicMaterial/Source/DynamicMaterialEditor/Private/UI/Widgets/Editor/EditorLayouts/SDMMaterialEditor_TopSlim.cpp
new file mode 100644
--- /dev/null
+++ b/Engine/Plugins/Experimental/DynamicMaterial/Source/DynamicMaterialEditor/Private/UI/Widgets/Editor/EditorLayouts/SDMMaterialEditor_TopSlim.cpp
@@ -0,0 +1,50 @@
+// Copyright Epic Games, Inc. All Rights Reserved.
+
+#include "UI/Widgets/Editor/EditorLayouts/SDMMaterialEditor_TopSlim.h"
+
+#include "Styling/AppStyle.h"
+#include "UI/Widgets/Editor/PropertySelectorLayouts/SDMMaterialPropertySelector_WrapSlim.h"
+#include "UI/Widgets/Editor/SDMMaterialPropertySelector.h"
+#include "UI/Widgets/Editor/SDMMaterialSlotEditor.h"
+#include "Widgets/Layout/SBorder.h"
+#include "Widgets/SBoxPanel.h"
+#include "Widgets/SNullWidget.h"
+
+void SDMMaterialEditor_TopSlim::Construct(const FArguments& InArgs, const TSharedRef<SDMMaterialDesigner>& InDesignerWidget)
+{
+	SDMMaterialEditor_TopBase::Construct(
+		SDMMaterialEditor_TopBase::FArguments()
+		.MaterialModelBase(InArgs._MaterialModelBase)
+		.MaterialProperty(InArgs._MaterialProperty),
+		InDesignerWidget
+	);
+}
+
+TSharedRef<SWidget> SDMMaterialEditor_TopSlim::CreateSlot_Top()
+{
+	using namespace UE::DynamicMaterialEditor::Private;
+
+	SHorizontalBox::FSlot* PropertySelectorSlotPtr = nullptr;
+
+	TSharedRef<SWidget> NewLeft = SNew(SBorder)
+		.BorderImage(FAppStyle::GetBrush(EditorDarkBackground))
+		.Padding(5.f, 5.f, 5.f, 0.f)
+		[
+			SNew(SHorizontalBox)
+			+ SHorizontalBox::Slot()
+			.Expose(PropertySelectorSlotPtr)
+			.FillWidth(1.0f)
+			[
+				SNullWidget::NullWidget
+			]
+		];
+
+	PropertySelectorSlot = TDMWidgetSlot<SDMMaterialPropertySelector>(PropertySelectorSlotPtr, CreateSlot_PropertySelector());
+
+	return NewLeft;
+}
+
+TSharedRef<SDMMaterialPropertySelector> SDMMaterialEditor_TopSlim::CreateSlot_PropertySelector_Impl()
+{
+	return SNew(SDMMaterialPropertySelector_WrapSlim, SharedThis(this));
+}
diff --git a/Engine/Plugins/Experimental/DynamicMaterial/Source/DynamicMaterialEditor/Private/UI/Widgets/Editor/EditorLayouts/SDMMaterialEditor_TopSlim.h b/Engine/Plugins/Experimental/DynamicMaterial/Source/DynamicMaterialEditor/Private/UI/Widgets/Editor/EditorLayouts/SDMMaterialEditor_TopSlim.h
new file mode 100644
--- /dev/null
+++ b/Engine/Plugins/Experimental/DynamicMaterial/Source/DynamicMaterialEditor/Private/UI/Widgets/Editor/EditorLayouts/SDMMaterialEditor_TopSlim.h
@@ -0,0 +1,30 @@
+// Copyright Epic Games, Inc. All Rights Reserved.
+
+#pragma once
+
+#include "UI/Widgets/Editor/EditorLayouts/SDMMaterialEditor_TopBase.h"
+
+class SDMMaterialEditor_TopSlim : public SDMMaterialEditor_TopBase
+{
+	SLATE_BEGIN_ARGS(SDMMaterialEditor_TopSlim)
+		: _MaterialModelBase(nullptr)
+		, _MaterialProperty(TOptional<FDMObjectMaterialProperty>())
+		{}
+		SLATE_ARGUMENT(UDynamicMaterialModelBase*, MaterialModelBase)
+		SLATE_ARGUMENT(TOptional<FDMObjectMaterialProperty>, MaterialProperty)
+	SLATE_END_ARGS()
+
+public:
+	virtual ~SDMMaterialEditor_TopSlim() override = default;
+
+	void Construct(const FArguments& InArgs, const TSharedRef<SDMMaterialDesigner>& InDesignerWidget);
+
+protected:
+	//~ Begin SDMMaterialEditor_TopBase
+	virtual TSharedRef<SWidget> CreateSlot_Top() override;
+	//~ End SDMMaterialEditor_TopBase
+
+	//~ Begin SDMMaterialEditor
+	virtual TSharedRef<SDMMaterialPropertySelector> CreateSlot_PropertySelector_Impl() override;
+	//~ End SDMMaterialEditor
+};
diff --git a/Engine/Plugins/Experimental/DynamicMaterial/Source/DynamicMaterialEditor/Private/UI/Widgets/Editor/EditorLayouts/SDMMaterialEditor_TopVertical.cpp b/Engine/Plugins/Experimental/DynamicMaterial/Source/DynamicMaterialEditor/Private/UI/Widgets/Editor/EditorLayouts/SDMMaterialEditor_TopVertical.cpp
new file mode 100644
--- /dev/null
+++ b/Engine/Plugins/Experimental/DynamicMaterial/Source/DynamicMaterialEditor/Private/UI/Widgets/Editor/EditorLayouts/SDMMaterialEditor_TopVertical.cpp
@@ -0,0 +1,61 @@
+// Copyright Epic Games, Inc. All Rights Reserved.
+
+#include "UI/Widgets/Editor/EditorLayouts/SDMMaterialEditor_TopVertical.h"
+
+#include "Styling/AppStyle.h"
+#include "UI/Widgets/Editor/PropertySelectorLayouts/SDMMaterialPropertySelector_Wrap.h"
+#include "UI/Widgets/Editor/SDMMaterialPreview.h"
+#include "UI/Widgets/Editor/SDMMaterialPropertySelector.h"
+#include "UI/Widgets/Editor/SDMMaterialSlotEditor.h"
+#include "Widgets/Layout/SBorder.h"
+#include "Widgets/SBoxPanel.h"
+
+void SDMMaterialEditor_TopVertical::Construct(const FArguments& InArgs, const TSharedRef<SDMMaterialDesigner>& InDesignerWidget)
+{
+	SDMMaterialEditor_TopBase::Construct(
+		SDMMaterialEditor_TopBase::FArguments()
+		.MaterialModelBase(InArgs._MaterialModelBase)
+		.MaterialProperty(InArgs._MaterialProperty),
+		InDesignerWidget
+	);
+}
+
+TSharedRef<SWidget> SDMMaterialEditor_TopVertical::CreateSlot_Top()
+{
+	using namespace UE::DynamicMaterialEditor::Private;
+
+	SVerticalBox::FSlot* MaterialPreviewSlotPtr = nullptr;
+	SVerticalBox::FSlot* PropertySelectorSlotPtr = nullptr;
+
+	TSharedRef<SWidget> NewLeft = SNew(SBorder)
+		.BorderImage(FAppStyle::GetBrush(EditorDarkBackground))
+		.Padding(FMargin(5.f, 5.f, 5.f, 0.f))
+		[
+			SNew(SVerticalBox)
+			+ SVerticalBox::Slot()
+			.Expose(MaterialPreviewSlotPtr)
+			.AutoHeight()
+			.Padding(0.f)
+			[
+				SNullWidget::NullWidget
+			]
+
+			+ SVerticalBox::Slot()
+			.Expose(PropertySelectorSlotPtr)
+			.FillHeight(1.0f)
+			.Padding(0.f, 5.f, 0.f, 0.f)
+			[
+				SNullWidget::NullWidget
+			]
+		];
+
+	MaterialPreviewSlot = TDMWidgetSlot<SDMMaterialPreview>(MaterialPreviewSlotPtr, CreateSlot_Preview());
+	PropertySelectorSlot = TDMWidgetSlot<SDMMaterialPropertySelector>(PropertySelectorSlotPtr, CreateSlot_PropertySelector());
+
+	return NewLeft;
+}
+
+TSharedRef<SDMMaterialPropertySelector> SDMMaterialEditor_TopVertical::CreateSlot_PropertySelector_Impl()
+{
+	return SNew(SDMMaterialPropertySelector_Wrap, SharedThis(this));
+}
diff --git a/Engine/Plugins/Experimental/DynamicMaterial/Source/DynamicMaterialEditor/Private/UI/Widgets/Editor/EditorLayouts/SDMMaterialEditor_TopVertical.h b/Engine/Plugins/Experimental/DynamicMaterial/Source/DynamicMaterialEditor/Private/UI/Widgets/Editor/EditorLayouts/SDMMaterialEditor_TopVertical.h
new file mode 100644
--- /dev/null
+++ b/Engine/Plugins/Experimental/DynamicMaterial/Source/DynamicMaterialEditor/Private/UI/Widgets/Editor/EditorLayouts/SDMMaterialEditor_TopVertical.h
@@ -0,0 +1,30 @@
+// Copyright Epic Games, Inc. All Rights Reserved.
+
+#pragma once
+
+#include "UI/Widgets/Editor/EditorLayouts/SDMMaterialEditor_TopBase.h"
+
+class SDMMaterialEditor_TopVertical : public SDMMaterialEditor_TopBase
+{
+	SLATE_BEGIN_ARGS(SDMMaterialEditor_TopVertical)
+		: _MaterialModelBase(nullptr)
+		, _MaterialProperty(TOptional<FDMObjectMaterialProperty>())
+		{}
+		SLATE_ARGUMENT(UDynamicMaterialModelBase*, MaterialModelBase)
+		SLATE_ARGUMENT(TOptional<FDMObjectMaterialProperty>, MaterialProperty)
+	SLATE_END_ARGS()
+
+public:
+	virtual ~SDMMaterialEditor_TopVertical() override = default;
+
+	void Construct(const FArguments& InArgs, const TSharedRef<SDMMaterialDesigner>& InDesignerWidget);
+
+protected:
+	//~ Begin SDMMaterialEditor_TopBase
+	virtual TSharedRef<SWidget> CreateSlot_Top() override;
+	//~ End SDMMaterialEditor_TopBase
+
+	//~ Begin SDMMaterialEditor
+	virtual TSharedRef<SDMMaterialPropertySelector> CreateSlot_PropertySelector_Impl() override;
+	//~ End SDMMaterialEditor
+};
diff --git a/Engine/Plugins/Experimental/DynamicMaterial/Source/DynamicMaterialEditor/Private/UI/Widgets/Editor/PropertySelectorLayouts/SDMMaterialPropertySelector_Vertical.cpp b/Engine/Plugins/Experimental/DynamicMaterial/Source/DynamicMaterialEditor/Private/UI/Widgets/Editor/PropertySelectorLayouts/SDMMaterialPropertySelector_Vertical.cpp
new file mode 100644
--- /dev/null
+++ b/Engine/Plugins/Experimental/DynamicMaterial/Source/DynamicMaterialEditor/Private/UI/Widgets/Editor/PropertySelectorLayouts/SDMMaterialPropertySelector_Vertical.cpp
@@ -0,0 +1,74 @@
+// Copyright Epic Games, Inc. All Rights Reserved.
+
+#include "UI/Widgets/Editor/PropertySelectorLayouts/SDMMaterialPropertySelector_Vertical.h"
+
+#include "DetailLayoutBuilder.h"
+#include "DMDefs.h"
+#include "Widgets/Images/SImage.h"
+#include "Widgets/Input/SCheckBox.h"
+#include "Widgets/Layout/SBox.h"
+#include "Widgets/Layout/SGridPanel.h"
+#include "Widgets/SBoxPanel.h"
+#include "Widgets/Text/STextBlock.h"
+
+#define LOCTEXT_NAMESPACE "SDMMaterialPropertySelector_Vertical"
+
+void SDMMaterialPropertySelector_Vertical::Construct(const FArguments& InArgs, const TSharedRef<SDMMaterialEditor> InEditorWidget)
+{
+	SDMMaterialPropertySelector_VerticalBase::Construct(
+		SDMMaterialPropertySelector_VerticalBase::FArguments(),
+		InEditorWidget
+	);
+}
+
+TSharedRef<SWidget> SDMMaterialPropertySelector_Vertical::CreateSlot_SelectButton(EDMMaterialPropertyType InMaterialProperty)
+{
+	UEnum* PropertyEnum = StaticEnum<EDMMaterialPropertyType>();
+
+	const FText ButtonText = InMaterialProperty == EDMMaterialPropertyType::None
+		? LOCTEXT("GlobalSettings", "Global Settings")
+		: PropertyEnum->GetDisplayNameTextByValue(static_cast<int64>(InMaterialProperty));
+
+	const FText Format = LOCTEXT("PropertySelectFormat", "Edit the {0} property.");
+
+	const FText ToolTip = (InMaterialProperty == EDMMaterialPropertyType::None)
+		? LOCTEXT("GeneralSettingsToolTip", "Edit the material global settings.")
+		: FText::Format(Format, PropertyEnum->GetDisplayNameTextByValue(static_cast<int64>(InMaterialProperty)));
+
+	return SNew(SCheckBox)
+		.Style(FAppStyle::Get(), "DetailsView.SectionButton")
+		.HAlign(EHorizontalAlignment::HAlign_Center)
+		.Padding(0.f)
+		.IsEnabled(this, &SDMMaterialPropertySelector_Vertical::GetPropertySelectEnabled, InMaterialProperty)
+		.IsChecked(this, &SDMMaterialPropertySelector_Vertical::GetPropertySelectState, InMaterialProperty)
+		.OnCheckStateChanged(this, &SDMMaterialPropertySelector_Vertical::OnPropertySelectStateChanged, InMaterialProperty)
+		.ToolTipText(ToolTip)
+		.Content()
+		[
+			SNew(SBox)
+			.WidthOverride(135.f)
+			.Clipping(EWidgetClipping::ClipToBounds)
+			[
+				SNew(SHorizontalBox)
+				+SHorizontalBox::Slot()
+				.VAlign(VAlign_Center)
+				.AutoWidth()
+				[
+					SNew(SImage)
+					.Image(FAppStyle::Get().GetBrush("FilterBar.FilterImage"))
+					.ColorAndOpacity(this, &SDMMaterialPropertySelector_Vertical::GetPropertySelectButtonChipColor, InMaterialProperty)
+				]
+				+SHorizontalBox::Slot()
+				.Padding(10.f, 6.f)
+				.VAlign(VAlign_Center)
+				.FillWidth(1.f)
+				[
+					SNew(STextBlock)
+					.Font(IDetailLayoutBuilder::GetDetailFont())
+					.Text(ButtonText)
+				]
+			]
+		];
+}
+
+#undef LOCTEXT_NAMESPACE
diff --git a/Engine/Plugins/Experimental/DynamicMaterial/Source/DynamicMaterialEditor/Private/UI/Widgets/Editor/PropertySelectorLayouts/SDMMaterialPropertySelector_Vertical.h b/Engine/Plugins/Experimental/DynamicMaterial/Source/DynamicMaterialEditor/Private/UI/Widgets/Editor/PropertySelectorLayouts/SDMMaterialPropertySelector_Vertical.h
new file mode 100644
--- /dev/null
+++ b/Engine/Plugins/Experimental/DynamicMaterial/Source/DynamicMaterialEditor/Private/UI/Widgets/Editor/PropertySelectorLayouts/SDMMaterialPropertySelector_Vertical.h
@@ -0,0 +1,21 @@
+// Copyright Epic Games, Inc. All Rights Reserved.
+
+#pragma once
+
+#include "UI/Widgets/Editor/PropertySelectorLayouts/SDMMaterialPropertySelector_VerticalBase.h"
+
+class SDMMaterialPropertySelector_Vertical : public SDMMaterialPropertySelector_VerticalBase
+{
+	SLATE_BEGIN_ARGS(SDMMaterialPropertySelector_Vertical) {}
+	SLATE_END_ARGS()
+
+public:
+	virtual ~SDMMaterialPropertySelector_Vertical() override = default;
+
+	void Construct(const FArguments& InArgs, const TSharedRef<SDMMaterialEditor> InEditorWidget);
+
+protected:
+	//~ Begin SDMMaterialPropertySelector
+	virtual TSharedRef<SWidget> CreateSlot_SelectButton(EDMMaterialPropertyType InMaterialProperty) override;
+	//~ End SDMMaterialPropertySelector
+};
diff --git a/Engine/Plugins/Experimental/DynamicMaterial/Source/DynamicMaterialEditor/Private/UI/Widgets/Editor/PropertySelectorLayouts/SDMMaterialPropertySelector_VerticalBase.cpp b/Engine/Plugins/Experimental/DynamicMaterial/Source/DynamicMaterialEditor/Private/UI/Widgets/Editor/PropertySelectorLayouts/SDMMaterialPropertySelector_VerticalBase.cpp
new file mode 100644
--- /dev/null
+++ b/Engine/Plugins/Experimental/DynamicMaterial/Source/DynamicMaterialEditor/Private/UI/Widgets/Editor/PropertySelectorLayouts/SDMMaterialPropertySelector_VerticalBase.cpp
@@ -0,0 +1,68 @@
+// Copyright Epic Games, Inc. All Rights Reserved.
+
+#include "UI/Widgets/Editor/PropertySelectorLayouts/SDMMaterialPropertySelector_VerticalBase.h"
+
+#include "DMDefs.h"
+#include "Model/DynamicMaterialModelEditorOnlyData.h"
+#include "Utils/DMPrivate.h"
+#include "Widgets/Layout/SBox.h"
+#include "Widgets/Layout/SGridPanel.h"
+#include "Widgets/SBoxPanel.h"
+
+#define LOCTEXT_NAMESPACE "SDMMaterialPropertySelector_VerticalBase"
+
+void SDMMaterialPropertySelector_VerticalBase::Construct(const FArguments& InArgs, const TSharedRef<SDMMaterialEditor> InEditorWidget)
+{
+	SDMMaterialPropertySelector::Construct(
+		SDMMaterialPropertySelector::FArguments(),
+		InEditorWidget
+	);
+}
+
+TSharedRef<SWidget> SDMMaterialPropertySelector_VerticalBase::CreateSlot_PropertyList()
+{
+	using namespace UE::DynamicMaterialEditor::Private;
+
+	TSharedRef<SGridPanel> NewSlotList = SNew(SGridPanel)
+		.FillColumn(PropertySelectorColumns::Select, 1.f);
+
+	UDynamicMaterialModelEditorOnlyData* EditorOnlyData = GetEditorOnlyData();
+
+	if (!EditorOnlyData)
+	{
+		return NewSlotList;
+	}
+
+	int32 Row = 0;
+
+	NewSlotList->AddSlot(PropertySelectorColumns::Select, Row)
+		[
+			CreateSlot_SelectButton(EDMMaterialPropertyType::None)
+		];
+
+	++Row;
+
+	for (const TPair<EDMMaterialPropertyType, UDMMaterialProperty*>& PropertyPair : EditorOnlyData->GetMaterialProperties())
+	{
+		if (IsCustomMaterialProperty(PropertyPair.Key))
+		{
+			continue;
+		}
+
+		NewSlotList->AddSlot(PropertySelectorColumns::Enable, Row)
+			[
+				CreateSlot_EnabledButton(PropertyPair.Key)
+			];
+
+		NewSlotList->AddSlot(PropertySelectorColumns::Select, Row)
+			[
+				CreateSlot_SelectButton(PropertyPair.Key)
+			];
+
+		++Row;
+	}
+
+	return NewSlotList;
+}
+
+#undef LOCTEXT_NAMESPACE
diff --git a/Engine/Plugins/Experimental/DynamicMaterial/Source/DynamicMaterialEditor/Private/UI/Widgets/Editor/PropertySelectorLayouts/SDMMaterialPropertySelector_VerticalBase.h b/Engine/Plugins/Experimental/DynamicMaterial/Source/DynamicMaterialEditor/Private/UI/Widgets/Editor/PropertySelectorLayouts/SDMMaterialPropertySelector_VerticalBase.h
new file mode 100644
--- /dev/null
+++ b/Engine/Plugins/Experimental/DynamicMaterial/Source/DynamicMaterialEditor/Private/UI/Widgets/Editor/PropertySelectorLayouts/SDMMaterialPropertySelector_VerticalBase.h
@@ -0,0 +1,21 @@
+// Copyright Epic Games, Inc. All Rights Reserved.
+
+#pragma once
+
+#include "UI/Widgets/Editor/SDMMaterialPropertySelector.h"
+
+class SDMMaterialPropertySelector_VerticalBase : public SDMMaterialPropertySelector
+{
+	SLATE_BEGIN_ARGS(SDMMaterialPropertySelector_VerticalBase) {}
+	SLATE_END_ARGS()
+
+public:
+	virtual ~SDMMaterialPropertySelector_VerticalBase() override = default;
+
+	void Construct(const FArguments& InArgs, const TSharedRef<SDMMaterialEditor> InEditorWidget);
+
+protected:
+	//~ Begin SDMMaterialPropertySelector
+	virtual TSharedRef<SWidget> CreateSlot_PropertyList() override;
+	//~ End SDMMaterialPropertySelector
+};
diff --git a/Engine/Plugins/Experimental/DynamicMaterial/Source/DynamicMaterialEditor/Private/UI/Widgets/Editor/PropertySelectorLayouts/SDMMaterialPropertySelector_VerticalSlim.cpp b/Engine/Plugins/Experimental/DynamicMaterial/Source/DynamicMaterialEditor/Private/UI/Widgets/Editor/PropertySelectorLayouts/SDMMaterialPropertySelector_VerticalSlim.cpp
new file mode 100644
--- /dev/null
+++ b/Engine/Plugins/Experimental/DynamicMaterial/Source/DynamicMaterialEditor/Private/UI/Widgets/Editor/PropertySelectorLayouts/SDMMaterialPropertySelector_VerticalSlim.cpp
@@ -0,0 +1,75 @@
+// Copyright Epic Games, Inc. All Rights Reserved.
+
+#include "UI/Widgets/Editor/PropertySelectorLayouts/SDMMaterialPropertySelector_VerticalSlim.h"
+
+#include "DetailLayoutBuilder.h"
+#include "DMDefs.h"
+#include "Utils/DMPrivate.h"
+#include "Widgets/Images/SImage.h"
+#include "Widgets/Input/SCheckBox.h"
+#include "Widgets/Layout/SBox.h"
+#include "Widgets/Layout/SGridPanel.h"
+#include "Widgets/SBoxPanel.h"
+#include "Widgets/Text/STextBlock.h"
+
+#define LOCTEXT_NAMESPACE "SDMMaterialPropertySelector_VerticalSlim"
+
+void SDMMaterialPropertySelector_VerticalSlim::Construct(const FArguments& InArgs, const TSharedRef<SDMMaterialEditor> InEditorWidget)
+{
+	SDMMaterialPropertySelector_VerticalBase::Construct(
+		SDMMaterialPropertySelector_VerticalBase::FArguments(),
+		InEditorWidget
+	);
+}
+
+TSharedRef<SWidget> SDMMaterialPropertySelector_VerticalSlim::CreateSlot_SelectButton(EDMMaterialPropertyType InMaterialProperty)
+{
+	UEnum* PropertyEnum = StaticEnum<EDMMaterialPropertyType>();
+
+	const FText ButtonText = InMaterialProperty == EDMMaterialPropertyType::None
+		? LOCTEXT("GlobalSettings", "Global Settings")
+		: UE::DynamicMaterialEditor::Private::GetMaterialPropertyShortDisplayName(InMaterialProperty);
+
+	const FText Format = LOCTEXT("PropertySelectFormat", "Edit the {0} property.");
+
+	const FText ToolTip = (InMaterialProperty == EDMMaterialPropertyType::None)
+		? LOCTEXT("GeneralSettingsToolTip", "Edit the material global settings.")
+		: FText::Format(Format, PropertyEnum->GetDisplayNameTextByValue(static_cast<int64>(InMaterialProperty)));
+
+	return SNew(SCheckBox)
+		.Style(FAppStyle::Get(), "DetailsView.SectionButton")
+		.HAlign(EHorizontalAlignment::HAlign_Center)
+		.Padding(0.f)
+		.IsEnabled(this, &SDMMaterialPropertySelector_VerticalSlim::GetPropertySelectEnabled, InMaterialProperty)
+		.IsChecked(this, &SDMMaterialPropertySelector_VerticalSlim::GetPropertySelectState, InMaterialProperty)
+		.OnCheckStateChanged(this, &SDMMaterialPropertySelector_VerticalSlim::OnPropertySelectStateChanged, InMaterialProperty)
+		.ToolTipText(ToolTip)
+		.Content()
+		[
+			SNew(SBox)
+			.WidthOverride(42.f)
+			.Clipping(EWidgetClipping::ClipToBounds)
+			[
+				SNew(SHorizontalBox)
+				+SHorizontalBox::Slot()
+				.VAlign(VAlign_Center)
+				.AutoWidth()
+				[
+					SNew(SImage)
+					.Image(FAppStyle::Get().GetBrush("FilterBar.FilterImage"))
+					.ColorAndOpacity(this, &SDMMaterialPropertySelector_VerticalSlim::GetPropertySelectButtonChipColor, InMaterialProperty)
+				]
+				+SHorizontalBox::Slot()
+				.Padding(2.f, 6.f)
+				.VAlign(VAlign_Center)
+				.FillWidth(1.f)
+				[
+					SNew(STextBlock)
+					.Font(IDetailLayoutBuilder::GetDetailFont())
+					.Text(ButtonText)
+				]
+			]
+		];
+}
+
+#undef LOCTEXT_NAMESPACE
diff --git a/Engine/Plugins/Experimental/DynamicMaterial/Source/DynamicMaterialEditor/Private/UI/Widgets/Editor/PropertySelectorLayouts/SDMMaterialPropertySelector_VerticalSlim.h b/Engine/Plugins/Experimental/DynamicMaterial/Source/DynamicMaterialEditor/Private/UI/Widgets/Editor/PropertySelectorLayouts/SDMMaterialPropertySelector_VerticalSlim.h
new file mode 100644
--- /dev/null
+++ b/Engine/Plugins/Experimental/DynamicMaterial/Source/DynamicMaterialEditor/Private/UI/Widgets/Editor/PropertySelectorLayouts/SDMMaterialPropertySelector_VerticalSlim.h
@@ -0,0 +1,21 @@
+// Copyright Epic Games, Inc. All Rights Reserved.
+
+#pragma once
+
+#include "UI/Widgets/Editor/PropertySelectorLayouts/SDMMaterialPropertySelector_VerticalBase.h"
+
+class SDMMaterialPropertySelector_VerticalSlim : public SDMMaterialPropertySelector_VerticalBase
+{
+	SLATE_BEGIN_ARGS(SDMMaterialPropertySelector_VerticalSlim) {}
+	SLATE_END_ARGS()
+
+public:
+	virtual ~SDMMaterialPropertySelector_VerticalSlim() override = default;
+
+	void Construct(const FArguments& InArgs, const TSharedRef<SDMMaterialEditor> InEditorWidget);
+
+protected:
+	//~ Begin SDMMaterialPropertySelector
+	virtual TSharedRef<SWidget> CreateSlot_SelectButton(EDMMaterialPropertyType InMaterialProperty) override;
+	//~ End SDMMaterialPropertySelector
+};
diff --git a/Engine/Plugins/Experimental/DynamicMaterial/Source/DynamicMaterialEditor/Private/UI/Widgets/Editor/PropertySelectorLayouts/SDMMaterialPropertySelector_Wrap.cpp b/Engine/Plugins/Experimental/DynamicMaterial/Source/DynamicMaterialEditor/Private/UI/Widgets/Editor/PropertySelectorLayouts/SDMMaterialPropertySelector_Wrap.cpp
new file mode 100644
--- /dev/null
+++ b/Engine/Plugins/Experimental/DynamicMaterial/Source/DynamicMaterialEditor/Private/UI/Widgets/Editor/PropertySelectorLayouts/SDMMaterialPropertySelector_Wrap.cpp
@@ -0,0 +1,15 @@
+// Copyright Epic Games, Inc. All Rights Reserved.
+
+#include "UI/Widgets/Editor/PropertySelectorLayouts/SDMMaterialPropertySelector_Wrap.h"
+
+#define LOCTEXT_NAMESPACE "SDMMaterialPropertySelector_Wrap"
+
+void SDMMaterialPropertySelector_Wrap::Construct(const FArguments& InArgs, const TSharedRef<SDMMaterialEditor> InEditorWidget)
+{
+	SDMMaterialPropertySelector_WrapBase::Construct(
+		SDMMaterialPropertySelector_WrapBase::FArguments(),
+		InEditorWidget
+	);
+}
+
+#undef LOCTEXT_NAMESPACE
diff --git a/Engine/Plugins/Experimental/DynamicMaterial/Source/DynamicMaterialEditor/Private/UI/Widgets/Editor/PropertySelectorLayouts/SDMMaterialPropertySelector_Wrap.h b/Engine/Plugins/Experimental/DynamicMaterial/Source/DynamicMaterialEditor/Private/UI/Widgets/Editor/PropertySelectorLayouts/SDMMaterialPropertySelector_Wrap.h
new file mode 100644
--- /dev/null
+++ b/Engine/Plugins/Experimental/DynamicMaterial/Source/DynamicMaterialEditor/Private/UI/Widgets/Editor/PropertySelectorLayouts/SDMMaterialPropertySelector_Wrap.h
@@ -0,0 +1,16 @@
+// Copyright Epic Games, Inc. All Rights Reserved.
+
+#pragma once
+
+#include "UI/Widgets/Editor/PropertySelectorLayouts/SDMMaterialPropertySelector_WrapBase.h"
+
+class SDMMaterialPropertySelector_Wrap : public SDMMaterialPropertySelector_WrapBase
+{
+	SLATE_BEGIN_ARGS(SDMMaterialPropertySelector_Wrap) {}
+	SLATE_END_ARGS()
+
+public:
+	virtual ~SDMMaterialPropertySelector_Wrap() override = default;
+
+	void Construct(const FArguments& InArgs, const TSharedRef<SDMMaterialEditor> InEditorWidget);
+};
diff --git a/Engine/Plugins/Experimental/DynamicMaterial/Source/DynamicMaterialEditor/Private/UI/Widgets/Editor/PropertySelectorLayouts/SDMMaterialPropertySelector_WrapBase.cpp b/Engine/Plugins/Experimental/DynamicMaterial/Source/DynamicMaterialEditor/Private/UI/Widgets/Editor/PropertySelectorLayouts/SDMMaterialPropertySelector_WrapBase.cpp
new file mode 100644
--- /dev/null
+++ b/Engine/Plugins/Experimental/DynamicMaterial/Source/DynamicMaterialEditor/Private/UI/Widgets/Editor/PropertySelectorLayouts/SDMMaterialPropertySelector_WrapBase.cpp
@@ -0,0 +1,133 @@
+// Copyright Epic Games, Inc. All Rights Reserved.
+
+#include "UI/Widgets/Editor/PropertySelectorLayouts/SDMMaterialPropertySelector_WrapBase.h"
+
+#include "Components/DMMaterialProperty.h"
+#include "DetailLayoutBuilder.h"
+#include "DMDefs.h"
+#include "Model/DynamicMaterialModelEditorOnlyData.h"
+#include "Utils/DMPrivate.h"
+#include "Widgets/Images/SImage.h"
+#include "Widgets/Input/SCheckBox.h"
+#include "Widgets/Layout/SBox.h"
+#include "Widgets/Layout/SWrapBox.h"
+#include "Widgets/SBoxPanel.h"
+#include "Widgets/Text/STextBlock.h"
+
+#define LOCTEXT_NAMESPACE "SDMMaterialPropertySelector_WrapBase"
+
+void SDMMaterialPropertySelector_WrapBase::Construct(const FArguments& InArgs, const TSharedRef<SDMMaterialEditor> InEditorWidget)
+{
+	SDMMaterialPropertySelector::Construct(
+		SDMMaterialPropertySelector::FArguments(),
+		InEditorWidget
+	);
+}
+
+TSharedRef<SWidget> SDMMaterialPropertySelector_WrapBase::CreateSlot_PropertyList()
+{
+	using namespace UE::DynamicMaterialEditor::Private;
+
+	TSharedRef<SWrapBox> NewSlotList = SNew(SWrapBox)
+		.InnerSlotPadding(FVector2D(6.f, 3.f))
+		.UseAllottedSize(true);
+
+	UDynamicMaterialModelEditorOnlyData* EditorOnlyData = GetEditorOnlyData();
+
+	if (!EditorOnlyData)
+	{
+		return NewSlotList;
+	}
+
+	NewSlotList->AddSlot()
+		[
+			SNew(SHorizontalBox)
+			+ SHorizontalBox::Slot()
+			.AutoWidth()
+			[
+				SNew(SBox)
+				.WidthOverride(20.f)
+			]
+			+ SHorizontalBox::Slot()
+			.AutoWidth()
+			[
+				CreateSlot_SelectButton(EDMMaterialPropertyType::None)
+			]
+		];
+
+	for (const TPair<EDMMaterialPropertyType, UDMMaterialProperty*>& PropertyPair : EditorOnlyData->GetMaterialProperties())
+	{
+		if (IsCustomMaterialProperty(PropertyPair.Key))
+		{
+			continue;
+		}
+
+		NewSlotList->AddSlot()
+			[
+				SNew(SHorizontalBox)
+				+ SHorizontalBox::Slot()
+				.AutoWidth()
+				[
+					CreateSlot_EnabledButton(PropertyPair.Key)
+				]
+				+ SHorizontalBox::Slot()
+				.AutoWidth()
+				[
+					CreateSlot_SelectButton(PropertyPair.Key)
+				]
+			];
+	}
+
+	return NewSlotList;
+}
+
+TSharedRef<SWidget> SDMMaterialPropertySelector_WrapBase::CreateSlot_SelectButton(EDMMaterialPropertyType InMaterialProperty)
+{
+	UEnum* PropertyEnum = StaticEnum<EDMMaterialPropertyType>();
+
+	const FText ButtonText = InMaterialProperty == EDMMaterialPropertyType::None
+		? LOCTEXT("GlobalSettings", "Global Settings")
+		: UE::DynamicMaterialEditor::Private::GetMaterialPropertyShortDisplayName(InMaterialProperty);
+
+	const FText Format = LOCTEXT("PropertySelectFormat", "Edit the {0} property.");
+
+	const FText ToolTip = (InMaterialProperty == EDMMaterialPropertyType::None)
+		? LOCTEXT("GeneralSettingsToolTip", "Edit the material global settings.")
+		: FText::Format(Format, PropertyEnum->GetDisplayNameTextByValue(static_cast<int64>(InMaterialProperty)));
+
+	return SNew(SCheckBox)
+		.Style(FAppStyle::Get(), "DetailsView.SectionButton")
+		.HAlign(EHorizontalAlignment::HAlign_Center)
+		.Padding(0.f)
+		.IsEnabled(this, &SDMMaterialPropertySelector_WrapBase::GetPropertySelectEnabled, InMaterialProperty)
+		.IsChecked(this, &SDMMaterialPropertySelector_WrapBase::GetPropertySelectState, InMaterialProperty)
+		.OnCheckStateChanged(this, &SDMMaterialPropertySelector_WrapBase::OnPropertySelectStateChanged, InMaterialProperty)
+		.ToolTipText(ToolTip)
+		.Content()
+		[
+			SNew(SHorizontalBox)
+			+SHorizontalBox::Slot()
+			.VAlign(VAlign_Center)
+			.AutoWidth()
+			[
+				SNew(SImage)
+				.Image(FAppStyle::Get().GetBrush("FilterBar.FilterImage"))
+				.ColorAndOpacity(this, &SDMMaterialPropertySelector_WrapBase::GetPropertySelectButtonChipColor, InMaterialProperty)
+			]
+			+SHorizontalBox::Slot()
+			.Padding(4.f, 6.f)
+			.VAlign(VAlign_Center)
+			.FillWidth(1.f)
+			[
+				SNew(SBox)
+				.WidthOverride(32.f)
+				[
+					SNew(STextBlock)
+					.Font(IDetailLayoutBuilder::GetDetailFont())
+					.Text(ButtonText)
+				]
+			]
+		];
+}
+
+#undef LOCTEXT_NAMESPACE
diff --git a/Engine/Plugins/Experimental/DynamicMaterial/Source/DynamicMaterialEditor/Private/UI/Widgets/Editor/PropertySelectorLayouts/SDMMaterialPropertySelector_WrapBase.h b/Engine/Plugins/Experimental/DynamicMaterial/Source/DynamicMaterialEditor/Private/UI/Widgets/Editor/PropertySelectorLayouts/SDMMaterialPropertySelector_WrapBase.h
new file mode 100644
--- /dev/null
+++ b/Engine/Plugins/Experimental/DynamicMaterial/Source/DynamicMaterialEditor/Private/UI/Widgets/Editor/PropertySelectorLayouts/SDMMaterialPropertySelector_WrapBase.h
@@ -0,0 +1,22 @@
+// Copyright Epic Games, Inc. All Rights Reserved.
+
+#pragma once
+
+#include "UI/Widgets/Editor/SDMMaterialPropertySelector.h"
+
+class SDMMaterialPropertySelector_WrapBase : public SDMMaterialPropertySelector
+{
+	SLATE_BEGIN_ARGS(SDMMaterialPropertySelector_WrapBase) {}
+	SLATE_END_ARGS()
+
+public:
+	virtual ~SDMMaterialPropertySelector_WrapBase() override = default;
+
+	void Construct(const FArguments& InArgs, const TSharedRef<SDMMaterialEditor> InEditorWidget);
+
+protected:
+	//~ Begin SDMMaterialPropertySelector
+	virtual TSharedRef<SWidget> CreateSlot_PropertyList() override;
+	virtual TSharedRef<SWidget> CreateSlot_SelectButton(EDMMaterialPropertyType InMaterialProperty) override;
+	//~ End SDMMaterialPropertySelector
+};
diff --git a/Engine/Plugins/Experimental/DynamicMaterial/Source/DynamicMaterialEditor/Private/UI/Widgets/Editor/PropertySelectorLayouts/SDMMaterialPropertySelector_WrapSlim.cpp b/Engine/Plugins/Experimental/DynamicMaterial/Source/DynamicMaterialEditor/Private/UI/Widgets/Editor/PropertySelectorLayouts/SDMMaterialPropertySelector_WrapSlim.cpp
new file mode 100644
--- /dev/null
+++ b/Engine/Plugins/Experimental/DynamicMaterial/Source/DynamicMaterialEditor/Private/UI/Widgets/Editor/PropertySelectorLayouts/SDMMaterialPropertySelector_WrapSlim.cpp
@@ -0,0 +1,60 @@
+// Copyright Epic Games, Inc. All Rights Reserved.
+
+#include "UI/Widgets/Editor/PropertySelectorLayouts/SDMMaterialPropertySelector_WrapSlim.h"
+
+#include "Components/DMMaterialProperty.h"
+#include "Model/DynamicMaterialModelEditorOnlyData.h"
+#include "Widgets/Layout/SWrapBox.h"
+
+#define LOCTEXT_NAMESPACE "SDMMaterialPropertySelector_WrapSlim"
+
+void SDMMaterialPropertySelector_WrapSlim::Construct(const FArguments& InArgs, const TSharedRef<SDMMaterialEditor> InEditorWidget)
+{
+	SDMMaterialPropertySelector_WrapBase::Construct(
+		SDMMaterialPropertySelector_WrapBase::FArguments(),
+		InEditorWidget
+	);
+}
+
+TSharedRef<SWidget> SDMMaterialPropertySelector_WrapSlim::CreateSlot_PropertyList()
+{
+	using namespace UE::DynamicMaterialEditor::Private;
+
+	TSharedRef<SWrapBox> NewSlotList = SNew(SWrapBox)
+		.InnerSlotPadding(FVector2D(6.f, 3.f))
+		.UseAllottedSize(true);
+
+	UDynamicMaterialModelEditorOnlyData* EditorOnlyData = GetEditorOnlyData();
+
+	if (!EditorOnlyData)
+	{
+		return NewSlotList;
+	}
+
+	NewSlotList->AddSlot()
+		[
+			CreateSlot_SelectButton(EDMMaterialPropertyType::None)
+		];
+
+	for (const TPair<EDMMaterialPropertyType, UDMMaterialProperty*>& PropertyPair : EditorOnlyData->GetMaterialProperties())
+	{
+		if (!PropertyPair.Value || !PropertyPair.Value->IsEnabled())
+		{
+			continue;
+		}
+
+		if (!EditorOnlyData->GetSlotForMaterialProperty(PropertyPair.Value->GetMaterialProperty()))
+		{
+			continue;
+		}
+
+		NewSlotList->AddSlot()
+			[
+				CreateSlot_SelectButton(PropertyPair.Key)
+			];
+	}
+
+	return NewSlotList;
+}
+
+#undef LOCTEXT_NAMESPACE
diff --git a/Engine/Plugins/Experimental/DynamicMaterial/Source/DynamicMaterialEditor/Private/UI/Widgets/Editor/PropertySelectorLayouts/SDMMaterialPropertySelector_WrapSlim.h b/Engine/Plugins/Experimental/DynamicMaterial/Source/DynamicMaterialEditor/Private/UI/Widgets/Editor/PropertySelectorLayouts/SDMMaterialPropertySelector_WrapSlim.h
new file mode 100644
--- /dev/null
+++ b/Engine/Plugins/Experimental/DynamicMaterial/Source/DynamicMaterialEditor/Private/UI/Widgets/Editor/PropertySelectorLayouts/SDMMaterialPropertySelector_WrapSlim.h
@@ -0,0 +1,21 @@
+// Copyright Epic Games, Inc. All Rights Reserved.
+
+#pragma once
+
+#include "UI/Widgets/Editor/PropertySelectorLayouts/SDMMaterialPropertySelector_WrapBase.h"
+
+class SDMMaterialPropertySelector_WrapSlim : public SDMMaterialPropertySelector_WrapBase
+{
+	SLATE_BEGIN_ARGS(SDMMaterialPropertySelector_WrapSlim) {}
+	SLATE_END_ARGS()
+
+public:
+	virtual ~SDMMaterialPropertySelector_WrapSlim() override = default;
+
+	void Construct(const FArguments& InArgs, const TSharedRef<SDMMaterialEditor> InEditorWidget);
+
+protected:
+	//~ Begin SDMMaterialPropertySelector_WrapBase
+	virtual TSharedRef<SWidget> CreateSlot_PropertyList() override;
+	//~ End SDMMaterialPropertySelector_WrapBase
+};
diff --git a/Engine/Plugins/Experimental/DynamicMaterial/Source/DynamicMaterialEditor/Private/UI/Widgets/Editor/SDMMaterialPropertySelector.cpp b/Engine/Plugins/Experimental/DynamicMaterial/Source/DynamicMaterialEditor/Private/UI/Widgets/Editor/SDMMaterialPropertySelector.cpp
--- a/Engine/Plugins/Experimental/DynamicMaterial/Source/DynamicMaterialEditor/Private/UI/Widgets/Editor/SDMMaterialPropertySelector.cpp
+++ b/Engine/Plugins/Experimental/DynamicMaterial/Source/DynamicMaterialEditor/Private/UI/Widgets/Editor/SDMMaterialPropertySelector.cpp
@@ -4,6 +4,7 @@
 
 #include "Components/DMMaterialProperty.h"
 #include "DetailLayoutBuilder.h"
+#include "DynamicMaterialEditorSettings.h"
 #include "Framework/Application/SlateApplication.h"
 #include "Model/DynamicMaterialModel.h"
 #include "Model/DynamicMaterialModelEditorOnlyData.h"
@@ -14,7 +15,6 @@
 #include "Widgets/Images/SImage.h"
 #include "Widgets/Input/SCheckBox.h"
 #include "Widgets/Layout/SBox.h"
-#include "Widgets/Layout/SGridPanel.h"
 #include "Widgets/Layout/SScrollBox.h"
 #include "Widgets/SBoxPanel.h"
 #include "Widgets/SNullWidget.h"
@@ -22,15 +22,6 @@
 
 #define LOCTEXT_NAMESPACE "SDMMaterialPropertySelector"
 
-namespace UE::DynamicMaterialEditor::Private
-{
-	namespace PropertySelectorColumns
-	{
-		constexpr int32 Enable = 0;
-		constexpr int32 Select = 1;
-	}
-}
-
 void SDMMaterialPropertySelector::PrivateRegisterAttributes(FSlateAttributeDescriptor::FInitializer&)
 {
 }
@@ -120,52 +111,6 @@ TSharedPtr<SDMMaterialSlotEditor> SDMMaterialPropertySelector::GetSlotEditorWidg
 	return EditorWidget->GetSlotEditorWidget();
 }
 
-TSharedRef<SWidget> SDMMaterialPropertySelector::CreateSlot_PropertyList()
-{
-	using namespace UE::DynamicMaterialEditor::Private;
-
-	TSharedRef<SGridPanel> NewSlotList = SNew(SGridPanel)
-		.FillColumn(PropertySelectorColumns::Select, 1.f);
-
-	UDynamicMaterialModelEditorOnlyData* EditorOnlyData = GetEditorOnlyData();
-
-	if (!EditorOnlyData)
-	{
-		return NewSlotList;
-	}
-
-	int32 Row = 0;
-
-	NewSlotList->AddSlot(PropertySelectorColumns::Select, Row)
-		[
-			CreateSlot_SelectButton(EDMMaterialPropertyType::None)
-		];
-
-	++Row;
-
-	for (const TPair<EDMMaterialPropertyType, UDMMaterialProperty*>& PropertyPair : EditorOnlyData->GetMaterialProperties())
-	{
-		if (IsCustomMaterialProperty(PropertyPair.Key))
-		{
-			continue;
-		}
-
-		NewSlotList->AddSlot(PropertySelectorColumns::Enable, Row)
-			[
-				CreateSlot_EnabledButton(PropertyPair.Key)
-			];
-
-		NewSlotList->AddSlot(PropertySelectorColumns::Select, Row)
-			[
-				CreateSlot_SelectButton(PropertyPair.Key)
-			];
-
-		++Row;
-	}
-
-	return NewSlotList;
-}
-
 TSharedRef<SWidget> SDMMaterialPropertySelector::CreateSlot_EnabledButton(EDMMaterialPropertyType InMaterialProperty)
 {
 	const FText Format = LOCTEXT("PropertyEnableFormat", "Toggle the {0} property.\n\nProperty must be valid for the Material Type.");
@@ -182,9 +127,34 @@ TSharedRef<SWidget> SDMMaterialPropertySelector::CreateSlot_EnabledButton(EDMMat
 
 TSharedRef<SWidget> SDMMaterialPropertySelector::CreateSlot_SelectButton(EDMMaterialPropertyType InMaterialProperty)
 {
+	EDMMaterialEditorLayout Layout = EDMMaterialEditorLayout::First;
+
+	if (UDynamicMaterialEditorSettings* Settings = UDynamicMaterialEditorSettings::Get())
+	{
+		Layout = Settings->Layout;
+	}
+
+	UEnum* PropertyEnum = StaticEnum<EDMMaterialPropertyType>();
+
+	auto UseShortName = [Layout]()
+		{
+			switch (Layout)
+			{
+				case EDMMaterialEditorLayout::LeftSlim:
+					return true;
+
+				default:
+					return false;
+			}
+		};
+
+	const bool bUseShortName = UseShortName();
+
 	const FText ButtonText = InMaterialProperty == EDMMaterialPropertyType::None
 		? LOCTEXT("GlobalSettings", "Global Settings")
-		: StaticEnum<EDMMaterialPropertyType>()->GetDisplayNameTextByValue(static_cast<int64>(InMaterialProperty));
+		: (bUseShortName 
+			? UE::DynamicMaterialEditor::Private::GetMaterialPropertyShortDisplayName(InMaterialProperty)
+			: PropertyEnum->GetDisplayNameTextByValue(static_cast<int64>(InMaterialProperty)));
 
 	const FText Format = LOCTEXT("PropertySelectFormat", "Edit the {0} property.");
 	UEnum* MaterialPropertyEnum = StaticEnum<EDMMaterialPropertyType>();
@@ -204,7 +174,8 @@ TSharedRef<SWidget> SDMMaterialPropertySelector::CreateSlot_SelectButton(EDMMate
 		.Content()
 		[
 			SNew(SBox)
-			.WidthOverride(135.f)
+			.WidthOverride(bUseShortName ? 42.f : 135.f)
+			.Clipping(EWidgetClipping::ClipToBounds)
 			[
 				SNew(SHorizontalBox)
 				+SHorizontalBox::Slot()
@@ -216,7 +187,7 @@ TSharedRef<SWidget> SDMMaterialPropertySelector::CreateSlot_SelectButton(EDMMate
 					.ColorAndOpacity(this, &SDMMaterialPropertySelector::GetPropertySelectButtonChipColor, InMaterialProperty)
 				]
 				+SHorizontalBox::Slot()
-				.Padding(10.f, 6.f)
+				.Padding(bUseShortName ? 2.f : 10.f, 6.f)
 				.VAlign(VAlign_Center)
 				.FillWidth(1.f)
 				[
diff --git a/Engine/Plugins/Experimental/DynamicMaterial/Source/DynamicMaterialEditor/Private/UI/Widgets/Editor/SDMMaterialPropertySelector.h b/Engine/Plugins/Experimental/DynamicMaterial/Source/DynamicMaterialEditor/Private/UI/Widgets/Editor/SDMMaterialPropertySelector.h
--- a/Engine/Plugins/Experimental/DynamicMaterial/Source/DynamicMaterialEditor/Private/UI/Widgets/Editor/SDMMaterialPropertySelector.h
+++ b/Engine/Plugins/Experimental/DynamicMaterial/Source/DynamicMaterialEditor/Private/UI/Widgets/Editor/SDMMaterialPropertySelector.h
@@ -16,6 +16,15 @@ class UDynamicMaterialModelEditorOnlyData;
 enum class ECheckBoxState : uint8;
 enum class EDMMaterialPropertyType : uint8;
 
+namespace UE::DynamicMaterialEditor::Private
+{
+	namespace PropertySelectorColumns
+	{
+		constexpr int32 Enable = 0;
+		constexpr int32 Select = 1;
+	}
+}
+
 class SDMMaterialPropertySelector : public SCompoundWidget
 {
 	SLATE_DECLARE_WIDGET(SDMMaterialPropertySelector, SCompoundWidget)
@@ -44,11 +53,11 @@ protected:
 
 	TSharedPtr<SDMMaterialSlotEditor> GetSlotEditorWidget() const;
 
-	TSharedRef<SWidget> CreateSlot_PropertyList();
+	virtual TSharedRef<SWidget> CreateSlot_PropertyList() = 0;
 
 	TSharedRef<SWidget> CreateSlot_EnabledButton(EDMMaterialPropertyType InMaterialProperty);
 
-	TSharedRef<SWidget> CreateSlot_SelectButton(EDMMaterialPropertyType InMaterialProperty);
+	virtual TSharedRef<SWidget> CreateSlot_SelectButton(EDMMaterialPropertyType InMaterialProperty) = 0;
 
 	bool IsPropertyEnabled(EDMMaterialPropertyType InMaterialProperty) const;
 
diff --git a/Engine/Plugins/Experimental/DynamicMaterial/Source/DynamicMaterialEditor/Private/UI/Widgets/SDMMaterialDesigner.cpp b/Engine/Plugins/Experimental/DynamicMaterial/Source/DynamicMaterialEditor/Private/UI/Widgets/SDMMaterialDesigner.cpp
--- a/Engine/Plugins/Experimental/DynamicMaterial/Source/DynamicMaterialEditor/Private/UI/Widgets/SDMMaterialDesigner.cpp
+++ b/Engine/Plugins/Experimental/DynamicMaterial/Source/DynamicMaterialEditor/Private/UI/Widgets/SDMMaterialDesigner.cpp
@@ -4,6 +4,9 @@
 
 #include "Containers/Set.h"
 #include "DynamicMaterialEditorSettings.h"
+#include "DynamicMaterialModule.h"
+#include "Editor/SDMMaterialComponentEditor.h"
+#include "Editor/SDMMaterialSlotEditor.h"
 #include "GameFramework/Actor.h"
 #include "Material/DynamicMaterialInstance.h"
 #include "Model/DynamicMaterialModel.h"
@@ -11,6 +14,11 @@
 #include "Model/DynamicMaterialModelEditorOnlyData.h"
 #include "SAssetDropTarget.h"
 #include "UI/Utils/DMDropTargetPrivateSetter.h"
+#include "UI/Widgets/Editor/EditorLayouts/SDMMaterialEditor_Left.h"
+#include "UI/Widgets/Editor/EditorLayouts/SDMMaterialEditor_LeftSlim.h"
+#include "UI/Widgets/Editor/EditorLayouts/SDMMaterialEditor_TopHorizontal.h"
+#include "UI/Widgets/Editor/EditorLayouts/SDMMaterialEditor_TopSlim.h"
+#include "UI/Widgets/Editor/EditorLayouts/SDMMaterialEditor_TopVertical.h"
 #include "UI/Widgets/SDMActorMaterialSelector.h"
 #include "UI/Widgets/SDMMaterialEditor.h"
 #include "UI/Widgets/SDMMaterialSelectPrompt.h"
@@ -22,6 +30,19 @@ void SDMMaterialDesigner::PrivateRegisterAttributes(FSlateAttributeDescriptor::F
 {
 }
 
+SDMMaterialDesigner::~SDMMaterialDesigner()
+{
+	if (!FDynamicMaterialModule::AreUObjectsSafe())
+	{
+		return;
+	}
+
+	if (UDynamicMaterialEditorSettings* Settings = GetMutableDefault<UDynamicMaterialEditorSettings>())
+	{
+		Settings->GetOnSettingsChanged().RemoveAll(this);
+	}
+}
+
 void SDMMaterialDesigner::Construct(const FArguments& InArgs)
 {
 	SetCanTick(true);
@@ -29,6 +50,11 @@ void SDMMaterialDesigner::Construct(const FArguments& InArgs)
 	ContentSlot = TDMWidgetSlot<SWidget>(SharedThis(this), 0, SNullWidget::NullWidget);
 
 	SetSelectPromptView();
+
+	if (UDynamicMaterialEditorSettings* Settings = GetMutableDefault<UDynamicMaterialEditorSettings>())
+	{
+		Settings->GetOnSettingsChanged().AddSP(this, &SDMMaterialDesigner::OnSettingsChanged);
+	}
 }
 
 bool SDMMaterialDesigner::OpenMaterialModelBase(UDynamicMaterialModelBase* InMaterialModelBase)
@@ -251,18 +277,26 @@ void SDMMaterialDesigner::SetWizardView(const FDMObjectMaterialProperty& InObjec
 
 void SDMMaterialDesigner::SetEditorView(UDynamicMaterialModelBase* InMaterialModelBase)
 {
-	TSharedRef<SDMMaterialEditor> Editor = SNew(SDMMaterialEditor, SharedThis(this))
-		.MaterialModelBase(InMaterialModelBase);
+	EDMMaterialEditorLayout Layout = EDMMaterialEditorLayout::Left;
+
+	if (UDynamicMaterialEditorSettings* Settings = GetMutableDefault<UDynamicMaterialEditorSettings>())
+	{
+		Layout = Settings->Layout;
+	}
 
-	SetWidget(Editor, /* Include Drop Target */ true);
+	SetEditorLayout(Layout, InMaterialModelBase);
 }
 
 void SDMMaterialDesigner::SetEditorView(const FDMObjectMaterialProperty& InObjectMaterialProperty)
 {
-	TSharedRef<SDMMaterialEditor> Editor = SNew(SDMMaterialEditor, SharedThis(this))
-		.MaterialProperty(InObjectMaterialProperty);
+	EDMMaterialEditorLayout Layout = EDMMaterialEditorLayout::Left;
 
-	SetWidget(Editor, /* Include Drop Target */ true);
+	if (UDynamicMaterialEditorSettings* Settings = GetMutableDefault<UDynamicMaterialEditorSettings>())
+	{
+		Layout = Settings->Layout;
+	}
+
+	SetEditorLayout(Layout, InObjectMaterialProperty);
 }
 
 void SDMMaterialDesigner::SetWidget(const TSharedRef<SWidget>& InWidget, bool bInIncludeAssetDropTarget)
@@ -378,3 +412,181 @@ void SDMMaterialDesigner::OnAssetsDropped(const FDragDropEvent& InDragDropEvent,
 		}
 	}
 }
+
+void SDMMaterialDesigner::OnSettingsChanged(const FPropertyChangedEvent& InPropertyChangedEvent)
+{
+	if (!Content.IsValid() || Content->GetWidgetClass().GetWidgetType() != SDMMaterialEditor::StaticWidgetClass().GetWidgetType())
+	{
+		return;
+	}
+
+	const FName MemberName = InPropertyChangedEvent.GetMemberPropertyName();
+
+	if (MemberName == GET_MEMBER_NAME_CHECKED(UDynamicMaterialEditorSettings, Layout))
+	{
+		OnLayoutChanged();
+	}
+}
+
+void SDMMaterialDesigner::OnLayoutChanged()
+{
+	UDynamicMaterialEditorSettings* Settings = GetMutableDefault<UDynamicMaterialEditorSettings>();
+
+	if (!Settings)
+	{
+		return;
+	}
+
+	// Already assured
+	TSharedPtr<SDMMaterialEditor> CurrentEditor = StaticCastSharedRef<SDMMaterialEditor>(Content.ToSharedRef());
+
+	UDynamicMaterialModelBase* MaterialModelBase = CurrentEditor->GetMaterialModelBase();
+	const FDMObjectMaterialProperty* MaterialObjectProperty = CurrentEditor->GetMaterialObjectProperty();
+	UDMMaterialSlot* EditedSlot = nullptr;
+	UDMMaterialComponent* EditedComponent = nullptr;
+
+	if (TSharedPtr<SDMMaterialSlotEditor> SlotEditorWidget = CurrentEditor->GetSlotEditorWidget())
+	{
+		EditedSlot = SlotEditorWidget->GetSlot();
+	}
+
+	if (TSharedPtr<SDMMaterialComponentEditor> ComponentEditorWidget = CurrentEditor->GetComponentEditorWidget())
+	{
+		EditedComponent = ComponentEditorWidget->GetComponent();
+	}
+
+	if (MaterialObjectProperty)
+	{
+		if (!SetEditorLayout(Settings->Layout, *MaterialObjectProperty))
+		{
+			return;
+		}
+	}
+	else
+	{
+		if (!SetEditorLayout(Settings->Layout, MaterialModelBase))
+		{
+			return;
+		}
+	}
+
+	TSharedRef<SDMMaterialEditor> NewEditor = StaticCastSharedRef<SDMMaterialEditor>(Content.ToSharedRef());
+
+	if (EditedSlot)
+	{
+		NewEditor->EditSlot(EditedSlot);
+	}
+
+	if (EditedComponent)
+	{
+		NewEditor->EditComponent(EditedComponent);
+	}
+}
+
+bool SDMMaterialDesigner::SetEditorLayout(EDMMaterialEditorLayout InLayout, UDynamicMaterialModelBase* InMaterialModelBase)
+{
+	TSharedPtr<SDMMaterialEditor> NewEditor;
+
+	switch (InLayout)
+	{
+		case EDMMaterialEditorLayout::Left:
+			NewEditor = SNew(SDMMaterialEditor_Left, SharedThis(this))
+				.MaterialModelBase(InMaterialModelBase);
+			break;
+
+		case EDMMaterialEditorLayout::LeftAutoHide:
+			break;
+
+		case EDMMaterialEditorLayout::LeftSlim:
+			NewEditor = SNew(SDMMaterialEditor_LeftSlim, SharedThis(this))
+				.MaterialModelBase(InMaterialModelBase);
+			break;
+
+		case EDMMaterialEditorLayout::TopHorizontal:
+			NewEditor = SNew(SDMMaterialEditor_TopHorizontal, SharedThis(this))
+				.MaterialModelBase(InMaterialModelBase);
+			break;
+
+		case EDMMaterialEditorLayout::TopHorizontalAutoHide:
+			break;
+
+		case EDMMaterialEditorLayout::TopVertical:
+			NewEditor = SNew(SDMMaterialEditor_TopVertical, SharedThis(this))
+				.MaterialModelBase(InMaterialModelBase);
+			break;
+
+		case EDMMaterialEditorLayout::TopVerticalAutoHide:
+			break;
+
+		case EDMMaterialEditorLayout::TopSlim:
+			NewEditor = SNew(SDMMaterialEditor_TopSlim, SharedThis(this))
+				.MaterialModelBase(InMaterialModelBase);
+			break;
+
+		default:
+			break;
+	}
+
+	if (!NewEditor.IsValid())
+	{
+		return false;
+	}
+
+	SetWidget(NewEditor.ToSharedRef(), /* Include Drop Target */ true);
+
+	return true;
+}
+
+bool SDMMaterialDesigner::SetEditorLayout(EDMMaterialEditorLayout InLayout, const FDMObjectMaterialProperty& InObjectMaterialProperty)
+{
+	TSharedPtr<SDMMaterialEditor> NewEditor;
+
+	switch (InLayout)
+	{
+		case EDMMaterialEditorLayout::Left:
+			NewEditor = SNew(SDMMaterialEditor_Left, SharedThis(this))
+				.MaterialProperty(InObjectMaterialProperty);
+			break;
+
+		case EDMMaterialEditorLayout::LeftAutoHide:
+			break;
+
+		case EDMMaterialEditorLayout::LeftSlim:
+			NewEditor = SNew(SDMMaterialEditor_LeftSlim, SharedThis(this))
+				.MaterialProperty(InObjectMaterialProperty);
+			break;
+
+		case EDMMaterialEditorLayout::TopHorizontal:
+			NewEditor = SNew(SDMMaterialEditor_TopHorizontal, SharedThis(this))
+				.MaterialProperty(InObjectMaterialProperty);
+			break;
+
+		case EDMMaterialEditorLayout::TopHorizontalAutoHide:
+			break;
+
+		case EDMMaterialEditorLayout::TopVertical:
+			NewEditor = SNew(SDMMaterialEditor_TopVertical, SharedThis(this))
+				.MaterialProperty(InObjectMaterialProperty);
+			break;
+
+		case EDMMaterialEditorLayout::TopVerticalAutoHide:
+			break;
+
+		case EDMMaterialEditorLayout::TopSlim:
+			NewEditor = SNew(SDMMaterialEditor_TopSlim, SharedThis(this))
+				.MaterialProperty(InObjectMaterialProperty);
+			break;
+
+		default:
+			break;
+	}
+
+	if (!NewEditor.IsValid())
+	{
+		return false;
+	}
+
+	SetWidget(NewEditor.ToSharedRef(), /* Include Drop Target */ true);
+
+	return true;
+}
diff --git a/Engine/Plugins/Experimental/DynamicMaterial/Source/DynamicMaterialEditor/Private/UI/Widgets/SDMMaterialDesigner.h b/Engine/Plugins/Experimental/DynamicMaterial/Source/DynamicMaterialEditor/Private/UI/Widgets/SDMMaterialDesigner.h
--- a/Engine/Plugins/Experimental/DynamicMaterial/Source/DynamicMaterialEditor/Private/UI/Widgets/SDMMaterialDesigner.h
+++ b/Engine/Plugins/Experimental/DynamicMaterial/Source/DynamicMaterialEditor/Private/UI/Widgets/SDMMaterialDesigner.h
@@ -14,7 +14,9 @@ class SDMMaterialWizard;
 class UDynamicMaterialInstance;
 class UDynamicMaterialModel;
 class UDynamicMaterialModelBase;
+enum class EDMMaterialEditorLayout : uint8;
 struct FDMObjectMaterialProperty;
+struct FPropertyChangedEvent;
 
 class SDMMaterialDesigner : public SCompoundWidget
 {
@@ -26,7 +28,7 @@ class SDMMaterialDesigner : public SCompoundWidget
 public:
 	static bool IsFollowingSelection();
 
-	virtual ~SDMMaterialDesigner() override = default;
+	virtual ~SDMMaterialDesigner() override;
 
 	void Construct(const FArguments& InArgs);
 
@@ -96,4 +98,13 @@ protected:
 	bool OnAssetDraggedOver(TArrayView<FAssetData> InAssets);
 
 	void OnAssetsDropped(const FDragDropEvent& InDragDropEvent, TArrayView<FAssetData> InAssets);
+
+	/** Events */
+	void OnSettingsChanged(const FPropertyChangedEvent& InPropertyChangedEvent);
+
+	void OnLayoutChanged();
+
+	bool SetEditorLayout(EDMMaterialEditorLayout InLayout, UDynamicMaterialModelBase* InMaterialModelBase);
+
+	bool SetEditorLayout(EDMMaterialEditorLayout InLayout, const FDMObjectMaterialProperty& InObjectMaterialProperty);
 };
diff --git a/Engine/Plugins/Experimental/DynamicMaterial/Source/DynamicMaterialEditor/Private/UI/Widgets/SDMMaterialEditor.cpp b/Engine/Plugins/Experimental/DynamicMaterial/Source/DynamicMaterialEditor/Private/UI/Widgets/SDMMaterialEditor.cpp
--- a/Engine/Plugins/Experimental/DynamicMaterial/Source/DynamicMaterialEditor/Private/UI/Widgets/SDMMaterialEditor.cpp
+++ b/Engine/Plugins/Experimental/DynamicMaterial/Source/DynamicMaterialEditor/Private/UI/Widgets/SDMMaterialEditor.cpp
@@ -32,37 +32,6 @@
 
 #define LOCTEXT_NAMESPACE "SDMMaterialEditor"
 
-namespace UE::DynamicMaterialEditor::Private
-{
-	namespace SlotList
-	{
-		constexpr int32 ToolBar = 0;
-		constexpr int32 MainLayout = 1;
-		constexpr int32 StatusBar = 2;
-
-		namespace Main
-		{
-			constexpr int32 Left = 0;
-			constexpr int32 Right = 1;
-		}
-
-		namespace Left
-		{
-			constexpr int32 Preview = 0;
-			constexpr int32 PropertySelector = 1;
-		}
-
-		namespace Right
-		{
-			constexpr int32 SlotEditor = 0;
-			constexpr int32 ComponentEditor = 1;
-		}
-	}
-
-	const TCHAR* EditorDarkBackground = TEXT("Brushes.Title");
-	const TCHAR* EditorLightBackground = TEXT("Brushes.Header");
-}
-
 void SDMMaterialEditor::PrivateRegisterAttributes(FSlateAttributeDescriptor::FInitializer&)
 {
 }
@@ -239,8 +208,6 @@ void SDMMaterialEditor::EditSlot(UDMMaterialSlot* InSlot, bool bInForceRefresh)
 		return;
 	}
 
-	RightSlot.Invalidate();
-
 	SlotEditorSlot.Invalidate();
 	SplitterSlot = nullptr;
 	SlotToEdit = InSlot;
@@ -272,7 +239,6 @@ void SDMMaterialEditor::EditComponent(UDMMaterialComponent* InComponent, bool bI
 
 	if (bGlobalSettingsMode)
 	{
-		RightSlot.Invalidate();
 		SlotEditorSlot.Invalidate();
 		SplitterSlot = nullptr;
 	}
@@ -292,7 +258,6 @@ void SDMMaterialEditor::EditGlobalSettings(bool bInForceRefresh)
 
 	if (!bGlobalSettingsMode)
 	{
-		RightSlot.Invalidate();
 		SlotEditorSlot.Invalidate();
 		SplitterSlot = nullptr;
 		ComponentEditorSlot.Invalidate();
@@ -500,28 +465,19 @@ void SDMMaterialEditor::ValidateSlots()
 	}
 	else
 	{
-		if (LeftSlot.HasBeenInvalidated())
-		{
-			LeftSlot << CreateSlot_Left();
-		}
-		else
-		{
-			if (MaterialPreviewSlot.HasBeenInvalidated())
-			{
-				MaterialPreviewSlot << CreateSlot_Preview();
-			}
+		ValidateSlots_Main();
 
-			if (PropertySelectorSlot.HasBeenInvalidated())
-			{
-				PropertySelectorSlot << CreateSlot_PropertySelector();
-			}
+		if (MaterialPreviewSlot.HasBeenInvalidated())
+		{
+			MaterialPreviewSlot << CreateSlot_Preview();
 		}
 
-		if (RightSlot.HasBeenInvalidated())
+		if (PropertySelectorSlot.HasBeenInvalidated())
 		{
-			RightSlot << CreateSlot_Right();
+			PropertySelectorSlot << CreateSlot_PropertySelector();
 		}
-		else if (bGlobalSettingsMode)
+
+		if (bGlobalSettingsMode)
 		{
 			if (GlobalSettingsEditorSlot.HasBeenInvalidated())
 			{
@@ -565,14 +521,15 @@ void SDMMaterialEditor::ClearSlots()
 	ContentSlot.ClearWidget();
 	ToolBarSlot.ClearWidget();
 	MainSlot.ClearWidget();
-	LeftSlot.ClearWidget();
-	RightSlot.ClearWidget();
+	SlotEditorSlot.ClearWidget();
 	MaterialPreviewSlot.ClearWidget();
 	PropertySelectorSlot.ClearWidget();
-	SlotEditorSlot.ClearWidget();
+	GlobalSettingsEditorSlot.ClearWidget();
 	SplitterSlot = nullptr;
 	ComponentEditorSlot.ClearWidget();
 	StatusBarSlot.ClearWidget();
+
+	ClearSlots_Main();
 }
 
 void SDMMaterialEditor::CreateLayout()
@@ -626,179 +583,11 @@ TSharedRef<SDMToolBar> SDMMaterialEditor::CreateSlot_ToolBar()
 	);
 }
 
-TSharedRef<SWidget> SDMMaterialEditor::CreateSlot_Main()
-{
-	SHorizontalBox::FSlot* LeftSlotPtr = nullptr;
-	SHorizontalBox::FSlot* RightSlotPtr = nullptr;
-
-	TSharedRef<SHorizontalBox> NewMain = SNew(SHorizontalBox)
-		+ SHorizontalBox::Slot()
-		.Expose(LeftSlotPtr)
-		.AutoWidth()
-		[
-			SNullWidget::NullWidget
-		]
-
-		+ SHorizontalBox::Slot()
-		.Expose(RightSlotPtr)
-		.FillWidth(1.0f)
-		[
-			SNullWidget::NullWidget
-		];
-
-	LeftSlot = TDMWidgetSlot<SWidget>(LeftSlotPtr, CreateSlot_Left());
-	RightSlot = TDMWidgetSlot<SWidget>(RightSlotPtr, CreateSlot_Right());
-
-	return NewMain;
-}
-
-TSharedRef<SWidget> SDMMaterialEditor::CreateSlot_Left()
-{
-	using namespace UE::DynamicMaterialEditor::Private;
-
-	SVerticalBox::FSlot* MaterialPreviewSlotPtr = nullptr;
-	SVerticalBox::FSlot* PropertySelectorSlotPtr = nullptr;
-
-	TSharedRef<SWidget> NewLeft = SNew(SBorder)
-		.BorderImage(FAppStyle::GetBrush(EditorDarkBackground))
-		.Padding(5.f)
-		[
-			SNew(SVerticalBox)
-			+ SVerticalBox::Slot()
-			.Expose(MaterialPreviewSlotPtr)
-			.AutoHeight()
-			.Padding(0.f)
-			[
-				SNullWidget::NullWidget
-			]
-
-			+ SVerticalBox::Slot()
-			.Expose(PropertySelectorSlotPtr)
-			.FillHeight(1.0f)
-			.Padding(0.f, 5.f, 0.f, 0.f)
-			[
-				SNullWidget::NullWidget
-			]
-		];		
-
-	MaterialPreviewSlot = TDMWidgetSlot<SDMMaterialPreview>(MaterialPreviewSlotPtr, CreateSlot_Preview());
-	PropertySelectorSlot = TDMWidgetSlot<SDMMaterialPropertySelector>(PropertySelectorSlotPtr, CreateSlot_PropertySelector());
-
-	return NewLeft;
-}
-
-TSharedRef<SWidget> SDMMaterialEditor::CreateSlot_Right()
-{
-	using namespace UE::DynamicMaterialEditor::Private;
-
-	const bool bHasSlotToEdit = SlotToEdit.IsValid();
-
-	if (!bGlobalSettingsMode && !bHasSlotToEdit)
-	{
-		bGlobalSettingsMode = true;
-	}
-	else if (bHasSlotToEdit)
-	{
-		bGlobalSettingsMode = false;
-	}
-
-	return SNew(SBorder)
-		.BorderImage(FAppStyle::GetBrush(EditorDarkBackground))
-		.Padding(FMargin(0.f, 5.f))
-		[
-			bGlobalSettingsMode
-				? CreateSlot_Right_GlobalSettings()
-				: CreateSlot_Right_Slot()
-		];
-}
-
-TSharedRef<SWidget> SDMMaterialEditor::CreateSlot_Right_GlobalSettings()
-{
-	using namespace UE::DynamicMaterialEditor::Private;
-
-	SScrollBox::FSlot* GlobalSettingsSlotPtr = nullptr;
-
-	TSharedRef<SBorder> NewRight = SNew(SBorder)
-		.BorderImage(FAppStyle::GetBrush(EditorLightBackground))
-		.Padding(0.f)
-		[
-			SNew(SScrollBox)
-			+ SScrollBox::Slot()
-			.Expose(GlobalSettingsSlotPtr)
-			.VAlign(EVerticalAlignment::VAlign_Fill)
-			[
-				SNullWidget::NullWidget
-			]
-		];
-
-	GlobalSettingsEditorSlot = TDMWidgetSlot<SDMMaterialGlobalSettingsEditor>(GlobalSettingsSlotPtr, CreateSlot_GlobalSettingsEditor());
-
-	return NewRight;
-}
-
 TSharedRef<SDMMaterialGlobalSettingsEditor> SDMMaterialEditor::CreateSlot_GlobalSettingsEditor()
 {
 	return SNew(SDMMaterialGlobalSettingsEditor, SharedThis(this), GetMaterialModelBase());
 }
 
-TSharedRef<SWidget> SDMMaterialEditor::CreateSlot_Right_Slot()
-{
-	using namespace UE::DynamicMaterialEditor::Private;
-
-	float SplitterValue = 0.5;
-
-	if (UDynamicMaterialEditorSettings* Settings = UDynamicMaterialEditorSettings::Get())
-	{
-		SplitterValue = Settings->SplitterLocation;
-	}
-
-	TSharedPtr<SBorder> TopBox;
-	TSharedPtr<SBorder> BottomBox;
-
-	SSplitter::FSlot* ExplosedSlot = nullptr;
-
-	TSharedRef<SSplitter> NewRight = SNew(SSplitter)
-		.Style(FAppStyle::Get(), "DetailsView.Splitter")
-		.Orientation(Orient_Vertical)
-		.ResizeMode(ESplitterResizeMode::Fill)
-		.PhysicalSplitterHandleSize(5.0f)
-		.HitDetectionSplitterHandleSize(5.0f)
-		.OnSplitterFinishedResizing(this, &SDMMaterialEditor::OnRightSlotSplitterResized)
-
-		+ SSplitter::Slot()
-		.Expose(ExplosedSlot)
-		.Resizable(true)
-		.SizeRule(SSplitter::ESizeRule::FractionOfParent)
-		.MinSize(165)
-		.Value(SplitterValue)
-		[
-			SAssignNew(TopBox, SBorder)
-			.BorderImage(FAppStyle::GetBrush(EditorLightBackground))
-			[
-				SNullWidget::NullWidget
-			]
-		]
-
-		+ SSplitter::Slot()
-		.Resizable(true)
-		.SizeRule(SSplitter::ESizeRule::FractionOfParent)
-		.MinSize(60)
-		.Value(1.f - SplitterValue)
-		[
-			SAssignNew(BottomBox, SBorder)
-			.BorderImage(FAppStyle::GetBrush(EditorLightBackground))
-			[
-				SNullWidget::NullWidget
-			]
-		];
-
-	SplitterSlot = ExplosedSlot;
-	SlotEditorSlot = TDMWidgetSlot<SDMMaterialSlotEditor>(TopBox.ToSharedRef(), 0, CreateSlot_SlotEditor());
-	ComponentEditorSlot = TDMWidgetSlot<SDMMaterialComponentEditor>(BottomBox.ToSharedRef(), 0, CreateSlot_ComponentEditor());
-
-	return NewRight;
-}
-
 TSharedRef<SDMMaterialPreview> SDMMaterialEditor::CreateSlot_Preview()
 {
 	return SNew(SDMMaterialPreview, SharedThis(this), GetMaterialModelBase());
@@ -806,7 +595,7 @@ TSharedRef<SDMMaterialPreview> SDMMaterialEditor::CreateSlot_Preview()
 
 TSharedRef<SDMMaterialPropertySelector> SDMMaterialEditor::CreateSlot_PropertySelector()
 {
-	TSharedRef<SDMMaterialPropertySelector> NewPropertySelector = SNew(SDMMaterialPropertySelector, SharedThis(this));
+	TSharedRef<SDMMaterialPropertySelector> NewPropertySelector = CreateSlot_PropertySelector_Impl();
 
 	if (!PropertyToSelect.IsSet())
 	{
@@ -893,7 +682,7 @@ void SDMMaterialEditor::OnEnginePreExit()
 	MaterialPreviewSlot.ClearWidget();
 }
 
-void SDMMaterialEditor::OnRightSlotSplitterResized()
+void SDMMaterialEditor::OnEditorSplitterResized()
 {
 	UDynamicMaterialEditorSettings* Settings = UDynamicMaterialEditorSettings::Get();
 
diff --git a/Engine/Plugins/Experimental/DynamicMaterial/Source/DynamicMaterialEditor/Private/UI/Widgets/SDMMaterialEditor.h b/Engine/Plugins/Experimental/DynamicMaterial/Source/DynamicMaterialEditor/Private/UI/Widgets/SDMMaterialEditor.h
--- a/Engine/Plugins/Experimental/DynamicMaterial/Source/DynamicMaterialEditor/Private/UI/Widgets/SDMMaterialEditor.h
+++ b/Engine/Plugins/Experimental/DynamicMaterial/Source/DynamicMaterialEditor/Private/UI/Widgets/SDMMaterialEditor.h
@@ -28,6 +28,12 @@ class UDynamicMaterialModelBase;
 enum class EDMMaterialPropertyType : uint8;
 enum class EDMUpdateType : uint8;
 
+namespace UE::DynamicMaterialEditor::Private
+{
+	inline const TCHAR* EditorDarkBackground = TEXT("Brushes.Title");
+	inline const TCHAR* EditorLightBackground = TEXT("Brushes.Header");
+}
+
 class SDMMaterialEditor : public SCompoundWidget, public FSelfRegisteringEditorUndoClient
 {
 	SLATE_DECLARE_WIDGET(SDMMaterialEditor, SCompoundWidget)
@@ -74,11 +80,11 @@ public:
 	/** Actions */
 	void SelectProperty(EDMMaterialPropertyType InProperty, bool bInForceRefresh = false);
 
-	void EditSlot(UDMMaterialSlot* InSlot, bool bInForceRefresh = false);
+	virtual void EditSlot(UDMMaterialSlot* InSlot, bool bInForceRefresh = false);
 
-	void EditComponent(UDMMaterialComponent* InComponent, bool bInForceRefresh = false);
+	virtual void EditComponent(UDMMaterialComponent* InComponent, bool bInForceRefresh = false);
 
-	void EditGlobalSettings(bool bInForceRefresh = false);
+	virtual void EditGlobalSettings(bool bInForceRefresh = false);
 
 	void Validate();
 
@@ -104,8 +110,6 @@ protected:
 	TDMWidgetSlot<SWidget> ContentSlot;
 	TDMWidgetSlot<SDMToolBar> ToolBarSlot;
 	TDMWidgetSlot<SWidget> MainSlot;
-	TDMWidgetSlot<SWidget> LeftSlot;
-	TDMWidgetSlot<SWidget> RightSlot;
 	TDMWidgetSlot<SDMMaterialPreview> MaterialPreviewSlot;
 	TDMWidgetSlot<SDMMaterialPropertySelector> PropertySelectorSlot;
 	TDMWidgetSlot<SDMMaterialGlobalSettingsEditor> GlobalSettingsEditorSlot;
@@ -143,8 +147,12 @@ protected:
 
 	void ValidateSlots();
 
+	virtual void ValidateSlots_Main() = 0;
+
 	void ClearSlots();
 
+	virtual void ClearSlots_Main() = 0;
+
 	/** Slots */
 	void CreateLayout();
 
@@ -152,22 +160,16 @@ protected:
 
 	TSharedRef<SDMToolBar> CreateSlot_ToolBar();
 
-	TSharedRef<SWidget> CreateSlot_Main();
-
-	TSharedRef<SWidget> CreateSlot_Left();
-
-	TSharedRef<SWidget> CreateSlot_Right();
-
-	TSharedRef<SWidget> CreateSlot_Right_GlobalSettings();
+	virtual TSharedRef<SWidget> CreateSlot_Main() = 0;
 
 	TSharedRef<SDMMaterialGlobalSettingsEditor> CreateSlot_GlobalSettingsEditor();
 
-	TSharedRef<SWidget> CreateSlot_Right_Slot();
-
 	TSharedRef<SDMMaterialPreview> CreateSlot_Preview();
 
 	TSharedRef<SDMMaterialPropertySelector> CreateSlot_PropertySelector();
 
+	virtual TSharedRef<SDMMaterialPropertySelector> CreateSlot_PropertySelector_Impl() = 0;
+
 	TSharedRef<SDMMaterialSlotEditor> CreateSlot_SlotEditor();
 
 	TSharedRef<SDMMaterialComponentEditor> CreateSlot_ComponentEditor();
@@ -180,5 +182,5 @@ protected:
 	/** The material preview window is not cleaned up properly on uobject shutdown, so do it here. */
 	void OnEnginePreExit();
 
-	void OnRightSlotSplitterResized();
+	void OnEditorSplitterResized();
 };
diff --git a/Engine/Plugins/Experimental/DynamicMaterial/Source/DynamicMaterialEditor/Private/UI/Widgets/SDMMaterialWizard.cpp b/Engine/Plugins/Experimental/DynamicMaterial/Source/DynamicMaterialEditor/Private/UI/Widgets/SDMMaterialWizard.cpp
--- a/Engine/Plugins/Experimental/DynamicMaterial/Source/DynamicMaterialEditor/Private/UI/Widgets/SDMMaterialWizard.cpp
+++ b/Engine/Plugins/Experimental/DynamicMaterial/Source/DynamicMaterialEditor/Private/UI/Widgets/SDMMaterialWizard.cpp
@@ -12,6 +12,7 @@
 #include "Styling/SlateTypes.h"
 #include "Styling/StyleColors.h"
 #include "UI/Widgets/SDMMaterialDesigner.h"
+#include "Utils/DMPrivate.h"
 #include "Widgets/Input/SButton.h"
 #include "Widgets/Input/SCheckBox.h"
 #include "Widgets/Layout/SBox.h"
@@ -233,15 +234,12 @@ TSharedRef<SWidget> SDMMaterialWizard::CreateChannelList()
 
 			if (Preset->IsPropertyEnabled(Property.Key))
 			{
-				constexpr const TCHAR* ShortNameName = TEXT("ShortName");
-				const FString ShortName = MaterialPropertyEnum->GetMetaData(ShortNameName, MaterialPropertyEnum->GetIndexByValue(static_cast<int64>(Property.Key)));
-
 				ChannelPresets->AddSlot()
 					.Padding(TextPadding)
 					[
 						SNew(STextBlock)
 						.TextStyle(FDynamicMaterialEditorStyle::Get(), "RegularFont")
-						.Text(!ShortName.IsEmpty() ? FText::FromString(ShortName) : MaterialPropertyEnum->GetDisplayNameTextByValue(static_cast<int64>(Property.Key)))
+						.Text(UE::DynamicMaterialEditor::Private::GetMaterialPropertyShortDisplayName(Property.Key))
 					];
 			}
 		}
diff --git a/Engine/Plugins/Experimental/DynamicMaterial/Source/DynamicMaterialEditor/Private/Utils/DMPrivate.cpp b/Engine/Plugins/Experimental/DynamicMaterial/Source/DynamicMaterialEditor/Private/Utils/DMPrivate.cpp
--- a/Engine/Plugins/Experimental/DynamicMaterial/Source/DynamicMaterialEditor/Private/Utils/DMPrivate.cpp
+++ b/Engine/Plugins/Experimental/DynamicMaterial/Source/DynamicMaterialEditor/Private/Utils/DMPrivate.cpp
@@ -93,6 +93,22 @@ namespace UE::DynamicMaterialEditor::Private
 			FSlateNotificationManager::Get().AddNotification(Info);
 		}
 	}
+
+	FText GetMaterialPropertyShortDisplayName(EDMMaterialPropertyType InMaterialProperty)
+	{
+		constexpr const TCHAR* ShortNameName = TEXT("ShortName");
+
+		UEnum* PropertyEnum = StaticEnum<EDMMaterialPropertyType>();
+
+		const FString ShortName = PropertyEnum->GetMetaData(ShortNameName, PropertyEnum->GetIndexByValue(static_cast<int64>(InMaterialProperty)));
+
+		if (!ShortName.IsEmpty())
+		{
+			return FText::FromString(ShortName);
+		}
+
+		return PropertyEnum->GetDisplayNameTextByValue(static_cast<int64>(InMaterialProperty));
+	}
 }
 
 FDMMaterialLayerReference::FDMMaterialLayerReference()
diff --git a/Engine/Plugins/Experimental/DynamicMaterial/Source/DynamicMaterialEditor/Private/Utils/DMPrivate.h b/Engine/Plugins/Experimental/DynamicMaterial/Source/DynamicMaterialEditor/Private/Utils/DMPrivate.h
--- a/Engine/Plugins/Experimental/DynamicMaterial/Source/DynamicMaterialEditor/Private/Utils/DMPrivate.h
+++ b/Engine/Plugins/Experimental/DynamicMaterial/Source/DynamicMaterialEditor/Private/Utils/DMPrivate.h
@@ -5,6 +5,7 @@
 #include "Containers/Array.h"
 #include "HAL/Platform.h"
 #include "ScopedTransaction.h"
+#include "UI/Menus/DMMaterialSlotLayerMenus.h"
 #include "UI/Widgets/Editor/SDMMaterialSlotEditor.h"
 #include "UObject/WeakObjectPtr.h"
 
@@ -32,6 +33,8 @@ namespace UE::DynamicMaterialEditor::Private
 	bool IsCustomMaterialProperty(EDMMaterialPropertyType InMaterialProperty);
 
 	void LogError(const FString& InMessage, bool bInToast = true, const UObject* InSource = nullptr);
+
+	FText GetMaterialPropertyShortDisplayName(EDMMaterialPropertyType InMaterialProperty);
 }
 
 struct FDMMaterialLayerReference
