Material Designer
- Added "replace stage" to stage right click.

diff --git a/Engine/Plugins/Experimental/DynamicMaterial/Source/DynamicMaterialEditor/Private/UI/Menus/DMMaterialStageMenus.cpp b/Engine/Plugins/Experimental/DynamicMaterial/Source/DynamicMaterialEditor/Private/UI/Menus/DMMaterialStageMenus.cpp
--- a/Engine/Plugins/Experimental/DynamicMaterial/Source/DynamicMaterialEditor/Private/UI/Menus/DMMaterialStageMenus.cpp
+++ b/Engine/Plugins/Experimental/DynamicMaterial/Source/DynamicMaterialEditor/Private/UI/Menus/DMMaterialStageMenus.cpp
@@ -5,6 +5,7 @@
 #include "Components/DMMaterialLayer.h"
 #include "Components/DMMaterialSlot.h"
 #include "Components/DMMaterialStage.h"
+#include "DMMaterialStageSourceMenus.h"
 #include "ScopedTransaction.h"
 #include "Templates/SharedPointer.h"
 #include "ToolMenu.h"
@@ -18,8 +19,9 @@
 
 namespace UE::DynamicMaterialEditor::Private
 {
-	const FLazyName StageMenuName = TEXT("MaterialDesigner.MaterialStage");
+	const FLazyName StageSettingsMenuName = TEXT("MaterialDesigner.MaterialStage");
 	const FLazyName StageMenuToggleName = TEXT("StageToggle");
+	const FLazyName StageSourceMenuName = TEXT("MaterialDesigner.MaterialStageSource");
 }
 
 TSharedRef<SWidget> FDMMaterialStageMenus::GenerateStageMenu(const TSharedPtr<SDMMaterialSlotEditor>& InSlotWidget,
@@ -29,24 +31,25 @@ TSharedRef<SWidget> FDMMaterialStageMenus::GenerateStageMenu(const TSharedPtr<SD
 
 	UToolMenus* ToolMenus = UToolMenus::Get();
 
-	if (!ToolMenus->IsMenuRegistered(StageMenuName))
+	if (!ToolMenus->IsMenuRegistered(StageSettingsMenuName))
 	{
-		UToolMenu* NewToolMenu = UDMMenuContext::GenerateContextMenuDefault(StageMenuName);
+		UToolMenu* NewToolMenu = UDMMenuContext::GenerateContextMenuDefault(StageSettingsMenuName);
 
 		if (!NewToolMenu)
 		{
 			return SNullWidget::NullWidget;
 		}
 
-		NewToolMenu->AddDynamicSection(NAME_None, FNewToolMenuDelegate::CreateStatic(&FDMMaterialStageMenus::AddStageSection));
+		NewToolMenu->AddDynamicSection(NAME_None, FNewToolMenuDelegate::CreateStatic(&FDMMaterialStageMenus::AddStageSettingsSection));
+		NewToolMenu->AddDynamicSection(NAME_None, FNewToolMenuDelegate::CreateStatic(&FDMMaterialStageMenus::AddStageSourceSection));
 	}
 
 	FToolMenuContext MenuContext(UDMMenuContext::CreateStage(InSlotWidget->GetEditorWidget(), InStageWidget));
 
-	return ToolMenus->GenerateWidget(StageMenuName, MenuContext);
+	return ToolMenus->GenerateWidget(StageSettingsMenuName, MenuContext);
 }
 
-void FDMMaterialStageMenus::AddStageSection(UToolMenu* InMenu)
+void FDMMaterialStageMenus::AddStageSettingsSection(UToolMenu* InMenu)
 {
 	using namespace UE::DynamicMaterialEditor::Private;
 
@@ -79,16 +82,11 @@ void FDMMaterialStageMenus::AddStageSection(UToolMenu* InMenu)
 		return;
 	}
 
-	const bool bAllowRemoveLayer = Slot->CanRemoveLayer(Layer);
 	const EDMMaterialLayerStage StageType = Layer->GetStageType(Stage);
+	const bool bAllowRemoveLayer = Slot->CanRemoveLayer(Layer);
 
 	FToolMenuSection& NewSection = InMenu->AddSection(StageMenuToggleName, LOCTEXT("MaterialStageMenu", "Material Stage"));
 
-	if (!bAllowRemoveLayer && StageType != EDMMaterialLayerStage::Mask)
-	{
-		return;
-	}
-
 	if (StageType == EDMMaterialLayerStage::Mask)
 	{
 		NewSection.AddMenuEntry(NAME_None,
@@ -183,4 +181,19 @@ void FDMMaterialStageMenus::AddStageSection(UToolMenu* InMenu)
 	}
 }
 
+void FDMMaterialStageMenus::AddStageSourceSection(UToolMenu* InMenu)
+{
+	using namespace UE::DynamicMaterialEditor::Private;
+
+	if (!IsValid(InMenu) || InMenu->ContainsSection(StageSourceMenuName))
+	{
+		return;
+	}
+
+	FToolMenuSection& NewSection = InMenu->AddSection(StageSourceMenuName, LOCTEXT("MaterialStageSource", "Stage Source"));
+	NewSection.Context.AddObject(InMenu->Context.FindContext<UDMMenuContext>());
+
+	FDMMaterialStageSourceMenus::CreateChangeMaterialStageSource(NewSection);
+}
+
 #undef LOCTEXT_NAMESPACE
diff --git a/Engine/Plugins/Experimental/DynamicMaterial/Source/DynamicMaterialEditor/Private/UI/Menus/DMMaterialStageMenus.h b/Engine/Plugins/Experimental/DynamicMaterial/Source/DynamicMaterialEditor/Private/UI/Menus/DMMaterialStageMenus.h
--- a/Engine/Plugins/Experimental/DynamicMaterial/Source/DynamicMaterialEditor/Private/UI/Menus/DMMaterialStageMenus.h
+++ b/Engine/Plugins/Experimental/DynamicMaterial/Source/DynamicMaterialEditor/Private/UI/Menus/DMMaterialStageMenus.h
@@ -16,5 +16,7 @@ public:
 	static TSharedRef<SWidget> GenerateStageMenu(const TSharedPtr<SDMMaterialSlotEditor>& InSlotWidget, const TSharedPtr<SDMMaterialStage>& InStageWidget);
 
 private:
-	static void AddStageSection(UToolMenu* InMenu);
+	static void AddStageSettingsSection(UToolMenu* InMenu);
+
+	static void AddStageSourceSection(UToolMenu* InMenu);
 };
diff --git a/Engine/Plugins/Experimental/DynamicMaterial/Source/DynamicMaterialEditor/Private/UI/PropertyGenerators/DMStagePropertyRowGenerator.cpp b/Engine/Plugins/Experimental/DynamicMaterial/Source/DynamicMaterialEditor/Private/UI/PropertyGenerators/DMStagePropertyRowGenerator.cpp
--- a/Engine/Plugins/Experimental/DynamicMaterial/Source/DynamicMaterialEditor/Private/UI/PropertyGenerators/DMStagePropertyRowGenerator.cpp
+++ b/Engine/Plugins/Experimental/DynamicMaterial/Source/DynamicMaterialEditor/Private/UI/PropertyGenerators/DMStagePropertyRowGenerator.cpp
@@ -10,6 +10,7 @@
 #include "DMEDefs.h"
 #include "DynamicMaterialEditorModule.h"
 #include "Internationalization/Text.h"
+#include "Internationalization/Text.h"
 #include "Model/DynamicMaterialModel.h"
 #include "Model/DynamicMaterialModelBase.h"
 #include "UI/Menus/DMMaterialStageSourceMenus.h"
