[Backout] - CL35450035
Original CL Desc
-----------------------------------------------------------------
Material Designer
- Layout options

diff --git a/Engine/Plugins/Experimental/DynamicMaterial/Source/DynamicMaterial/Public/DMDefs.h b/Engine/Plugins/Experimental/DynamicMaterial/Source/DynamicMaterial/Public/DMDefs.h
--- a/Engine/Plugins/Experimental/DynamicMaterial/Source/DynamicMaterial/Public/DMDefs.h
+++ b/Engine/Plugins/Experimental/DynamicMaterial/Source/DynamicMaterial/Public/DMDefs.h
@@ -16,27 +16,27 @@ UENUM(BlueprintType)
 enum class EDMMaterialPropertyType : uint8
 {
 	None = 0,
-	BaseColor           UMETA(ShortName = "Base"),
-	EmissiveColor       UMETA(ShortName = "Emis"),
+	BaseColor,
+	EmissiveColor UMETA(ShortName = "Emissive"),
 	Opacity,
-	OpacityMask         UMETA(ShortName = "Mask"),
-	Roughness           UMETA(ShortName = "Rough"),
-	Specular            UMETA(ShortName = "Spec"),
-	Metallic            UMETA(ShortName = "Metal"),
-	Normal              UMETA(ShortName = "Norm"),
-	PixelDepthOffset    UMETA(ShortName = "PDO"),
+	OpacityMask UMETA(ShortName = "Mask"),
+	Roughness,
+	Specular,
+	Metallic,
+	Normal,
+	PixelDepthOffset UMETA(ShortName = "PDO"),
 	WorldPositionOffset UMETA(ShortName = "WPO"),
-	AmbientOcclusion    UMETA(ShortName = "AO"),
-	Anisotropy          UMETA(ShortName = "Aniso"),
-	Refraction          UMETA(ShortName = "Refra"),
-	Tangent             UMETA(ShortName = "Tang"),
+	AmbientOcclusion UMETA(ShortName = "AO"),
+	Anisotropy,
+	Refraction,
+	Tangent,
 	Custom1,
 	Custom2,
 	Custom3,
 	Custom4,
-	Displacement        UMETA(ShortName = "Disp"),
-	SubsurfaceColor     UMETA(ShortName = "Surf"),
-	SurfaceThickness    UMETA(ShortName = "Thick"),
+	Displacement,
+	SubsurfaceColor,
+	SurfaceThickness,
 	Any
 };
 
diff --git a/Engine/Plugins/Experimental/DynamicMaterial/Source/DynamicMaterialEditor/Private/Components/DMMaterialSlot.cpp b/Engine/Plugins/Experimental/DynamicMaterial/Source/DynamicMaterialEditor/Private/Components/DMMaterialSlot.cpp
--- a/Engine/Plugins/Experimental/DynamicMaterial/Source/DynamicMaterialEditor/Private/Components/DMMaterialSlot.cpp
+++ b/Engine/Plugins/Experimental/DynamicMaterial/Source/DynamicMaterialEditor/Private/Components/DMMaterialSlot.cpp
@@ -23,7 +23,6 @@
 #include "Model/DynamicMaterialModelEditorOnlyData.h"
 #include "UObject/Package.h"
 #include "UObject/UObjectGlobals.h"
-#include "Utils/DMPrivate.h"
 #include "Utils/DMUtils.h"
 
 #define LOCTEXT_NAMESPACE "DMMaterialStage"
@@ -1383,13 +1382,16 @@ FString UDMMaterialSlot::GetComponentPathComponent() const
 
 		if (SlotProperties.Num() == 1)
 		{
-			const FText Token = UE::DynamicMaterialEditor::Private::GetMaterialPropertyShortDisplayName(SlotProperties[0]);
+			UEnum* MaterialPropertyEnum = StaticEnum<EDMMaterialPropertyType>();
+			constexpr const TCHAR* ShortNameName = TEXT("ShortName");
+			const FString ShortName = MaterialPropertyEnum->GetMetaData(ShortNameName, MaterialPropertyEnum->GetIndexByValue(static_cast<int64>(SlotProperties[0])));
+			const FString Token = !ShortName.IsEmpty() ? ShortName : MaterialPropertyEnum->GetNameStringByValue(static_cast<int64>(SlotProperties[0]));
 
 			return FString::Printf(
 				TEXT("%s%hc%s%hc"),
 				*UDynamicMaterialModelEditorOnlyData::SlotsPathToken,
 				FDMComponentPath::ParameterOpen,
-				*Token.ToString(),
+				*Token,
 				FDMComponentPath::ParameterClose
 			);
 		}
diff --git a/Engine/Plugins/Experimental/DynamicMaterial/Source/DynamicMaterialEditor/Private/DynamicMaterialEditorSettings.cpp b/Engine/Plugins/Experimental/DynamicMaterial/Source/DynamicMaterialEditor/Private/DynamicMaterialEditorSettings.cpp
--- a/Engine/Plugins/Experimental/DynamicMaterial/Source/DynamicMaterialEditor/Private/DynamicMaterialEditorSettings.cpp
+++ b/Engine/Plugins/Experimental/DynamicMaterial/Source/DynamicMaterialEditor/Private/DynamicMaterialEditorSettings.cpp
@@ -267,8 +267,7 @@ void UDynamicMaterialEditorSettings::PostEditChangeProperty(FPropertyChangedEven
 
 	const FName PropertyName = InPropertyChangedEvent.GetMemberPropertyName();
 
-	if (PropertyName == GET_MEMBER_NAME_CHECKED(UDynamicMaterialEditorSettings, Layout)
-		|| PropertyName == GET_MEMBER_NAME_CHECKED(UDynamicMaterialEditorSettings, SplitterLocation)
+	if (PropertyName == GET_MEMBER_NAME_CHECKED(UDynamicMaterialEditorSettings, SplitterLocation)
 		|| PropertyName == GET_MEMBER_NAME_CHECKED(UDynamicMaterialEditorSettings, bPreviewImagesUseTextureUVs)
 		|| PropertyName == GET_MEMBER_NAME_CHECKED(UDynamicMaterialEditorSettings, PreviewMesh)
 		|| PropertyName == GET_MEMBER_NAME_CHECKED(UDynamicMaterialEditorSettings, bShowPreviewBackground))
@@ -298,7 +297,6 @@ void UDynamicMaterialEditorSettings::OpenEditorSettingsWindow() const
 
 void UDynamicMaterialEditorSettings::ResetAllLayoutSettings()
 {
-	Layout = EDMMaterialEditorLayout::Left;
 	SplitterLocation = 240;
 	ThumbnailSize = 256.f;
 	PreviewMesh = EDMMaterialPreviewMesh::Plane;
diff --git a/Engine/Plugins/Experimental/DynamicMaterial/Source/DynamicMaterialEditor/Private/DynamicMaterialEditorSettings.h b/Engine/Plugins/Experimental/DynamicMaterial/Source/DynamicMaterialEditor/Private/DynamicMaterialEditorSettings.h
--- a/Engine/Plugins/Experimental/DynamicMaterial/Source/DynamicMaterialEditor/Private/DynamicMaterialEditorSettings.h
+++ b/Engine/Plugins/Experimental/DynamicMaterial/Source/DynamicMaterialEditor/Private/DynamicMaterialEditorSettings.h
@@ -137,22 +137,6 @@ enum class EDMMaterialPreviewMesh : uint8
 	ShaderBall
 };
 
-UENUM(BlueprintType)
-enum class EDMMaterialEditorLayout : uint8
-{
-	Left,
-	LeftAutoHide,
-	LeftSlim,
-	TopVertical,
-	TopVerticalAutoHide,
-	TopHorizontal,
-	TopHorizontalAutoHide,
-	TopSlim,
-
-	First = Left,
-	Last = TopSlim
-};
-
 /**
  * Material Designer Settings
  */
@@ -171,9 +155,6 @@ public:
 	UPROPERTY(Config, EditAnywhere, BlueprintReadWrite, Category = "Editor")
 	bool bFollowSelection;
 
-	UPROPERTY(Config, EditAnywhere, BlueprintReadWrite, Category = "Layout")
-	EDMMaterialEditorLayout Layout;
-
 	/** Adjusts the vertical size of the material layer view. */
 	UPROPERTY(Config, EditAnywhere, BlueprintReadWrite, Category = "Layout", meta = (
 		ClampMin = "0.05", UIMin = "0.05", ClampMax = "0.95", UIMax = "0.95"))
diff --git a/Engine/Plugins/Experimental/DynamicMaterial/Source/DynamicMaterialEditor/Private/UI/Menus/DMToolBarMenus.cpp b/Engine/Plugins/Experimental/DynamicMaterial/Source/DynamicMaterialEditor/Private/UI/Menus/DMToolBarMenus.cpp
--- a/Engine/Plugins/Experimental/DynamicMaterial/Source/DynamicMaterialEditor/Private/UI/Menus/DMToolBarMenus.cpp
+++ b/Engine/Plugins/Experimental/DynamicMaterial/Source/DynamicMaterialEditor/Private/UI/Menus/DMToolBarMenus.cpp
@@ -55,7 +55,7 @@ TSharedRef<SWidget> FDMToolBarMenus::MakeEditorLayoutMenu(const TSharedPtr<SDMMa
 
 		NewToolMenu->AddDynamicSection(
 			TEXT("MaterialDesignerSettings"),
-			FNewToolMenuDelegate::CreateStatic(&AddMenu)
+			FNewToolMenuDelegate::CreateStatic(&AddToolBarEditorLayoutMenu)
 		);
 	}
 
@@ -68,13 +68,7 @@ TSharedRef<SWidget> FDMToolBarMenus::MakeEditorLayoutMenu(const TSharedPtr<SDMMa
 	return ToolMenus->GenerateWidget(ToolBarEditorLayoutMenuName, MenuContext);
 }
 
-void FDMToolBarMenus::AddMenu(UToolMenu* InMenu)
-{
-	AddExportMenu(InMenu);
-	AddSettingsMenu(InMenu);
-}
-
-void FDMToolBarMenus::AddExportMenu(UToolMenu* InMenu)
+void FDMToolBarMenus::AddToolBarExportMenu(UToolMenu* InMenu)
 {
 	using namespace UE::DynamicMaterialEditor::Private;
 
@@ -163,7 +157,23 @@ void FDMToolBarMenus::AddExportMenu(UToolMenu* InMenu)
 	);
 }
 
-void FDMToolBarMenus::AddSettingsMenu(UToolMenu* InMenu)
+void FDMToolBarMenus::AddToolBarAdvancedSection(UToolMenu* InMenu)
+{
+	FToolMenuSection& NewSection = InMenu->AddSection(TEXT("AdvancedSettings"), LOCTEXT("AdvancedSettingsSection", "Advanced Settings"));
+
+	NewSection.AddMenuEntry(
+		NAME_None,
+		LOCTEXT("ResetAllSettingsToDefaults", "Reset All To Defaults"),
+		LOCTEXT("ResetAllSettingsToDefaultsTooltip", "Resets all the Material Designer settings to their default values."),
+		FSlateIcon(),
+		FUIAction(FExecuteAction::CreateUObject(
+			UDynamicMaterialEditorSettings::Get(), 
+			&UDynamicMaterialEditorSettings::ResetAllLayoutSettings)
+		)
+	);
+}
+
+void FDMToolBarMenus::AddToolBarSettingsMenu(UToolMenu* InMenu)
 {
 	using namespace UE::DynamicMaterialEditor::Private;
 
@@ -178,14 +188,7 @@ void FDMToolBarMenus::AddSettingsMenu(UToolMenu* InMenu)
 		"AdvancedSettings",
 		LOCTEXT("AdvancedSettingsSubMenu", "Advanced Settings"),
 		LOCTEXT("AdvancedSettingsSubMenu_ToolTip", "Display advanced Material Designer settings"),
-		FNewToolMenuDelegate::CreateStatic(&AddAdvancedSection)
-	);
-
-	NewSection.AddSubMenu(
-		"EditorLayout",
-		LOCTEXT("EditorLayoutSubMenu", "Editor Layout"),
-		LOCTEXT("EditorLayoutSubMenu_ToolTip", "Change the layout of the Material Designer Editor"),
-		FNewToolMenuDelegate::CreateStatic(&AddEditorLayoutSection)
+		FNewToolMenuDelegate::CreateStatic(&AddToolBarAdvancedSection)
 	);
 
 	NewSection.AddMenuEntry(
@@ -200,93 +203,10 @@ void FDMToolBarMenus::AddSettingsMenu(UToolMenu* InMenu)
 	);
 }
 
-void FDMToolBarMenus::AddEditorLayoutSection(UToolMenu* InMenu)
+void FDMToolBarMenus::AddToolBarEditorLayoutMenu(UToolMenu* InMenu)
 {
-	FToolMenuSection& NewSection = InMenu->AddSection(TEXT("EditorLayout"), LOCTEXT("EditorLayoutSection", "EditorLayout"));
-
-	UEnum* LayoutEnum = StaticEnum<EDMMaterialEditorLayout>();
-
-	for (EDMMaterialEditorLayout Layout = EDMMaterialEditorLayout::First;
-		Layout <= EDMMaterialEditorLayout::Last;
-		Layout = static_cast<EDMMaterialEditorLayout>(static_cast<uint8>(Layout)+1))
-	{
-		FUIAction Action;
-
-		switch (Layout)
-		{
-			case EDMMaterialEditorLayout::LeftAutoHide:
-			case EDMMaterialEditorLayout::TopHorizontalAutoHide:
-			case EDMMaterialEditorLayout::TopVerticalAutoHide:
-				Action.CanExecuteAction = FCanExecuteAction::CreateLambda([]() { return false; });
-				break;
-
-			default:
-				break;
-		}
-
-		Action.GetActionCheckState = FGetActionCheckState::CreateLambda(
-			[Layout]()
-			{
-				if (UDynamicMaterialEditorSettings* Settings = GetMutableDefault<UDynamicMaterialEditorSettings>())
-				{
-					if (Settings->Layout == Layout)
-					{
-						return ECheckBoxState::Checked;
-					}
-				}
-
-				return ECheckBoxState::Unchecked;
-			}
-		);
-
-		Action.ExecuteAction = FExecuteAction::CreateLambda(
-			[Layout]()
-			{
-				UDynamicMaterialEditorSettings* Settings = GetMutableDefault<UDynamicMaterialEditorSettings>();
-
-				if (!Settings || Settings->Layout == Layout)
-				{
-					return;
-				}
-
-				Settings->Layout = Layout;
-				TArray<UObject*> TopLeftObjects = {Settings};
-
-				FPropertyChangedEvent PropertyChangedEvent(
-					UDynamicMaterialEditorSettings::StaticClass()->FindPropertyByName(GET_MEMBER_NAME_CHECKED(UDynamicMaterialEditorSettings, Layout)),
-					EPropertyChangeType::Interactive,
-					TopLeftObjects
-				);
-
-				Settings->PostEditChangeProperty(PropertyChangedEvent);
-			}
-		);
-
-		NewSection.AddMenuEntry(
-			NAME_None,
-			LayoutEnum->GetDisplayNameTextByValue(static_cast<int64>(Layout)),
-			FText::GetEmpty(),
-			FSlateIcon(),
-			Action,
-			EUserInterfaceActionType::RadioButton
-		);
-	}
-}
-
-void FDMToolBarMenus::AddAdvancedSection(UToolMenu* InMenu)
-{
-	FToolMenuSection& NewSection = InMenu->AddSection(TEXT("AdvancedSettings"), LOCTEXT("AdvancedSettingsSection", "Advanced Settings"));
-
-	NewSection.AddMenuEntry(
-		NAME_None,
-		LOCTEXT("ResetAllSettingsToDefaults", "Reset All To Defaults"),
-		LOCTEXT("ResetAllSettingsToDefaultsTooltip", "Resets all the Material Designer settings to their default values."),
-		FSlateIcon(),
-		FUIAction(FExecuteAction::CreateUObject(
-			UDynamicMaterialEditorSettings::Get(), 
-			&UDynamicMaterialEditorSettings::ResetAllLayoutSettings)
-		)
-	);
+	AddToolBarExportMenu(InMenu);
+	AddToolBarSettingsMenu(InMenu);
 }
 
 void FDMToolBarMenus::OpenMaterialEditorFromContext(UDMMenuContext* InMenuContext)
@@ -437,7 +357,7 @@ void FDMToolBarMenus::SnapshotMaterial(TWeakObjectPtr<UDynamicMaterialModelBase>
 	}
 }
 
-void FDMToolBarMenus::AddBoolOptionMenuEntry(FToolMenuSection& InSection, const FName& InPropertyName, const FUIAction InAction)
+void FDMToolBarMenus::AddToolBarBoolOptionMenuEntry(FToolMenuSection& InSection, const FName& InPropertyName, const FUIAction InAction)
 {
 	const FProperty* const OptionProperty = UDynamicMaterialEditorSettings::StaticClass()->FindPropertyByName(InPropertyName);
 
@@ -453,7 +373,7 @@ void FDMToolBarMenus::AddBoolOptionMenuEntry(FToolMenuSection& InSection, const
 	}
 }
 
-void FDMToolBarMenus::AddIntOptionMenuEntry(FToolMenuSection& InSection, const FName& InPropertyName,
+void FDMToolBarMenus::AddToolBarIntOptionMenuEntry(FToolMenuSection& InSection, const FName& InPropertyName,
 	TAttribute<bool> InIsEnabledAttribute, TAttribute<EVisibility> InVisibilityAttribute)
 {
 	InSection.AddDynamicEntry(NAME_None,
diff --git a/Engine/Plugins/Experimental/DynamicMaterial/Source/DynamicMaterialEditor/Private/UI/Menus/DMToolBarMenus.h b/Engine/Plugins/Experimental/DynamicMaterial/Source/DynamicMaterialEditor/Private/UI/Menus/DMToolBarMenus.h
--- a/Engine/Plugins/Experimental/DynamicMaterial/Source/DynamicMaterialEditor/Private/UI/Menus/DMToolBarMenus.h
+++ b/Engine/Plugins/Experimental/DynamicMaterial/Source/DynamicMaterialEditor/Private/UI/Menus/DMToolBarMenus.h
@@ -14,7 +14,6 @@ class UDMMenuContext;
 class UDynamicMaterialInstance;
 class UDynamicMaterialModelBase;
 class UToolMenu;
-enum class EDMMaterialEditorLayout : uint8;
 struct FToolMenuSection;
 struct FUIAction;
 
@@ -24,19 +23,17 @@ public:
 	static TSharedRef<SWidget> MakeEditorLayoutMenu(const TSharedPtr<SDMMaterialEditor>& InEditorWidget = nullptr);
 
 private:
-	static void AddMenu(UToolMenu* InMenu);
+	static void AddToolBarExportMenu(UToolMenu* InMenu);
 
-	static void AddExportMenu(UToolMenu* InMenu);
+	static void AddToolBarAdvancedSection(UToolMenu* InMenu);
 
-	static void AddSettingsMenu(UToolMenu* InMenu);
+	static void AddToolBarSettingsMenu(UToolMenu* InMenu);
 
-	static void AddEditorLayoutSection(UToolMenu* InMenu);
+	static void AddToolBarEditorLayoutMenu(UToolMenu* InMenu);
 
-	static void AddAdvancedSection(UToolMenu* InMenu);
+	static void AddToolBarBoolOptionMenuEntry(FToolMenuSection& InSection, const FName& InPropertyName, const FUIAction InAction);
 
-	static void AddBoolOptionMenuEntry(FToolMenuSection& InSection, const FName& InPropertyName, const FUIAction InAction);
-
-	static void AddIntOptionMenuEntry(FToolMenuSection& InSection, const FName& InPropertyName,
+	static void AddToolBarIntOptionMenuEntry(FToolMenuSection& InSection, const FName& InPropertyName,
 		TAttribute<bool> InIsEnabledAttribute = TAttribute<bool>(),
 		TAttribute<EVisibility> InVisibilityAttribute = TAttribute<EVisibility>());
 
diff --git a/Engine/Plugins/Experimental/DynamicMaterial/Source/DynamicMaterialEditor/Private/UI/Widgets/Editor/EditorLayouts/SDMMaterialEditor_Left.cpp b/Engine/Plugins/Experimental/DynamicMaterial/Source/DynamicMaterialEditor/Private/UI/Widgets/Editor/EditorLayouts/SDMMaterialEditor_Left.cpp
deleted file mode 100644
diff --git a/Engine/Plugins/Experimental/DynamicMaterial/Source/DynamicMaterialEditor/Private/UI/Widgets/Editor/EditorLayouts/SDMMaterialEditor_Left.h b/Engine/Plugins/Experimental/DynamicMaterial/Source/DynamicMaterialEditor/Private/UI/Widgets/Editor/EditorLayouts/SDMMaterialEditor_Left.h
deleted file mode 100644
diff --git a/Engine/Plugins/Experimental/DynamicMaterial/Source/DynamicMaterialEditor/Private/UI/Widgets/Editor/EditorLayouts/SDMMaterialEditor_LeftBase.cpp b/Engine/Plugins/Experimental/DynamicMaterial/Source/DynamicMaterialEditor/Private/UI/Widgets/Editor/EditorLayouts/SDMMaterialEditor_LeftBase.cpp
deleted file mode 100644
diff --git a/Engine/Plugins/Experimental/DynamicMaterial/Source/DynamicMaterialEditor/Private/UI/Widgets/Editor/EditorLayouts/SDMMaterialEditor_LeftBase.h b/Engine/Plugins/Experimental/DynamicMaterial/Source/DynamicMaterialEditor/Private/UI/Widgets/Editor/EditorLayouts/SDMMaterialEditor_LeftBase.h
deleted file mode 100644
diff --git a/Engine/Plugins/Experimental/DynamicMaterial/Source/DynamicMaterialEditor/Private/UI/Widgets/Editor/EditorLayouts/SDMMaterialEditor_LeftSlim.cpp b/Engine/Plugins/Experimental/DynamicMaterial/Source/DynamicMaterialEditor/Private/UI/Widgets/Editor/EditorLayouts/SDMMaterialEditor_LeftSlim.cpp
deleted file mode 100644
diff --git a/Engine/Plugins/Experimental/DynamicMaterial/Source/DynamicMaterialEditor/Private/UI/Widgets/Editor/EditorLayouts/SDMMaterialEditor_LeftSlim.h b/Engine/Plugins/Experimental/DynamicMaterial/Source/DynamicMaterialEditor/Private/UI/Widgets/Editor/EditorLayouts/SDMMaterialEditor_LeftSlim.h
deleted file mode 100644
diff --git a/Engine/Plugins/Experimental/DynamicMaterial/Source/DynamicMaterialEditor/Private/UI/Widgets/Editor/EditorLayouts/SDMMaterialEditor_TopBase.cpp b/Engine/Plugins/Experimental/DynamicMaterial/Source/DynamicMaterialEditor/Private/UI/Widgets/Editor/EditorLayouts/SDMMaterialEditor_TopBase.cpp
deleted file mode 100644
diff --git a/Engine/Plugins/Experimental/DynamicMaterial/Source/DynamicMaterialEditor/Private/UI/Widgets/Editor/EditorLayouts/SDMMaterialEditor_TopBase.h b/Engine/Plugins/Experimental/DynamicMaterial/Source/DynamicMaterialEditor/Private/UI/Widgets/Editor/EditorLayouts/SDMMaterialEditor_TopBase.h
deleted file mode 100644
diff --git a/Engine/Plugins/Experimental/DynamicMaterial/Source/DynamicMaterialEditor/Private/UI/Widgets/Editor/EditorLayouts/SDMMaterialEditor_TopHorizontal.cpp b/Engine/Plugins/Experimental/DynamicMaterial/Source/DynamicMaterialEditor/Private/UI/Widgets/Editor/EditorLayouts/SDMMaterialEditor_TopHorizontal.cpp
deleted file mode 100644
diff --git a/Engine/Plugins/Experimental/DynamicMaterial/Source/DynamicMaterialEditor/Private/UI/Widgets/Editor/EditorLayouts/SDMMaterialEditor_TopHorizontal.h b/Engine/Plugins/Experimental/DynamicMaterial/Source/DynamicMaterialEditor/Private/UI/Widgets/Editor/EditorLayouts/SDMMaterialEditor_TopHorizontal.h
deleted file mode 100644
diff --git a/Engine/Plugins/Experimental/DynamicMaterial/Source/DynamicMaterialEditor/Private/UI/Widgets/Editor/EditorLayouts/SDMMaterialEditor_TopSlim.cpp b/Engine/Plugins/Experimental/DynamicMaterial/Source/DynamicMaterialEditor/Private/UI/Widgets/Editor/EditorLayouts/SDMMaterialEditor_TopSlim.cpp
deleted file mode 100644
diff --git a/Engine/Plugins/Experimental/DynamicMaterial/Source/DynamicMaterialEditor/Private/UI/Widgets/Editor/EditorLayouts/SDMMaterialEditor_TopSlim.h b/Engine/Plugins/Experimental/DynamicMaterial/Source/DynamicMaterialEditor/Private/UI/Widgets/Editor/EditorLayouts/SDMMaterialEditor_TopSlim.h
deleted file mode 100644
diff --git a/Engine/Plugins/Experimental/DynamicMaterial/Source/DynamicMaterialEditor/Private/UI/Widgets/Editor/EditorLayouts/SDMMaterialEditor_TopVertical.cpp b/Engine/Plugins/Experimental/DynamicMaterial/Source/DynamicMaterialEditor/Private/UI/Widgets/Editor/EditorLayouts/SDMMaterialEditor_TopVertical.cpp
deleted file mode 100644
diff --git a/Engine/Plugins/Experimental/DynamicMaterial/Source/DynamicMaterialEditor/Private/UI/Widgets/Editor/EditorLayouts/SDMMaterialEditor_TopVertical.h b/Engine/Plugins/Experimental/DynamicMaterial/Source/DynamicMaterialEditor/Private/UI/Widgets/Editor/EditorLayouts/SDMMaterialEditor_TopVertical.h
deleted file mode 100644
diff --git a/Engine/Plugins/Experimental/DynamicMaterial/Source/DynamicMaterialEditor/Private/UI/Widgets/Editor/PropertySelectorLayouts/SDMMaterialPropertySelector_Vertical.cpp b/Engine/Plugins/Experimental/DynamicMaterial/Source/DynamicMaterialEditor/Private/UI/Widgets/Editor/PropertySelectorLayouts/SDMMaterialPropertySelector_Vertical.cpp
deleted file mode 100644
diff --git a/Engine/Plugins/Experimental/DynamicMaterial/Source/DynamicMaterialEditor/Private/UI/Widgets/Editor/PropertySelectorLayouts/SDMMaterialPropertySelector_Vertical.h b/Engine/Plugins/Experimental/DynamicMaterial/Source/DynamicMaterialEditor/Private/UI/Widgets/Editor/PropertySelectorLayouts/SDMMaterialPropertySelector_Vertical.h
deleted file mode 100644
diff --git a/Engine/Plugins/Experimental/DynamicMaterial/Source/DynamicMaterialEditor/Private/UI/Widgets/Editor/PropertySelectorLayouts/SDMMaterialPropertySelector_VerticalBase.cpp b/Engine/Plugins/Experimental/DynamicMaterial/Source/DynamicMaterialEditor/Private/UI/Widgets/Editor/PropertySelectorLayouts/SDMMaterialPropertySelector_VerticalBase.cpp
deleted file mode 100644
diff --git a/Engine/Plugins/Experimental/DynamicMaterial/Source/DynamicMaterialEditor/Private/UI/Widgets/Editor/PropertySelectorLayouts/SDMMaterialPropertySelector_VerticalBase.h b/Engine/Plugins/Experimental/DynamicMaterial/Source/DynamicMaterialEditor/Private/UI/Widgets/Editor/PropertySelectorLayouts/SDMMaterialPropertySelector_VerticalBase.h
deleted file mode 100644
diff --git a/Engine/Plugins/Experimental/DynamicMaterial/Source/DynamicMaterialEditor/Private/UI/Widgets/Editor/PropertySelectorLayouts/SDMMaterialPropertySelector_VerticalSlim.cpp b/Engine/Plugins/Experimental/DynamicMaterial/Source/DynamicMaterialEditor/Private/UI/Widgets/Editor/PropertySelectorLayouts/SDMMaterialPropertySelector_VerticalSlim.cpp
deleted file mode 100644
diff --git a/Engine/Plugins/Experimental/DynamicMaterial/Source/DynamicMaterialEditor/Private/UI/Widgets/Editor/PropertySelectorLayouts/SDMMaterialPropertySelector_VerticalSlim.h b/Engine/Plugins/Experimental/DynamicMaterial/Source/DynamicMaterialEditor/Private/UI/Widgets/Editor/PropertySelectorLayouts/SDMMaterialPropertySelector_VerticalSlim.h
deleted file mode 100644
diff --git a/Engine/Plugins/Experimental/DynamicMaterial/Source/DynamicMaterialEditor/Private/UI/Widgets/Editor/PropertySelectorLayouts/SDMMaterialPropertySelector_Wrap.cpp b/Engine/Plugins/Experimental/DynamicMaterial/Source/DynamicMaterialEditor/Private/UI/Widgets/Editor/PropertySelectorLayouts/SDMMaterialPropertySelector_Wrap.cpp
deleted file mode 100644
diff --git a/Engine/Plugins/Experimental/DynamicMaterial/Source/DynamicMaterialEditor/Private/UI/Widgets/Editor/PropertySelectorLayouts/SDMMaterialPropertySelector_Wrap.h b/Engine/Plugins/Experimental/DynamicMaterial/Source/DynamicMaterialEditor/Private/UI/Widgets/Editor/PropertySelectorLayouts/SDMMaterialPropertySelector_Wrap.h
deleted file mode 100644
diff --git a/Engine/Plugins/Experimental/DynamicMaterial/Source/DynamicMaterialEditor/Private/UI/Widgets/Editor/PropertySelectorLayouts/SDMMaterialPropertySelector_WrapBase.cpp b/Engine/Plugins/Experimental/DynamicMaterial/Source/DynamicMaterialEditor/Private/UI/Widgets/Editor/PropertySelectorLayouts/SDMMaterialPropertySelector_WrapBase.cpp
deleted file mode 100644
diff --git a/Engine/Plugins/Experimental/DynamicMaterial/Source/DynamicMaterialEditor/Private/UI/Widgets/Editor/PropertySelectorLayouts/SDMMaterialPropertySelector_WrapBase.h b/Engine/Plugins/Experimental/DynamicMaterial/Source/DynamicMaterialEditor/Private/UI/Widgets/Editor/PropertySelectorLayouts/SDMMaterialPropertySelector_WrapBase.h
deleted file mode 100644
diff --git a/Engine/Plugins/Experimental/DynamicMaterial/Source/DynamicMaterialEditor/Private/UI/Widgets/Editor/PropertySelectorLayouts/SDMMaterialPropertySelector_WrapSlim.cpp b/Engine/Plugins/Experimental/DynamicMaterial/Source/DynamicMaterialEditor/Private/UI/Widgets/Editor/PropertySelectorLayouts/SDMMaterialPropertySelector_WrapSlim.cpp
deleted file mode 100644
diff --git a/Engine/Plugins/Experimental/DynamicMaterial/Source/DynamicMaterialEditor/Private/UI/Widgets/Editor/PropertySelectorLayouts/SDMMaterialPropertySelector_WrapSlim.h b/Engine/Plugins/Experimental/DynamicMaterial/Source/DynamicMaterialEditor/Private/UI/Widgets/Editor/PropertySelectorLayouts/SDMMaterialPropertySelector_WrapSlim.h
deleted file mode 100644
diff --git a/Engine/Plugins/Experimental/DynamicMaterial/Source/DynamicMaterialEditor/Private/UI/Widgets/Editor/SDMMaterialPropertySelector.cpp b/Engine/Plugins/Experimental/DynamicMaterial/Source/DynamicMaterialEditor/Private/UI/Widgets/Editor/SDMMaterialPropertySelector.cpp
--- a/Engine/Plugins/Experimental/DynamicMaterial/Source/DynamicMaterialEditor/Private/UI/Widgets/Editor/SDMMaterialPropertySelector.cpp
+++ b/Engine/Plugins/Experimental/DynamicMaterial/Source/DynamicMaterialEditor/Private/UI/Widgets/Editor/SDMMaterialPropertySelector.cpp
@@ -4,7 +4,6 @@
 
 #include "Components/DMMaterialProperty.h"
 #include "DetailLayoutBuilder.h"
-#include "DynamicMaterialEditorSettings.h"
 #include "Framework/Application/SlateApplication.h"
 #include "Model/DynamicMaterialModel.h"
 #include "Model/DynamicMaterialModelEditorOnlyData.h"
@@ -15,6 +14,7 @@
 #include "Widgets/Images/SImage.h"
 #include "Widgets/Input/SCheckBox.h"
 #include "Widgets/Layout/SBox.h"
+#include "Widgets/Layout/SGridPanel.h"
 #include "Widgets/Layout/SScrollBox.h"
 #include "Widgets/SBoxPanel.h"
 #include "Widgets/SNullWidget.h"
@@ -22,6 +22,15 @@
 
 #define LOCTEXT_NAMESPACE "SDMMaterialPropertySelector"
 
+namespace UE::DynamicMaterialEditor::Private
+{
+	namespace PropertySelectorColumns
+	{
+		constexpr int32 Enable = 0;
+		constexpr int32 Select = 1;
+	}
+}
+
 void SDMMaterialPropertySelector::PrivateRegisterAttributes(FSlateAttributeDescriptor::FInitializer&)
 {
 }
@@ -111,6 +120,52 @@ TSharedPtr<SDMMaterialSlotEditor> SDMMaterialPropertySelector::GetSlotEditorWidg
 	return EditorWidget->GetSlotEditorWidget();
 }
 
+TSharedRef<SWidget> SDMMaterialPropertySelector::CreateSlot_PropertyList()
+{
+	using namespace UE::DynamicMaterialEditor::Private;
+
+	TSharedRef<SGridPanel> NewSlotList = SNew(SGridPanel)
+		.FillColumn(PropertySelectorColumns::Select, 1.f);
+
+	UDynamicMaterialModelEditorOnlyData* EditorOnlyData = GetEditorOnlyData();
+
+	if (!EditorOnlyData)
+	{
+		return NewSlotList;
+	}
+
+	int32 Row = 0;
+
+	NewSlotList->AddSlot(PropertySelectorColumns::Select, Row)
+		[
+			CreateSlot_SelectButton(EDMMaterialPropertyType::None)
+		];
+
+	++Row;
+
+	for (const TPair<EDMMaterialPropertyType, UDMMaterialProperty*>& PropertyPair : EditorOnlyData->GetMaterialProperties())
+	{
+		if (IsCustomMaterialProperty(PropertyPair.Key))
+		{
+			continue;
+		}
+
+		NewSlotList->AddSlot(PropertySelectorColumns::Enable, Row)
+			[
+				CreateSlot_EnabledButton(PropertyPair.Key)
+			];
+
+		NewSlotList->AddSlot(PropertySelectorColumns::Select, Row)
+			[
+				CreateSlot_SelectButton(PropertyPair.Key)
+			];
+
+		++Row;
+	}
+
+	return NewSlotList;
+}
+
 TSharedRef<SWidget> SDMMaterialPropertySelector::CreateSlot_EnabledButton(EDMMaterialPropertyType InMaterialProperty)
 {
 	const FText Format = LOCTEXT("PropertyEnableFormat", "Toggle the {0} property.\n\nProperty must be valid for the Material Type.");
@@ -127,34 +182,9 @@ TSharedRef<SWidget> SDMMaterialPropertySelector::CreateSlot_EnabledButton(EDMMat
 
 TSharedRef<SWidget> SDMMaterialPropertySelector::CreateSlot_SelectButton(EDMMaterialPropertyType InMaterialProperty)
 {
-	EDMMaterialEditorLayout Layout = EDMMaterialEditorLayout::First;
-
-	if (UDynamicMaterialEditorSettings* Settings = UDynamicMaterialEditorSettings::Get())
-	{
-		Layout = Settings->Layout;
-	}
-
-	UEnum* PropertyEnum = StaticEnum<EDMMaterialPropertyType>();
-
-	auto UseShortName = [Layout]()
-		{
-			switch (Layout)
-			{
-				case EDMMaterialEditorLayout::LeftSlim:
-					return true;
-
-				default:
-					return false;
-			}
-		};
-
-	const bool bUseShortName = UseShortName();
-
 	const FText ButtonText = InMaterialProperty == EDMMaterialPropertyType::None
 		? LOCTEXT("GlobalSettings", "Global Settings")
-		: (bUseShortName 
-			? UE::DynamicMaterialEditor::Private::GetMaterialPropertyShortDisplayName(InMaterialProperty)
-			: PropertyEnum->GetDisplayNameTextByValue(static_cast<int64>(InMaterialProperty)));
+		: StaticEnum<EDMMaterialPropertyType>()->GetDisplayNameTextByValue(static_cast<int64>(InMaterialProperty));
 
 	const FText Format = LOCTEXT("PropertySelectFormat", "Edit the {0} property.");
 	UEnum* MaterialPropertyEnum = StaticEnum<EDMMaterialPropertyType>();
@@ -174,8 +204,7 @@ TSharedRef<SWidget> SDMMaterialPropertySelector::CreateSlot_SelectButton(EDMMate
 		.Content()
 		[
 			SNew(SBox)
-			.WidthOverride(bUseShortName ? 42.f : 135.f)
-			.Clipping(EWidgetClipping::ClipToBounds)
+			.WidthOverride(135.f)
 			[
 				SNew(SHorizontalBox)
 				+SHorizontalBox::Slot()
@@ -187,7 +216,7 @@ TSharedRef<SWidget> SDMMaterialPropertySelector::CreateSlot_SelectButton(EDMMate
 					.ColorAndOpacity(this, &SDMMaterialPropertySelector::GetPropertySelectButtonChipColor, InMaterialProperty)
 				]
 				+SHorizontalBox::Slot()
-				.Padding(bUseShortName ? 2.f : 10.f, 6.f)
+				.Padding(10.f, 6.f)
 				.VAlign(VAlign_Center)
 				.FillWidth(1.f)
 				[
diff --git a/Engine/Plugins/Experimental/DynamicMaterial/Source/DynamicMaterialEditor/Private/UI/Widgets/Editor/SDMMaterialPropertySelector.h b/Engine/Plugins/Experimental/DynamicMaterial/Source/DynamicMaterialEditor/Private/UI/Widgets/Editor/SDMMaterialPropertySelector.h
--- a/Engine/Plugins/Experimental/DynamicMaterial/Source/DynamicMaterialEditor/Private/UI/Widgets/Editor/SDMMaterialPropertySelector.h
+++ b/Engine/Plugins/Experimental/DynamicMaterial/Source/DynamicMaterialEditor/Private/UI/Widgets/Editor/SDMMaterialPropertySelector.h
@@ -16,15 +16,6 @@ class UDynamicMaterialModelEditorOnlyData;
 enum class ECheckBoxState : uint8;
 enum class EDMMaterialPropertyType : uint8;
 
-namespace UE::DynamicMaterialEditor::Private
-{
-	namespace PropertySelectorColumns
-	{
-		constexpr int32 Enable = 0;
-		constexpr int32 Select = 1;
-	}
-}
-
 class SDMMaterialPropertySelector : public SCompoundWidget
 {
 	SLATE_DECLARE_WIDGET(SDMMaterialPropertySelector, SCompoundWidget)
@@ -53,11 +44,11 @@ protected:
 
 	TSharedPtr<SDMMaterialSlotEditor> GetSlotEditorWidget() const;
 
-	virtual TSharedRef<SWidget> CreateSlot_PropertyList() = 0;
+	TSharedRef<SWidget> CreateSlot_PropertyList();
 
 	TSharedRef<SWidget> CreateSlot_EnabledButton(EDMMaterialPropertyType InMaterialProperty);
 
-	virtual TSharedRef<SWidget> CreateSlot_SelectButton(EDMMaterialPropertyType InMaterialProperty) = 0;
+	TSharedRef<SWidget> CreateSlot_SelectButton(EDMMaterialPropertyType InMaterialProperty);
 
 	bool IsPropertyEnabled(EDMMaterialPropertyType InMaterialProperty) const;
 
diff --git a/Engine/Plugins/Experimental/DynamicMaterial/Source/DynamicMaterialEditor/Private/UI/Widgets/SDMMaterialDesigner.cpp b/Engine/Plugins/Experimental/DynamicMaterial/Source/DynamicMaterialEditor/Private/UI/Widgets/SDMMaterialDesigner.cpp
--- a/Engine/Plugins/Experimental/DynamicMaterial/Source/DynamicMaterialEditor/Private/UI/Widgets/SDMMaterialDesigner.cpp
+++ b/Engine/Plugins/Experimental/DynamicMaterial/Source/DynamicMaterialEditor/Private/UI/Widgets/SDMMaterialDesigner.cpp
@@ -4,9 +4,6 @@
 
 #include "Containers/Set.h"
 #include "DynamicMaterialEditorSettings.h"
-#include "DynamicMaterialModule.h"
-#include "Editor/SDMMaterialComponentEditor.h"
-#include "Editor/SDMMaterialSlotEditor.h"
 #include "GameFramework/Actor.h"
 #include "Material/DynamicMaterialInstance.h"
 #include "Model/DynamicMaterialModel.h"
@@ -14,11 +11,6 @@
 #include "Model/DynamicMaterialModelEditorOnlyData.h"
 #include "SAssetDropTarget.h"
 #include "UI/Utils/DMDropTargetPrivateSetter.h"
-#include "UI/Widgets/Editor/EditorLayouts/SDMMaterialEditor_Left.h"
-#include "UI/Widgets/Editor/EditorLayouts/SDMMaterialEditor_LeftSlim.h"
-#include "UI/Widgets/Editor/EditorLayouts/SDMMaterialEditor_TopHorizontal.h"
-#include "UI/Widgets/Editor/EditorLayouts/SDMMaterialEditor_TopSlim.h"
-#include "UI/Widgets/Editor/EditorLayouts/SDMMaterialEditor_TopVertical.h"
 #include "UI/Widgets/SDMActorMaterialSelector.h"
 #include "UI/Widgets/SDMMaterialEditor.h"
 #include "UI/Widgets/SDMMaterialSelectPrompt.h"
@@ -30,19 +22,6 @@ void SDMMaterialDesigner::PrivateRegisterAttributes(FSlateAttributeDescriptor::F
 {
 }
 
-SDMMaterialDesigner::~SDMMaterialDesigner()
-{
-	if (!FDynamicMaterialModule::AreUObjectsSafe())
-	{
-		return;
-	}
-
-	if (UDynamicMaterialEditorSettings* Settings = GetMutableDefault<UDynamicMaterialEditorSettings>())
-	{
-		Settings->GetOnSettingsChanged().RemoveAll(this);
-	}
-}
-
 void SDMMaterialDesigner::Construct(const FArguments& InArgs)
 {
 	SetCanTick(true);
@@ -50,11 +29,6 @@ void SDMMaterialDesigner::Construct(const FArguments& InArgs)
 	ContentSlot = TDMWidgetSlot<SWidget>(SharedThis(this), 0, SNullWidget::NullWidget);
 
 	SetSelectPromptView();
-
-	if (UDynamicMaterialEditorSettings* Settings = GetMutableDefault<UDynamicMaterialEditorSettings>())
-	{
-		Settings->GetOnSettingsChanged().AddSP(this, &SDMMaterialDesigner::OnSettingsChanged);
-	}
 }
 
 bool SDMMaterialDesigner::OpenMaterialModelBase(UDynamicMaterialModelBase* InMaterialModelBase)
@@ -277,26 +251,18 @@ void SDMMaterialDesigner::SetWizardView(const FDMObjectMaterialProperty& InObjec
 
 void SDMMaterialDesigner::SetEditorView(UDynamicMaterialModelBase* InMaterialModelBase)
 {
-	EDMMaterialEditorLayout Layout = EDMMaterialEditorLayout::Left;
-
-	if (UDynamicMaterialEditorSettings* Settings = GetMutableDefault<UDynamicMaterialEditorSettings>())
-	{
-		Layout = Settings->Layout;
-	}
+	TSharedRef<SDMMaterialEditor> Editor = SNew(SDMMaterialEditor, SharedThis(this))
+		.MaterialModelBase(InMaterialModelBase);
 
-	SetEditorLayout(Layout, InMaterialModelBase);
+	SetWidget(Editor, /* Include Drop Target */ true);
 }
 
 void SDMMaterialDesigner::SetEditorView(const FDMObjectMaterialProperty& InObjectMaterialProperty)
 {
-	EDMMaterialEditorLayout Layout = EDMMaterialEditorLayout::Left;
-
-	if (UDynamicMaterialEditorSettings* Settings = GetMutableDefault<UDynamicMaterialEditorSettings>())
-	{
-		Layout = Settings->Layout;
-	}
+	TSharedRef<SDMMaterialEditor> Editor = SNew(SDMMaterialEditor, SharedThis(this))
+		.MaterialProperty(InObjectMaterialProperty);
 
-	SetEditorLayout(Layout, InObjectMaterialProperty);
+	SetWidget(Editor, /* Include Drop Target */ true);
 }
 
 void SDMMaterialDesigner::SetWidget(const TSharedRef<SWidget>& InWidget, bool bInIncludeAssetDropTarget)
@@ -412,181 +378,3 @@ void SDMMaterialDesigner::OnAssetsDropped(const FDragDropEvent& InDragDropEvent,
 		}
 	}
 }
-
-void SDMMaterialDesigner::OnSettingsChanged(const FPropertyChangedEvent& InPropertyChangedEvent)
-{
-	if (!Content.IsValid() || Content->GetWidgetClass().GetWidgetType() != SDMMaterialEditor::StaticWidgetClass().GetWidgetType())
-	{
-		return;
-	}
-
-	const FName MemberName = InPropertyChangedEvent.GetMemberPropertyName();
-
-	if (MemberName == GET_MEMBER_NAME_CHECKED(UDynamicMaterialEditorSettings, Layout))
-	{
-		OnLayoutChanged();
-	}
-}
-
-void SDMMaterialDesigner::OnLayoutChanged()
-{
-	UDynamicMaterialEditorSettings* Settings = GetMutableDefault<UDynamicMaterialEditorSettings>();
-
-	if (!Settings)
-	{
-		return;
-	}
-
-	// Already assured
-	TSharedPtr<SDMMaterialEditor> CurrentEditor = StaticCastSharedRef<SDMMaterialEditor>(Content.ToSharedRef());
-
-	UDynamicMaterialModelBase* MaterialModelBase = CurrentEditor->GetMaterialModelBase();
-	const FDMObjectMaterialProperty* MaterialObjectProperty = CurrentEditor->GetMaterialObjectProperty();
-	UDMMaterialSlot* EditedSlot = nullptr;
-	UDMMaterialComponent* EditedComponent = nullptr;
-
-	if (TSharedPtr<SDMMaterialSlotEditor> SlotEditorWidget = CurrentEditor->GetSlotEditorWidget())
-	{
-		EditedSlot = SlotEditorWidget->GetSlot();
-	}
-
-	if (TSharedPtr<SDMMaterialComponentEditor> ComponentEditorWidget = CurrentEditor->GetComponentEditorWidget())
-	{
-		EditedComponent = ComponentEditorWidget->GetComponent();
-	}
-
-	if (MaterialObjectProperty)
-	{
-		if (!SetEditorLayout(Settings->Layout, *MaterialObjectProperty))
-		{
-			return;
-		}
-	}
-	else
-	{
-		if (!SetEditorLayout(Settings->Layout, MaterialModelBase))
-		{
-			return;
-		}
-	}
-
-	TSharedRef<SDMMaterialEditor> NewEditor = StaticCastSharedRef<SDMMaterialEditor>(Content.ToSharedRef());
-
-	if (EditedSlot)
-	{
-		NewEditor->EditSlot(EditedSlot);
-	}
-
-	if (EditedComponent)
-	{
-		NewEditor->EditComponent(EditedComponent);
-	}
-}
-
-bool SDMMaterialDesigner::SetEditorLayout(EDMMaterialEditorLayout InLayout, UDynamicMaterialModelBase* InMaterialModelBase)
-{
-	TSharedPtr<SDMMaterialEditor> NewEditor;
-
-	switch (InLayout)
-	{
-		case EDMMaterialEditorLayout::Left:
-			NewEditor = SNew(SDMMaterialEditor_Left, SharedThis(this))
-				.MaterialModelBase(InMaterialModelBase);
-			break;
-
-		case EDMMaterialEditorLayout::LeftAutoHide:
-			break;
-
-		case EDMMaterialEditorLayout::LeftSlim:
-			NewEditor = SNew(SDMMaterialEditor_LeftSlim, SharedThis(this))
-				.MaterialModelBase(InMaterialModelBase);
-			break;
-
-		case EDMMaterialEditorLayout::TopHorizontal:
-			NewEditor = SNew(SDMMaterialEditor_TopHorizontal, SharedThis(this))
-				.MaterialModelBase(InMaterialModelBase);
-			break;
-
-		case EDMMaterialEditorLayout::TopHorizontalAutoHide:
-			break;
-
-		case EDMMaterialEditorLayout::TopVertical:
-			NewEditor = SNew(SDMMaterialEditor_TopVertical, SharedThis(this))
-				.MaterialModelBase(InMaterialModelBase);
-			break;
-
-		case EDMMaterialEditorLayout::TopVerticalAutoHide:
-			break;
-
-		case EDMMaterialEditorLayout::TopSlim:
-			NewEditor = SNew(SDMMaterialEditor_TopSlim, SharedThis(this))
-				.MaterialModelBase(InMaterialModelBase);
-			break;
-
-		default:
-			break;
-	}
-
-	if (!NewEditor.IsValid())
-	{
-		return false;
-	}
-
-	SetWidget(NewEditor.ToSharedRef(), /* Include Drop Target */ true);
-
-	return true;
-}
-
-bool SDMMaterialDesigner::SetEditorLayout(EDMMaterialEditorLayout InLayout, const FDMObjectMaterialProperty& InObjectMaterialProperty)
-{
-	TSharedPtr<SDMMaterialEditor> NewEditor;
-
-	switch (InLayout)
-	{
-		case EDMMaterialEditorLayout::Left:
-			NewEditor = SNew(SDMMaterialEditor_Left, SharedThis(this))
-				.MaterialProperty(InObjectMaterialProperty);
-			break;
-
-		case EDMMaterialEditorLayout::LeftAutoHide:
-			break;
-
-		case EDMMaterialEditorLayout::LeftSlim:
-			NewEditor = SNew(SDMMaterialEditor_LeftSlim, SharedThis(this))
-				.MaterialProperty(InObjectMaterialProperty);
-			break;
-
-		case EDMMaterialEditorLayout::TopHorizontal:
-			NewEditor = SNew(SDMMaterialEditor_TopHorizontal, SharedThis(this))
-				.MaterialProperty(InObjectMaterialProperty);
-			break;
-
-		case EDMMaterialEditorLayout::TopHorizontalAutoHide:
-			break;
-
-		case EDMMaterialEditorLayout::TopVertical:
-			NewEditor = SNew(SDMMaterialEditor_TopVertical, SharedThis(this))
-				.MaterialProperty(InObjectMaterialProperty);
-			break;
-
-		case EDMMaterialEditorLayout::TopVerticalAutoHide:
-			break;
-
-		case EDMMaterialEditorLayout::TopSlim:
-			NewEditor = SNew(SDMMaterialEditor_TopSlim, SharedThis(this))
-				.MaterialProperty(InObjectMaterialProperty);
-			break;
-
-		default:
-			break;
-	}
-
-	if (!NewEditor.IsValid())
-	{
-		return false;
-	}
-
-	SetWidget(NewEditor.ToSharedRef(), /* Include Drop Target */ true);
-
-	return true;
-}
diff --git a/Engine/Plugins/Experimental/DynamicMaterial/Source/DynamicMaterialEditor/Private/UI/Widgets/SDMMaterialDesigner.h b/Engine/Plugins/Experimental/DynamicMaterial/Source/DynamicMaterialEditor/Private/UI/Widgets/SDMMaterialDesigner.h
--- a/Engine/Plugins/Experimental/DynamicMaterial/Source/DynamicMaterialEditor/Private/UI/Widgets/SDMMaterialDesigner.h
+++ b/Engine/Plugins/Experimental/DynamicMaterial/Source/DynamicMaterialEditor/Private/UI/Widgets/SDMMaterialDesigner.h
@@ -14,9 +14,7 @@ class SDMMaterialWizard;
 class UDynamicMaterialInstance;
 class UDynamicMaterialModel;
 class UDynamicMaterialModelBase;
-enum class EDMMaterialEditorLayout : uint8;
 struct FDMObjectMaterialProperty;
-struct FPropertyChangedEvent;
 
 class SDMMaterialDesigner : public SCompoundWidget
 {
@@ -28,7 +26,7 @@ class SDMMaterialDesigner : public SCompoundWidget
 public:
 	static bool IsFollowingSelection();
 
-	virtual ~SDMMaterialDesigner() override;
+	virtual ~SDMMaterialDesigner() override = default;
 
 	void Construct(const FArguments& InArgs);
 
@@ -98,13 +96,4 @@ protected:
 	bool OnAssetDraggedOver(TArrayView<FAssetData> InAssets);
 
 	void OnAssetsDropped(const FDragDropEvent& InDragDropEvent, TArrayView<FAssetData> InAssets);
-
-	/** Events */
-	void OnSettingsChanged(const FPropertyChangedEvent& InPropertyChangedEvent);
-
-	void OnLayoutChanged();
-
-	bool SetEditorLayout(EDMMaterialEditorLayout InLayout, UDynamicMaterialModelBase* InMaterialModelBase);
-
-	bool SetEditorLayout(EDMMaterialEditorLayout InLayout, const FDMObjectMaterialProperty& InObjectMaterialProperty);
 };
diff --git a/Engine/Plugins/Experimental/DynamicMaterial/Source/DynamicMaterialEditor/Private/UI/Widgets/SDMMaterialEditor.cpp b/Engine/Plugins/Experimental/DynamicMaterial/Source/DynamicMaterialEditor/Private/UI/Widgets/SDMMaterialEditor.cpp
--- a/Engine/Plugins/Experimental/DynamicMaterial/Source/DynamicMaterialEditor/Private/UI/Widgets/SDMMaterialEditor.cpp
+++ b/Engine/Plugins/Experimental/DynamicMaterial/Source/DynamicMaterialEditor/Private/UI/Widgets/SDMMaterialEditor.cpp
@@ -32,6 +32,37 @@
 
 #define LOCTEXT_NAMESPACE "SDMMaterialEditor"
 
+namespace UE::DynamicMaterialEditor::Private
+{
+	namespace SlotList
+	{
+		constexpr int32 ToolBar = 0;
+		constexpr int32 MainLayout = 1;
+		constexpr int32 StatusBar = 2;
+
+		namespace Main
+		{
+			constexpr int32 Left = 0;
+			constexpr int32 Right = 1;
+		}
+
+		namespace Left
+		{
+			constexpr int32 Preview = 0;
+			constexpr int32 PropertySelector = 1;
+		}
+
+		namespace Right
+		{
+			constexpr int32 SlotEditor = 0;
+			constexpr int32 ComponentEditor = 1;
+		}
+	}
+
+	const TCHAR* EditorDarkBackground = TEXT("Brushes.Title");
+	const TCHAR* EditorLightBackground = TEXT("Brushes.Header");
+}
+
 void SDMMaterialEditor::PrivateRegisterAttributes(FSlateAttributeDescriptor::FInitializer&)
 {
 }
@@ -208,6 +239,8 @@ void SDMMaterialEditor::EditSlot(UDMMaterialSlot* InSlot, bool bInForceRefresh)
 		return;
 	}
 
+	RightSlot.Invalidate();
+
 	SlotEditorSlot.Invalidate();
 	SplitterSlot = nullptr;
 	SlotToEdit = InSlot;
@@ -239,6 +272,7 @@ void SDMMaterialEditor::EditComponent(UDMMaterialComponent* InComponent, bool bI
 
 	if (bGlobalSettingsMode)
 	{
+		RightSlot.Invalidate();
 		SlotEditorSlot.Invalidate();
 		SplitterSlot = nullptr;
 	}
@@ -258,6 +292,7 @@ void SDMMaterialEditor::EditGlobalSettings(bool bInForceRefresh)
 
 	if (!bGlobalSettingsMode)
 	{
+		RightSlot.Invalidate();
 		SlotEditorSlot.Invalidate();
 		SplitterSlot = nullptr;
 		ComponentEditorSlot.Invalidate();
@@ -465,19 +500,28 @@ void SDMMaterialEditor::ValidateSlots()
 	}
 	else
 	{
-		ValidateSlots_Main();
-
-		if (MaterialPreviewSlot.HasBeenInvalidated())
+		if (LeftSlot.HasBeenInvalidated())
 		{
-			MaterialPreviewSlot << CreateSlot_Preview();
+			LeftSlot << CreateSlot_Left();
 		}
-
-		if (PropertySelectorSlot.HasBeenInvalidated())
+		else
 		{
-			PropertySelectorSlot << CreateSlot_PropertySelector();
+			if (MaterialPreviewSlot.HasBeenInvalidated())
+			{
+				MaterialPreviewSlot << CreateSlot_Preview();
+			}
+
+			if (PropertySelectorSlot.HasBeenInvalidated())
+			{
+				PropertySelectorSlot << CreateSlot_PropertySelector();
+			}
 		}
 
-		if (bGlobalSettingsMode)
+		if (RightSlot.HasBeenInvalidated())
+		{
+			RightSlot << CreateSlot_Right();
+		}
+		else if (bGlobalSettingsMode)
 		{
 			if (GlobalSettingsEditorSlot.HasBeenInvalidated())
 			{
@@ -521,15 +565,14 @@ void SDMMaterialEditor::ClearSlots()
 	ContentSlot.ClearWidget();
 	ToolBarSlot.ClearWidget();
 	MainSlot.ClearWidget();
-	SlotEditorSlot.ClearWidget();
+	LeftSlot.ClearWidget();
+	RightSlot.ClearWidget();
 	MaterialPreviewSlot.ClearWidget();
 	PropertySelectorSlot.ClearWidget();
-	GlobalSettingsEditorSlot.ClearWidget();
+	SlotEditorSlot.ClearWidget();
 	SplitterSlot = nullptr;
 	ComponentEditorSlot.ClearWidget();
 	StatusBarSlot.ClearWidget();
-
-	ClearSlots_Main();
 }
 
 void SDMMaterialEditor::CreateLayout()
@@ -583,11 +626,179 @@ TSharedRef<SDMToolBar> SDMMaterialEditor::CreateSlot_ToolBar()
 	);
 }
 
+TSharedRef<SWidget> SDMMaterialEditor::CreateSlot_Main()
+{
+	SHorizontalBox::FSlot* LeftSlotPtr = nullptr;
+	SHorizontalBox::FSlot* RightSlotPtr = nullptr;
+
+	TSharedRef<SHorizontalBox> NewMain = SNew(SHorizontalBox)
+		+ SHorizontalBox::Slot()
+		.Expose(LeftSlotPtr)
+		.AutoWidth()
+		[
+			SNullWidget::NullWidget
+		]
+
+		+ SHorizontalBox::Slot()
+		.Expose(RightSlotPtr)
+		.FillWidth(1.0f)
+		[
+			SNullWidget::NullWidget
+		];
+
+	LeftSlot = TDMWidgetSlot<SWidget>(LeftSlotPtr, CreateSlot_Left());
+	RightSlot = TDMWidgetSlot<SWidget>(RightSlotPtr, CreateSlot_Right());
+
+	return NewMain;
+}
+
+TSharedRef<SWidget> SDMMaterialEditor::CreateSlot_Left()
+{
+	using namespace UE::DynamicMaterialEditor::Private;
+
+	SVerticalBox::FSlot* MaterialPreviewSlotPtr = nullptr;
+	SVerticalBox::FSlot* PropertySelectorSlotPtr = nullptr;
+
+	TSharedRef<SWidget> NewLeft = SNew(SBorder)
+		.BorderImage(FAppStyle::GetBrush(EditorDarkBackground))
+		.Padding(5.f)
+		[
+			SNew(SVerticalBox)
+			+ SVerticalBox::Slot()
+			.Expose(MaterialPreviewSlotPtr)
+			.AutoHeight()
+			.Padding(0.f)
+			[
+				SNullWidget::NullWidget
+			]
+
+			+ SVerticalBox::Slot()
+			.Expose(PropertySelectorSlotPtr)
+			.FillHeight(1.0f)
+			.Padding(0.f, 5.f, 0.f, 0.f)
+			[
+				SNullWidget::NullWidget
+			]
+		];		
+
+	MaterialPreviewSlot = TDMWidgetSlot<SDMMaterialPreview>(MaterialPreviewSlotPtr, CreateSlot_Preview());
+	PropertySelectorSlot = TDMWidgetSlot<SDMMaterialPropertySelector>(PropertySelectorSlotPtr, CreateSlot_PropertySelector());
+
+	return NewLeft;
+}
+
+TSharedRef<SWidget> SDMMaterialEditor::CreateSlot_Right()
+{
+	using namespace UE::DynamicMaterialEditor::Private;
+
+	const bool bHasSlotToEdit = SlotToEdit.IsValid();
+
+	if (!bGlobalSettingsMode && !bHasSlotToEdit)
+	{
+		bGlobalSettingsMode = true;
+	}
+	else if (bHasSlotToEdit)
+	{
+		bGlobalSettingsMode = false;
+	}
+
+	return SNew(SBorder)
+		.BorderImage(FAppStyle::GetBrush(EditorDarkBackground))
+		.Padding(FMargin(0.f, 5.f))
+		[
+			bGlobalSettingsMode
+				? CreateSlot_Right_GlobalSettings()
+				: CreateSlot_Right_Slot()
+		];
+}
+
+TSharedRef<SWidget> SDMMaterialEditor::CreateSlot_Right_GlobalSettings()
+{
+	using namespace UE::DynamicMaterialEditor::Private;
+
+	SScrollBox::FSlot* GlobalSettingsSlotPtr = nullptr;
+
+	TSharedRef<SBorder> NewRight = SNew(SBorder)
+		.BorderImage(FAppStyle::GetBrush(EditorLightBackground))
+		.Padding(0.f)
+		[
+			SNew(SScrollBox)
+			+ SScrollBox::Slot()
+			.Expose(GlobalSettingsSlotPtr)
+			.VAlign(EVerticalAlignment::VAlign_Fill)
+			[
+				SNullWidget::NullWidget
+			]
+		];
+
+	GlobalSettingsEditorSlot = TDMWidgetSlot<SDMMaterialGlobalSettingsEditor>(GlobalSettingsSlotPtr, CreateSlot_GlobalSettingsEditor());
+
+	return NewRight;
+}
+
 TSharedRef<SDMMaterialGlobalSettingsEditor> SDMMaterialEditor::CreateSlot_GlobalSettingsEditor()
 {
 	return SNew(SDMMaterialGlobalSettingsEditor, SharedThis(this), GetMaterialModelBase());
 }
 
+TSharedRef<SWidget> SDMMaterialEditor::CreateSlot_Right_Slot()
+{
+	using namespace UE::DynamicMaterialEditor::Private;
+
+	float SplitterValue = 0.5;
+
+	if (UDynamicMaterialEditorSettings* Settings = UDynamicMaterialEditorSettings::Get())
+	{
+		SplitterValue = Settings->SplitterLocation;
+	}
+
+	TSharedPtr<SBorder> TopBox;
+	TSharedPtr<SBorder> BottomBox;
+
+	SSplitter::FSlot* ExplosedSlot = nullptr;
+
+	TSharedRef<SSplitter> NewRight = SNew(SSplitter)
+		.Style(FAppStyle::Get(), "DetailsView.Splitter")
+		.Orientation(Orient_Vertical)
+		.ResizeMode(ESplitterResizeMode::Fill)
+		.PhysicalSplitterHandleSize(5.0f)
+		.HitDetectionSplitterHandleSize(5.0f)
+		.OnSplitterFinishedResizing(this, &SDMMaterialEditor::OnRightSlotSplitterResized)
+
+		+ SSplitter::Slot()
+		.Expose(ExplosedSlot)
+		.Resizable(true)
+		.SizeRule(SSplitter::ESizeRule::FractionOfParent)
+		.MinSize(165)
+		.Value(SplitterValue)
+		[
+			SAssignNew(TopBox, SBorder)
+			.BorderImage(FAppStyle::GetBrush(EditorLightBackground))
+			[
+				SNullWidget::NullWidget
+			]
+		]
+
+		+ SSplitter::Slot()
+		.Resizable(true)
+		.SizeRule(SSplitter::ESizeRule::FractionOfParent)
+		.MinSize(60)
+		.Value(1.f - SplitterValue)
+		[
+			SAssignNew(BottomBox, SBorder)
+			.BorderImage(FAppStyle::GetBrush(EditorLightBackground))
+			[
+				SNullWidget::NullWidget
+			]
+		];
+
+	SplitterSlot = ExplosedSlot;
+	SlotEditorSlot = TDMWidgetSlot<SDMMaterialSlotEditor>(TopBox.ToSharedRef(), 0, CreateSlot_SlotEditor());
+	ComponentEditorSlot = TDMWidgetSlot<SDMMaterialComponentEditor>(BottomBox.ToSharedRef(), 0, CreateSlot_ComponentEditor());
+
+	return NewRight;
+}
+
 TSharedRef<SDMMaterialPreview> SDMMaterialEditor::CreateSlot_Preview()
 {
 	return SNew(SDMMaterialPreview, SharedThis(this), GetMaterialModelBase());
@@ -595,7 +806,7 @@ TSharedRef<SDMMaterialPreview> SDMMaterialEditor::CreateSlot_Preview()
 
 TSharedRef<SDMMaterialPropertySelector> SDMMaterialEditor::CreateSlot_PropertySelector()
 {
-	TSharedRef<SDMMaterialPropertySelector> NewPropertySelector = CreateSlot_PropertySelector_Impl();
+	TSharedRef<SDMMaterialPropertySelector> NewPropertySelector = SNew(SDMMaterialPropertySelector, SharedThis(this));
 
 	if (!PropertyToSelect.IsSet())
 	{
@@ -682,7 +893,7 @@ void SDMMaterialEditor::OnEnginePreExit()
 	MaterialPreviewSlot.ClearWidget();
 }
 
-void SDMMaterialEditor::OnEditorSplitterResized()
+void SDMMaterialEditor::OnRightSlotSplitterResized()
 {
 	UDynamicMaterialEditorSettings* Settings = UDynamicMaterialEditorSettings::Get();
 
diff --git a/Engine/Plugins/Experimental/DynamicMaterial/Source/DynamicMaterialEditor/Private/UI/Widgets/SDMMaterialEditor.h b/Engine/Plugins/Experimental/DynamicMaterial/Source/DynamicMaterialEditor/Private/UI/Widgets/SDMMaterialEditor.h
--- a/Engine/Plugins/Experimental/DynamicMaterial/Source/DynamicMaterialEditor/Private/UI/Widgets/SDMMaterialEditor.h
+++ b/Engine/Plugins/Experimental/DynamicMaterial/Source/DynamicMaterialEditor/Private/UI/Widgets/SDMMaterialEditor.h
@@ -28,12 +28,6 @@ class UDynamicMaterialModelBase;
 enum class EDMMaterialPropertyType : uint8;
 enum class EDMUpdateType : uint8;
 
-namespace UE::DynamicMaterialEditor::Private
-{
-	inline const TCHAR* EditorDarkBackground = TEXT("Brushes.Title");
-	inline const TCHAR* EditorLightBackground = TEXT("Brushes.Header");
-}
-
 class SDMMaterialEditor : public SCompoundWidget, public FSelfRegisteringEditorUndoClient
 {
 	SLATE_DECLARE_WIDGET(SDMMaterialEditor, SCompoundWidget)
@@ -80,11 +74,11 @@ public:
 	/** Actions */
 	void SelectProperty(EDMMaterialPropertyType InProperty, bool bInForceRefresh = false);
 
-	virtual void EditSlot(UDMMaterialSlot* InSlot, bool bInForceRefresh = false);
+	void EditSlot(UDMMaterialSlot* InSlot, bool bInForceRefresh = false);
 
-	virtual void EditComponent(UDMMaterialComponent* InComponent, bool bInForceRefresh = false);
+	void EditComponent(UDMMaterialComponent* InComponent, bool bInForceRefresh = false);
 
-	virtual void EditGlobalSettings(bool bInForceRefresh = false);
+	void EditGlobalSettings(bool bInForceRefresh = false);
 
 	void Validate();
 
@@ -110,6 +104,8 @@ protected:
 	TDMWidgetSlot<SWidget> ContentSlot;
 	TDMWidgetSlot<SDMToolBar> ToolBarSlot;
 	TDMWidgetSlot<SWidget> MainSlot;
+	TDMWidgetSlot<SWidget> LeftSlot;
+	TDMWidgetSlot<SWidget> RightSlot;
 	TDMWidgetSlot<SDMMaterialPreview> MaterialPreviewSlot;
 	TDMWidgetSlot<SDMMaterialPropertySelector> PropertySelectorSlot;
 	TDMWidgetSlot<SDMMaterialGlobalSettingsEditor> GlobalSettingsEditorSlot;
@@ -147,12 +143,8 @@ protected:
 
 	void ValidateSlots();
 
-	virtual void ValidateSlots_Main() = 0;
-
 	void ClearSlots();
 
-	virtual void ClearSlots_Main() = 0;
-
 	/** Slots */
 	void CreateLayout();
 
@@ -160,16 +152,22 @@ protected:
 
 	TSharedRef<SDMToolBar> CreateSlot_ToolBar();
 
-	virtual TSharedRef<SWidget> CreateSlot_Main() = 0;
+	TSharedRef<SWidget> CreateSlot_Main();
+
+	TSharedRef<SWidget> CreateSlot_Left();
+
+	TSharedRef<SWidget> CreateSlot_Right();
+
+	TSharedRef<SWidget> CreateSlot_Right_GlobalSettings();
 
 	TSharedRef<SDMMaterialGlobalSettingsEditor> CreateSlot_GlobalSettingsEditor();
 
+	TSharedRef<SWidget> CreateSlot_Right_Slot();
+
 	TSharedRef<SDMMaterialPreview> CreateSlot_Preview();
 
 	TSharedRef<SDMMaterialPropertySelector> CreateSlot_PropertySelector();
 
-	virtual TSharedRef<SDMMaterialPropertySelector> CreateSlot_PropertySelector_Impl() = 0;
-
 	TSharedRef<SDMMaterialSlotEditor> CreateSlot_SlotEditor();
 
 	TSharedRef<SDMMaterialComponentEditor> CreateSlot_ComponentEditor();
@@ -182,5 +180,5 @@ protected:
 	/** The material preview window is not cleaned up properly on uobject shutdown, so do it here. */
 	void OnEnginePreExit();
 
-	void OnEditorSplitterResized();
+	void OnRightSlotSplitterResized();
 };
diff --git a/Engine/Plugins/Experimental/DynamicMaterial/Source/DynamicMaterialEditor/Private/UI/Widgets/SDMMaterialWizard.cpp b/Engine/Plugins/Experimental/DynamicMaterial/Source/DynamicMaterialEditor/Private/UI/Widgets/SDMMaterialWizard.cpp
--- a/Engine/Plugins/Experimental/DynamicMaterial/Source/DynamicMaterialEditor/Private/UI/Widgets/SDMMaterialWizard.cpp
+++ b/Engine/Plugins/Experimental/DynamicMaterial/Source/DynamicMaterialEditor/Private/UI/Widgets/SDMMaterialWizard.cpp
@@ -12,7 +12,6 @@
 #include "Styling/SlateTypes.h"
 #include "Styling/StyleColors.h"
 #include "UI/Widgets/SDMMaterialDesigner.h"
-#include "Utils/DMPrivate.h"
 #include "Widgets/Input/SButton.h"
 #include "Widgets/Input/SCheckBox.h"
 #include "Widgets/Layout/SBox.h"
@@ -234,12 +233,15 @@ TSharedRef<SWidget> SDMMaterialWizard::CreateChannelList()
 
 			if (Preset->IsPropertyEnabled(Property.Key))
 			{
+				constexpr const TCHAR* ShortNameName = TEXT("ShortName");
+				const FString ShortName = MaterialPropertyEnum->GetMetaData(ShortNameName, MaterialPropertyEnum->GetIndexByValue(static_cast<int64>(Property.Key)));
+
 				ChannelPresets->AddSlot()
 					.Padding(TextPadding)
 					[
 						SNew(STextBlock)
 						.TextStyle(FDynamicMaterialEditorStyle::Get(), "RegularFont")
-						.Text(UE::DynamicMaterialEditor::Private::GetMaterialPropertyShortDisplayName(Property.Key))
+						.Text(!ShortName.IsEmpty() ? FText::FromString(ShortName) : MaterialPropertyEnum->GetDisplayNameTextByValue(static_cast<int64>(Property.Key)))
 					];
 			}
 		}
diff --git a/Engine/Plugins/Experimental/DynamicMaterial/Source/DynamicMaterialEditor/Private/Utils/DMPrivate.cpp b/Engine/Plugins/Experimental/DynamicMaterial/Source/DynamicMaterialEditor/Private/Utils/DMPrivate.cpp
--- a/Engine/Plugins/Experimental/DynamicMaterial/Source/DynamicMaterialEditor/Private/Utils/DMPrivate.cpp
+++ b/Engine/Plugins/Experimental/DynamicMaterial/Source/DynamicMaterialEditor/Private/Utils/DMPrivate.cpp
@@ -93,22 +93,6 @@ namespace UE::DynamicMaterialEditor::Private
 			FSlateNotificationManager::Get().AddNotification(Info);
 		}
 	}
-
-	FText GetMaterialPropertyShortDisplayName(EDMMaterialPropertyType InMaterialProperty)
-	{
-		constexpr const TCHAR* ShortNameName = TEXT("ShortName");
-
-		UEnum* PropertyEnum = StaticEnum<EDMMaterialPropertyType>();
-
-		const FString ShortName = PropertyEnum->GetMetaData(ShortNameName, PropertyEnum->GetIndexByValue(static_cast<int64>(InMaterialProperty)));
-
-		if (!ShortName.IsEmpty())
-		{
-			return FText::FromString(ShortName);
-		}
-
-		return PropertyEnum->GetDisplayNameTextByValue(static_cast<int64>(InMaterialProperty));
-	}
 }
 
 FDMMaterialLayerReference::FDMMaterialLayerReference()
diff --git a/Engine/Plugins/Experimental/DynamicMaterial/Source/DynamicMaterialEditor/Private/Utils/DMPrivate.h b/Engine/Plugins/Experimental/DynamicMaterial/Source/DynamicMaterialEditor/Private/Utils/DMPrivate.h
--- a/Engine/Plugins/Experimental/DynamicMaterial/Source/DynamicMaterialEditor/Private/Utils/DMPrivate.h
+++ b/Engine/Plugins/Experimental/DynamicMaterial/Source/DynamicMaterialEditor/Private/Utils/DMPrivate.h
@@ -5,7 +5,6 @@
 #include "Containers/Array.h"
 #include "HAL/Platform.h"
 #include "ScopedTransaction.h"
-#include "UI/Menus/DMMaterialSlotLayerMenus.h"
 #include "UI/Widgets/Editor/SDMMaterialSlotEditor.h"
 #include "UObject/WeakObjectPtr.h"
 
@@ -33,8 +32,6 @@ namespace UE::DynamicMaterialEditor::Private
 	bool IsCustomMaterialProperty(EDMMaterialPropertyType InMaterialProperty);
 
 	void LogError(const FString& InMessage, bool bInToast = true, const UObject* InSource = nullptr);
-
-	FText GetMaterialPropertyShortDisplayName(EDMMaterialPropertyType InMaterialProperty);
 }
 
 struct FDMMaterialLayerReference
