Material Designer
- Added parameter to Update function and event.

diff --git a/Engine/Plugins/Experimental/Avalanche/Source/AvalancheModifiers/Private/Modifiers/AvaGlobalOpacityModifier.cpp b/Engine/Plugins/Experimental/Avalanche/Source/AvalancheModifiers/Private/Modifiers/AvaGlobalOpacityModifier.cpp
--- a/Engine/Plugins/Experimental/Avalanche/Source/AvalancheModifiers/Private/Modifiers/AvaGlobalOpacityModifier.cpp
+++ b/Engine/Plugins/Experimental/Avalanche/Source/AvalancheModifiers/Private/Modifiers/AvaGlobalOpacityModifier.cpp
@@ -91,9 +91,9 @@ void UAvaGlobalOpacityModifier::OnActorMaterialRemoved(UMaterialInstanceDynamic*
 #endif
 }
 
-void UAvaGlobalOpacityModifier::OnDynamicMaterialValueChanged(UDMMaterialComponent* InComponent, EDMUpdateType InUpdateType)
+void UAvaGlobalOpacityModifier::OnDynamicMaterialValueChanged(UDMMaterialComponent* InComponent, UDMMaterialComponent* InSource, EDMUpdateType InUpdateType)
 {
-	if (InUpdateType != EDMUpdateType::Value)
+	if (!EnumHasAnyFlags(InUpdateType, EDMUpdateType::Value))
 	{
 		return;
 	}
diff --git a/Engine/Plugins/Experimental/Avalanche/Source/AvalancheModifiers/Public/Modifiers/AvaGlobalOpacityModifier.h b/Engine/Plugins/Experimental/Avalanche/Source/AvalancheModifiers/Public/Modifiers/AvaGlobalOpacityModifier.h
--- a/Engine/Plugins/Experimental/Avalanche/Source/AvalancheModifiers/Public/Modifiers/AvaGlobalOpacityModifier.h
+++ b/Engine/Plugins/Experimental/Avalanche/Source/AvalancheModifiers/Public/Modifiers/AvaGlobalOpacityModifier.h
@@ -41,7 +41,7 @@ protected:
 
 	virtual void OnActorMaterialAdded(UMaterialInstanceDynamic* InAdded) override;
 	virtual void OnActorMaterialRemoved(UMaterialInstanceDynamic* InRemoved) override;
-	void OnDynamicMaterialValueChanged(UDMMaterialComponent* InComponent, EDMUpdateType InUpdateType);
+	void OnDynamicMaterialValueChanged(UDMMaterialComponent* InComponent, UDMMaterialComponent* InSource, EDMUpdateType InUpdateType);
 
 	/** Global opacity to set on all Material Designer Instances */
 	UPROPERTY(EditInstanceOnly, Setter="SetGlobalOpacity", Getter="GetGlobalOpacity", Category="Material Parameter", meta=(ClampMin="0", ClampMax="1", AllowPrivateAccess="true"))
diff --git a/Engine/Plugins/Experimental/DynamicMaterial/Source/DynamicMaterial/Private/Components/DMMaterialComponent.cpp b/Engine/Plugins/Experimental/DynamicMaterial/Source/DynamicMaterial/Private/Components/DMMaterialComponent.cpp
--- a/Engine/Plugins/Experimental/DynamicMaterial/Source/DynamicMaterial/Private/Components/DMMaterialComponent.cpp
+++ b/Engine/Plugins/Experimental/DynamicMaterial/Source/DynamicMaterial/Private/Components/DMMaterialComponent.cpp
@@ -73,7 +73,7 @@ UDMMaterialComponent* UDMMaterialComponent::GetSubComponentByPath(FDMComponentPa
 	return nullptr;
 }
 
-void UDMMaterialComponent::Update(EDMUpdateType InUpdateType)
+void UDMMaterialComponent::Update(UDMMaterialComponent* InSource, EDMUpdateType InUpdateType)
 {
 	if (!FDMUpdateGuard::CanUpdate())
 	{
@@ -81,7 +81,7 @@ void UDMMaterialComponent::Update(EDMUpdateType InUpdateType)
 	}
 
 #if WITH_EDITOR
-	OnUpdate.Broadcast(this, InUpdateType);
+	OnUpdate.Broadcast(this, InSource, InUpdateType);
 #endif
 }
 
diff --git a/Engine/Plugins/Experimental/DynamicMaterial/Source/DynamicMaterial/Private/Components/DMMaterialValue.cpp b/Engine/Plugins/Experimental/DynamicMaterial/Source/DynamicMaterial/Private/Components/DMMaterialValue.cpp
--- a/Engine/Plugins/Experimental/DynamicMaterial/Source/DynamicMaterial/Private/Components/DMMaterialValue.cpp
+++ b/Engine/Plugins/Experimental/DynamicMaterial/Source/DynamicMaterial/Private/Components/DMMaterialValue.cpp
@@ -232,8 +232,8 @@ void UDMMaterialValue::SetShouldExposeParameter(bool bInExpose)
 	}
 
 	bExposeParameter = bInExpose;
-
-	Update(EDMUpdateType::Structure);
+	
+	Update(this, EDMUpdateType::Structure);
 }
 
 void UDMMaterialValue::OnComponentAdded()
@@ -386,12 +386,12 @@ void UDMMaterialValue::OnValueChanged(EDMUpdateType InUpdateType)
 		return;
 	}
 
-	Update(InUpdateType);
+	Update(this, InUpdateType);
 
 #if WITH_EDITOR
 	if (EnumHasAnyFlags(InUpdateType, EDMUpdateType::AllowParentUpdate) && ParentComponent)
 	{
-		ParentComponent->Update(InUpdateType);
+		ParentComponent->Update(this, InUpdateType);
 	}
 #endif
 }
@@ -415,7 +415,7 @@ void UDMMaterialValue::UpdateCachedParameterName()
 }
 #endif
 
-void UDMMaterialValue::Update(EDMUpdateType InUpdateType)
+void UDMMaterialValue::Update(UDMMaterialComponent* InSource, EDMUpdateType InUpdateType)
 {
 	if (!FDMUpdateGuard::CanUpdate())
 	{
@@ -441,7 +441,7 @@ void UDMMaterialValue::Update(EDMUpdateType InUpdateType)
 	}
 #endif
 
-	Super::Update(InUpdateType);
+	Super::Update(InSource, InUpdateType);
 
 	if (UDynamicMaterialModel* MaterialModel = GetMaterialModel())
 	{
diff --git a/Engine/Plugins/Experimental/DynamicMaterial/Source/DynamicMaterial/Private/Components/DMMaterialValueDynamic.cpp b/Engine/Plugins/Experimental/DynamicMaterial/Source/DynamicMaterial/Private/Components/DMMaterialValueDynamic.cpp
--- a/Engine/Plugins/Experimental/DynamicMaterial/Source/DynamicMaterial/Private/Components/DMMaterialValueDynamic.cpp
+++ b/Engine/Plugins/Experimental/DynamicMaterial/Source/DynamicMaterial/Private/Components/DMMaterialValueDynamic.cpp
@@ -48,7 +48,7 @@ UDMMaterialValue* UDMMaterialValueDynamic::GetParentValue() const
 	return Cast<UDMMaterialValue>(GeResolvedParentComponent());
 }
 
-void UDMMaterialValueDynamic::Update(EDMUpdateType InUpdateType)
+void UDMMaterialValueDynamic::Update(UDMMaterialComponent* InSource, EDMUpdateType InUpdateType)
 {
 	if (!FDMUpdateGuard::CanUpdate())
 	{
@@ -69,7 +69,7 @@ void UDMMaterialValueDynamic::Update(EDMUpdateType InUpdateType)
 	MarkComponentDirty();
 #endif
 
-	Super::Update(InUpdateType);
+	Super::Update(InSource, InUpdateType);
 
 	if (UDynamicMaterialModelDynamic* MaterialModelDynamic = GetMaterialModelDynamic())
 	{
@@ -127,7 +127,7 @@ void UDMMaterialValueDynamic::OnValueChanged()
 		return;
 	}
 
-	Update();
+	Update(this, EDMUpdateType::Value);
 }
 
 #if WITH_EDITOR
diff --git a/Engine/Plugins/Experimental/DynamicMaterial/Source/DynamicMaterial/Private/Components/DMRenderTargetRenderer.cpp b/Engine/Plugins/Experimental/DynamicMaterial/Source/DynamicMaterial/Private/Components/DMRenderTargetRenderer.cpp
--- a/Engine/Plugins/Experimental/DynamicMaterial/Source/DynamicMaterial/Private/Components/DMRenderTargetRenderer.cpp
+++ b/Engine/Plugins/Experimental/DynamicMaterial/Source/DynamicMaterial/Private/Components/DMRenderTargetRenderer.cpp
@@ -76,7 +76,7 @@ bool UDMRenderTargetRenderer::JsonDeserialize(const TSharedPtr<FJsonValue>& InJs
 	return false;
 }
 
-void UDMRenderTargetRenderer::Update(EDMUpdateType InUpdateType)
+void UDMRenderTargetRenderer::Update(UDMMaterialComponent* InSource, EDMUpdateType InUpdateType)
 {
 	if (!FDMUpdateGuard::CanUpdate())
 	{
@@ -99,10 +99,10 @@ void UDMRenderTargetRenderer::Update(EDMUpdateType InUpdateType)
 
 	if (UDMMaterialValueRenderTarget* RenderTarget = GetRenderTargetValue())
 	{
-		RenderTarget->Update(InUpdateType);
+		RenderTarget->Update(InSource, InUpdateType);
 	}
 
-	Super::Update(InUpdateType);
+	Super::Update(InSource, InUpdateType);
 }
 
 void UDMRenderTargetRenderer::PostLoad()
diff --git a/Engine/Plugins/Experimental/DynamicMaterial/Source/DynamicMaterial/Private/Components/DMTextureUV.cpp b/Engine/Plugins/Experimental/DynamicMaterial/Source/DynamicMaterial/Private/Components/DMTextureUV.cpp
--- a/Engine/Plugins/Experimental/DynamicMaterial/Source/DynamicMaterial/Private/Components/DMTextureUV.cpp
+++ b/Engine/Plugins/Experimental/DynamicMaterial/Source/DynamicMaterial/Private/Components/DMTextureUV.cpp
@@ -329,7 +329,7 @@ void UDMTextureUV::SetShouldExposeParameter(FName InPropertyName, int32 InCompon
 		ExposedParameters.Remove(ParamId);
 	}
 
-	Update(EDMUpdateType::Structure);
+	Update(this, EDMUpdateType::Structure);
 }
 
 void UDMTextureUV::PostEditorDuplicate(UDynamicMaterialModel* InMaterialModel, UDMMaterialComponent* InParent)
@@ -446,7 +446,7 @@ bool UDMTextureUV::Modify(bool bInAlwaysMarkDirty)
 }
 #endif
 
-void UDMTextureUV::Update(EDMUpdateType InUpdateType)
+void UDMTextureUV::Update(UDMMaterialComponent* InSource, EDMUpdateType InUpdateType)
 {
 	if (!FDMUpdateGuard::CanUpdate())
 	{
@@ -472,7 +472,7 @@ void UDMTextureUV::Update(EDMUpdateType InUpdateType)
 	}
 #endif
 
-	Super::Update(InUpdateType);
+	Super::Update(InSource, InUpdateType);
 
 	if (UDynamicMaterialModel* MaterialModel = GetMaterialModel())
 	{
@@ -645,12 +645,12 @@ void UDMTextureUV::OnTextureUVChanged(EDMUpdateType InUpdateType)
 		return;
 	}
 
-	Update(InUpdateType);
+	Update(this, InUpdateType);
 
 #if WITH_EDITOR
 	if (EnumHasAnyFlags(InUpdateType, EDMUpdateType::AllowParentUpdate) && ParentComponent)
 	{
-		ParentComponent->Update(InUpdateType);
+		ParentComponent->Update(this, InUpdateType);
 	}
 #endif
 }
diff --git a/Engine/Plugins/Experimental/DynamicMaterial/Source/DynamicMaterial/Private/Components/DMTextureUVDynamic.cpp b/Engine/Plugins/Experimental/DynamicMaterial/Source/DynamicMaterial/Private/Components/DMTextureUVDynamic.cpp
--- a/Engine/Plugins/Experimental/DynamicMaterial/Source/DynamicMaterial/Private/Components/DMTextureUVDynamic.cpp
+++ b/Engine/Plugins/Experimental/DynamicMaterial/Source/DynamicMaterial/Private/Components/DMTextureUVDynamic.cpp
@@ -154,11 +154,11 @@ void UDMTextureUVDynamic::CopyDynamicPropertiesTo(UDMMaterialComponent* InDestin
 		DestinationTextureUV->SetPivot(GetPivot());
 	}
 
-	DestinationTextureUV->Update(EDMUpdateType::Value);
+	DestinationTextureUV->Update(DestinationTextureUV, EDMUpdateType::Value);
 }
 #endif
 
-void UDMTextureUVDynamic::Update(EDMUpdateType InUpdateType)
+void UDMTextureUVDynamic::Update(UDMMaterialComponent* InSource, EDMUpdateType InUpdateType)
 {
 	if (!IsComponentValid())
 	{
@@ -174,7 +174,7 @@ void UDMTextureUVDynamic::Update(EDMUpdateType InUpdateType)
 	MarkComponentDirty();
 #endif
 
-	Super::Update(InUpdateType);
+	Super::Update(InSource, InUpdateType);
 
 	if (UDynamicMaterialModelDynamic* MaterialModelDynamic = GetMaterialModelDynamic())
 	{
@@ -232,7 +232,7 @@ void UDMTextureUVDynamic::OnTextureUVChanged()
 
 	if (FDMUpdateGuard::CanUpdate())
 	{
-		Update(EDMUpdateType::Value);
+		Update(this, EDMUpdateType::Value);
 	}
 }
 
diff --git a/Engine/Plugins/Experimental/DynamicMaterial/Source/DynamicMaterial/Private/Components/MaterialValues/DMMaterialValueRenderTarget.cpp b/Engine/Plugins/Experimental/DynamicMaterial/Source/DynamicMaterial/Private/Components/MaterialValues/DMMaterialValueRenderTarget.cpp
--- a/Engine/Plugins/Experimental/DynamicMaterial/Source/DynamicMaterial/Private/Components/MaterialValues/DMMaterialValueRenderTarget.cpp
+++ b/Engine/Plugins/Experimental/DynamicMaterial/Source/DynamicMaterial/Private/Components/MaterialValues/DMMaterialValueRenderTarget.cpp
@@ -142,14 +142,14 @@ void UDMMaterialValueRenderTarget::FlushCreateRenderTarget()
 	}
 }
 
-void UDMMaterialValueRenderTarget::Update(EDMUpdateType InUpdateType)
+void UDMMaterialValueRenderTarget::Update(UDMMaterialComponent* InSource, EDMUpdateType InUpdateType)
 {
 	if (!FDMUpdateGuard::CanUpdate())
 	{
 		return;
 	}
 
-	Super::Update(InUpdateType);
+	Super::Update(InSource, InUpdateType);
 
 	if (!IsValid(GetRenderTarget()))
 	{
diff --git a/Engine/Plugins/Experimental/DynamicMaterial/Source/DynamicMaterial/Private/Components/RenderTargetRenderers/DMRenderTargetTextRenderer.cpp b/Engine/Plugins/Experimental/DynamicMaterial/Source/DynamicMaterial/Private/Components/RenderTargetRenderers/DMRenderTargetTextRenderer.cpp
--- a/Engine/Plugins/Experimental/DynamicMaterial/Source/DynamicMaterial/Private/Components/RenderTargetRenderers/DMRenderTargetTextRenderer.cpp
+++ b/Engine/Plugins/Experimental/DynamicMaterial/Source/DynamicMaterial/Private/Components/RenderTargetRenderers/DMRenderTargetTextRenderer.cpp
@@ -258,7 +258,7 @@ void UDMRenderTargetTextRenderer::PostEditChangeProperty(FPropertyChangedEvent&
 		|| PropertyName == FDMRenderTargetTextRenderer::AutoWrapTextName)
 	{
 		// Cause details panel refresh.
-		Update(EDMUpdateType::Structure);
+		Update(this, EDMUpdateType::Structure);
 	}
 
 	if (FDMRenderTargetTextRenderer::PropertyNames.Contains(PropertyName))
diff --git a/Engine/Plugins/Experimental/DynamicMaterial/Source/DynamicMaterial/Public/Components/DMMaterialComponent.h b/Engine/Plugins/Experimental/DynamicMaterial/Source/DynamicMaterial/Public/Components/DMMaterialComponent.h
--- a/Engine/Plugins/Experimental/DynamicMaterial/Source/DynamicMaterial/Public/Components/DMMaterialComponent.h
+++ b/Engine/Plugins/Experimental/DynamicMaterial/Source/DynamicMaterial/Public/Components/DMMaterialComponent.h
@@ -60,7 +60,7 @@ public:
 
 	/** Event that is triggered when this component, or a sub-component, changes to trigger other updates in the model. */
 	UFUNCTION(BlueprintCallable, Category = "Material Designer")
-	DYNAMICMATERIAL_API virtual void Update(EDMUpdateType InUpdateType);
+	DYNAMICMATERIAL_API virtual void Update(UDMMaterialComponent* InSource, EDMUpdateType InUpdateType);
 
 #if WITH_EDITOR
 	/** Returns the complete path from the model to this component. */
@@ -139,7 +139,12 @@ public:
 	/** Called to ensure that the object hierarchy is correct. */
 	DYNAMICMATERIAL_API virtual void PostEditorDuplicate(UDynamicMaterialModel* InMaterialModel, UDMMaterialComponent* InParent);
 
-	DECLARE_MULTICAST_DELEGATE_TwoParams(FOnUpdate, UDMMaterialComponent*, EDMUpdateType)
+	DECLARE_MULTICAST_DELEGATE_ThreeParams(
+		FOnUpdate,
+		UDMMaterialComponent* /* Triggering Component */,
+		UDMMaterialComponent* /* Source Component */,
+		EDMUpdateType
+	)
 
 	/** Delegate called when this component's Update method is called. */
 	FOnUpdate::RegistrationType& GetOnUpdate() { return OnUpdate; }
diff --git a/Engine/Plugins/Experimental/DynamicMaterial/Source/DynamicMaterial/Public/Components/DMMaterialValue.h b/Engine/Plugins/Experimental/DynamicMaterial/Source/DynamicMaterial/Public/Components/DMMaterialValue.h
--- a/Engine/Plugins/Experimental/DynamicMaterial/Source/DynamicMaterial/Public/Components/DMMaterialValue.h
+++ b/Engine/Plugins/Experimental/DynamicMaterial/Source/DynamicMaterial/Public/Components/DMMaterialValue.h
@@ -139,7 +139,7 @@ public:
 #endif
  
 	//~ Begin UDMMaterialComponent
-	DYNAMICMATERIAL_API virtual void Update(EDMUpdateType InUpdateType) override;
+	DYNAMICMATERIAL_API virtual void Update(UDMMaterialComponent* InSource, EDMUpdateType InUpdateType) override;
 #if WITH_EDITOR
 	DYNAMICMATERIAL_API virtual void PostEditorDuplicate(UDynamicMaterialModel* InMaterialModel, UDMMaterialComponent* InParent) override;
 #endif
diff --git a/Engine/Plugins/Experimental/DynamicMaterial/Source/DynamicMaterial/Public/Components/DMMaterialValueDynamic.h b/Engine/Plugins/Experimental/DynamicMaterial/Source/DynamicMaterial/Public/Components/DMMaterialValueDynamic.h
--- a/Engine/Plugins/Experimental/DynamicMaterial/Source/DynamicMaterial/Public/Components/DMMaterialValueDynamic.h
+++ b/Engine/Plugins/Experimental/DynamicMaterial/Source/DynamicMaterial/Public/Components/DMMaterialValueDynamic.h
@@ -54,7 +54,7 @@ public:
 	virtual void SetMIDParameter(UMaterialInstanceDynamic* InMID) const PURE_VIRTUAL(UDMMaterialValueDynamic::SetMIDParameter)
 
 	//~ Begin UDMMaterialComponent
-	DYNAMICMATERIAL_API virtual void Update(EDMUpdateType InUpdateType = EDMUpdateType::Value) override;
+	DYNAMICMATERIAL_API virtual void Update(UDMMaterialComponent* InSource, EDMUpdateType InUpdateType) override;
 	//~ End UDMMaterialComponent
 
 	/** Non-editor implementation. */
diff --git a/Engine/Plugins/Experimental/DynamicMaterial/Source/DynamicMaterial/Public/Components/DMRenderTargetRenderer.h b/Engine/Plugins/Experimental/DynamicMaterial/Source/DynamicMaterial/Public/Components/DMRenderTargetRenderer.h
--- a/Engine/Plugins/Experimental/DynamicMaterial/Source/DynamicMaterial/Public/Components/DMRenderTargetRenderer.h
+++ b/Engine/Plugins/Experimental/DynamicMaterial/Source/DynamicMaterial/Public/Components/DMRenderTargetRenderer.h
@@ -57,7 +57,7 @@ public:
 	//~ End IDMJsonSerializable
 
 	//~ Begin UDMMaterialComponent
-	DYNAMICMATERIAL_API virtual void Update(EDMUpdateType InUpdateType) override;
+	DYNAMICMATERIAL_API virtual void Update(UDMMaterialComponent* InSource, EDMUpdateType InUpdateType) override;
 	//~ End UDMMaterialComponent
 
 	//~ Begin UObject
diff --git a/Engine/Plugins/Experimental/DynamicMaterial/Source/DynamicMaterial/Public/Components/DMTextureUV.h b/Engine/Plugins/Experimental/DynamicMaterial/Source/DynamicMaterial/Public/Components/DMTextureUV.h
--- a/Engine/Plugins/Experimental/DynamicMaterial/Source/DynamicMaterial/Public/Components/DMTextureUV.h
+++ b/Engine/Plugins/Experimental/DynamicMaterial/Source/DynamicMaterial/Public/Components/DMTextureUV.h
@@ -187,7 +187,7 @@ public:
 	DYNAMICMATERIAL_API virtual FString GetComponentPathComponent() const override;
 	DYNAMICMATERIAL_API virtual void PostEditorDuplicate(UDynamicMaterialModel* InMaterialModel, UDMMaterialComponent* InParent) override;
 #endif
-	DYNAMICMATERIAL_API virtual void Update(EDMUpdateType InUpdateType) override;
+	DYNAMICMATERIAL_API virtual void Update(UDMMaterialComponent* InSource, EDMUpdateType InUpdateType) override;
 	//~ End UDMMaterialComponent
 
 protected:
diff --git a/Engine/Plugins/Experimental/DynamicMaterial/Source/DynamicMaterial/Public/Components/DMTextureUVDynamic.h b/Engine/Plugins/Experimental/DynamicMaterial/Source/DynamicMaterial/Public/Components/DMTextureUVDynamic.h
--- a/Engine/Plugins/Experimental/DynamicMaterial/Source/DynamicMaterial/Public/Components/DMTextureUVDynamic.h
+++ b/Engine/Plugins/Experimental/DynamicMaterial/Source/DynamicMaterial/Public/Components/DMTextureUVDynamic.h
@@ -75,7 +75,7 @@ public:
 #endif
 
 	//~ Begin UDMMaterialComponent
-	DYNAMICMATERIAL_API virtual void Update(EDMUpdateType InUpdateType = EDMUpdateType::Value) override;
+	DYNAMICMATERIAL_API virtual void Update(UDMMaterialComponent* InSource, EDMUpdateType InUpdateType = EDMUpdateType::Value) override;
 	//~ End UDMMaterialComponent
 
 #if WITH_EDITOR
diff --git a/Engine/Plugins/Experimental/DynamicMaterial/Source/DynamicMaterial/Public/Components/MaterialValues/DMMaterialValueRenderTarget.h b/Engine/Plugins/Experimental/DynamicMaterial/Source/DynamicMaterial/Public/Components/MaterialValues/DMMaterialValueRenderTarget.h
--- a/Engine/Plugins/Experimental/DynamicMaterial/Source/DynamicMaterial/Public/Components/MaterialValues/DMMaterialValueRenderTarget.h
+++ b/Engine/Plugins/Experimental/DynamicMaterial/Source/DynamicMaterial/Public/Components/MaterialValues/DMMaterialValueRenderTarget.h
@@ -76,7 +76,7 @@ public:
 #endif
 
 	//~ Begin UDMMaterialComponent
-	DYNAMICMATERIAL_API virtual void Update(EDMUpdateType InUpdateType) override;
+	DYNAMICMATERIAL_API virtual void Update(UDMMaterialComponent* InSource, EDMUpdateType InUpdateType) override;
 #if WITH_EDITOR
 	DYNAMICMATERIAL_API virtual FString GetComponentPathComponent() const override;
 	DYNAMICMATERIAL_API virtual void PostEditorDuplicate(UDynamicMaterialModel* InMaterialModel, UDMMaterialComponent* InParent) override;
diff --git a/Engine/Plugins/Experimental/DynamicMaterial/Source/DynamicMaterialEditor/Private/Components/DMMaterialEffect.cpp b/Engine/Plugins/Experimental/DynamicMaterial/Source/DynamicMaterialEditor/Private/Components/DMMaterialEffect.cpp
--- a/Engine/Plugins/Experimental/DynamicMaterial/Source/DynamicMaterialEditor/Private/Components/DMMaterialEffect.cpp
+++ b/Engine/Plugins/Experimental/DynamicMaterial/Source/DynamicMaterialEditor/Private/Components/DMMaterialEffect.cpp
@@ -71,7 +71,7 @@ bool UDMMaterialEffect::SetEnabled(bool bInIsEnabled)
 
 	bEnabled = bInIsEnabled;
 
-	Update(EDMUpdateType::Structure);
+	Update(this, EDMUpdateType::Structure);
 
 	return true;
 }
@@ -108,7 +108,7 @@ FText UDMMaterialEffect::GetComponentDescription() const
 	return Description;
 }
 
-void UDMMaterialEffect::Update(EDMUpdateType InUpdateType)
+void UDMMaterialEffect::Update(UDMMaterialComponent* InSource, EDMUpdateType InUpdateType)
 {
 	if (!FDMUpdateGuard::CanUpdate())
 	{
@@ -136,14 +136,14 @@ void UDMMaterialEffect::Update(EDMUpdateType InUpdateType)
 
 	if (Index != INDEX_NONE && EffectStack->GetEffects().IsValidIndex(Index + 1))
 	{
-		EffectStack->GetEffects()[Index + 1]->Update(InUpdateType);
+		EffectStack->GetEffects()[Index + 1]->Update(InSource, InUpdateType);
 	}
 	else
 	{
-		EffectStack->Update(InUpdateType);
+		EffectStack->Update(InSource, InUpdateType);
 	}
 
-	Super::Update(InUpdateType);
+	Super::Update(InSource, InUpdateType);
 
 	if (UDMMaterialLayerObject* Layer = EffectStack->GetLayer())
 	{
@@ -151,7 +151,7 @@ void UDMMaterialEffect::Update(EDMUpdateType InUpdateType)
 		{
 			if (UDMMaterialStage* BaseStage = Layer->GetStage(EDMMaterialLayerStage::Base))
 			{
-				BaseStage->Update(InUpdateType);
+				BaseStage->Update(InSource, InUpdateType);
 			}
 		}
 
@@ -159,7 +159,7 @@ void UDMMaterialEffect::Update(EDMUpdateType InUpdateType)
 		{
 			if (UDMMaterialStage* MaskStage = Layer->GetStage(EDMMaterialLayerStage::Mask))
 			{
-				MaskStage->Update(InUpdateType);
+				MaskStage->Update(InSource, InUpdateType);
 			}
 		}
 	}
@@ -194,7 +194,7 @@ void UDMMaterialEffect::PostEditUndo()
 
 	MarkComponentDirty();
 
-	Update(EDMUpdateType::Structure);
+	Update(this, EDMUpdateType::Structure);
 }
 
 #undef LOCTEXT_NAMESPACE
diff --git a/Engine/Plugins/Experimental/DynamicMaterial/Source/DynamicMaterialEditor/Private/Components/DMMaterialEffectFunction.cpp b/Engine/Plugins/Experimental/DynamicMaterial/Source/DynamicMaterialEditor/Private/Components/DMMaterialEffectFunction.cpp
--- a/Engine/Plugins/Experimental/DynamicMaterial/Source/DynamicMaterialEditor/Private/Components/DMMaterialEffectFunction.cpp
+++ b/Engine/Plugins/Experimental/DynamicMaterial/Source/DynamicMaterialEditor/Private/Components/DMMaterialEffectFunction.cpp
@@ -50,7 +50,7 @@ bool UDMMaterialEffectFunction::SetMaterialFunction(UMaterialFunctionInterface*
 
 	OnMaterialFunctionChanged();
 
-	Update(EDMUpdateType::Structure);
+	Update(this, EDMUpdateType::Structure);
 
 	return true;
 }
@@ -181,7 +181,7 @@ bool UDMMaterialEffectFunction::JsonDeserialize(const TSharedPtr<FJsonValue>& In
 
 	if (bSuccess)
 	{
-		Update(EDMUpdateType::Structure);
+		Update(this, EDMUpdateType::Structure);
 	}
 
 	return bSuccess;
diff --git a/Engine/Plugins/Experimental/DynamicMaterial/Source/DynamicMaterialEditor/Private/Components/DMMaterialEffectStack.cpp b/Engine/Plugins/Experimental/DynamicMaterial/Source/DynamicMaterialEditor/Private/Components/DMMaterialEffectStack.cpp
--- a/Engine/Plugins/Experimental/DynamicMaterial/Source/DynamicMaterialEditor/Private/Components/DMMaterialEffectStack.cpp
+++ b/Engine/Plugins/Experimental/DynamicMaterial/Source/DynamicMaterialEditor/Private/Components/DMMaterialEffectStack.cpp
@@ -57,7 +57,7 @@ bool UDMMaterialEffectStack::SetEnabled(bool bInIsEnabled)
 
 	bEnabled = bInIsEnabled;
 
-	Update(EDMUpdateType::Structure);
+	Update(this, EDMUpdateType::Structure);
 
 	return true;
 }
@@ -158,7 +158,7 @@ bool UDMMaterialEffectStack::AddEffect(UDMMaterialEffect* InEffect)
 		InEffect->SetComponentState(EDMComponentLifetimeState::Added);
 	}
 
-	InEffect->Update(EDMUpdateType::Structure);
+	InEffect->Update(this, EDMUpdateType::Structure);
 
 	return true;
 }
@@ -194,7 +194,7 @@ UDMMaterialEffect* UDMMaterialEffectStack::SetEffect(int32 InIndex, UDMMaterialE
 	InEffect->SetComponentState(EDMComponentLifetimeState::Added);
 
 	Effects[InIndex] = InEffect;
-	Effects[InIndex]->Update(EDMUpdateType::Structure);
+	Effects[InIndex]->Update(this, EDMUpdateType::Structure);
 
 	return OldEffect;
 }
@@ -222,7 +222,7 @@ bool UDMMaterialEffectStack::MoveEffect(int32 InIndex, int32 InNewIndex)
 
 	const int MinIndex = FMath::Min(InIndex, InNewIndex);
 
-	Effects[MinIndex]->Update(EDMUpdateType::Structure);
+	Effects[MinIndex]->Update(this, EDMUpdateType::Structure);
 
 	return true;
 }
@@ -269,7 +269,7 @@ UDMMaterialEffect* UDMMaterialEffectStack::RemoveEffect(int32 InIndex)
 
 	Effects.RemoveAt(InIndex);
 
-	Update(EDMUpdateType::Structure);
+	Update(this, EDMUpdateType::Structure);
 
 	return Effect;
 }
@@ -387,7 +387,7 @@ FText UDMMaterialEffectStack::GetComponentDescription() const
 	return Description;
 }
 
-void UDMMaterialEffectStack::Update(EDMUpdateType InUpdateType)
+void UDMMaterialEffectStack::Update(UDMMaterialComponent* InSource, EDMUpdateType InUpdateType)
 {
 	if (!FDMUpdateGuard::CanUpdate())
 	{
@@ -404,11 +404,11 @@ void UDMMaterialEffectStack::Update(EDMUpdateType InUpdateType)
 		return;
 	}
 
-	Super::Update(InUpdateType);
+	Super::Update(InSource, InUpdateType);
 
 	if (UDMMaterialComponent* Parent = GetParentComponent())
 	{
-		Parent->Update(InUpdateType);
+		Parent->Update(InSource, InUpdateType);
 	}
 }
 
@@ -464,7 +464,7 @@ void UDMMaterialEffectStack::PostEditUndo()
 
 	MarkComponentDirty();
 
-	Update(EDMUpdateType::Structure);
+	Update(this, EDMUpdateType::Structure);
 }
 
 TArray<UDMMaterialEffect*> UDMMaterialEffectStack::GetIncompatibleEffects(UDMMaterialEffect* InEffect)
diff --git a/Engine/Plugins/Experimental/DynamicMaterial/Source/DynamicMaterialEditor/Private/Components/DMMaterialLayer.cpp b/Engine/Plugins/Experimental/DynamicMaterial/Source/DynamicMaterialEditor/Private/Components/DMMaterialLayer.cpp
--- a/Engine/Plugins/Experimental/DynamicMaterial/Source/DynamicMaterialEditor/Private/Components/DMMaterialLayer.cpp
+++ b/Engine/Plugins/Experimental/DynamicMaterial/Source/DynamicMaterialEditor/Private/Components/DMMaterialLayer.cpp
@@ -173,7 +173,7 @@ bool UDMMaterialLayerObject::SetMaterialProperty(EDMMaterialPropertyType InMater
 
 	MaterialProperty = InMaterialProperty;
 
-	Update(EDMUpdateType::Structure);
+	Update(this, EDMUpdateType::Structure);
 
 	return true;
 }
@@ -192,7 +192,7 @@ bool UDMMaterialLayerObject::SetTextureUVLinkEnabled(bool bInValue)
 
 	bLinkedUVs = bInValue;
 
-	Update(EDMUpdateType::Structure);
+	Update(this, EDMUpdateType::Structure);
 
 	return true;
 }
@@ -479,7 +479,7 @@ bool UDMMaterialLayerObject::SetEnabled(bool bInIsEnabled)
 
 	bEnabled = bInIsEnabled;
 
-	Update(EDMUpdateType::Structure);
+	Update(this, EDMUpdateType::Structure);
 
 	return true;
 }
@@ -764,7 +764,7 @@ FText UDMMaterialLayerObject::GetComponentDescription() const
 	return LayerName;
 }
 
-void UDMMaterialLayerObject::Update(EDMUpdateType InUpdateType)
+void UDMMaterialLayerObject::Update(UDMMaterialComponent* InSource, EDMUpdateType InUpdateType)
 {
 	if (!FDMUpdateGuard::CanUpdate())
 	{
@@ -785,23 +785,23 @@ void UDMMaterialLayerObject::Update(EDMUpdateType InUpdateType)
 	{
 		if (UDMMaterialStage* FirstStage = NextLayer->GetFirstEnabledStage(EDMMaterialLayerStage::All))
 		{
-			FirstStage->Update(InUpdateType);
+			FirstStage->Update(InSource, InUpdateType);
 		}
 		else if (UDMMaterialEffectStack* NextEffectStack = NextLayer->GetEffectStack())
 		{
-			NextEffectStack->Update(InUpdateType);
+			NextEffectStack->Update(InSource, InUpdateType);
 		}
 		else
 		{
-			NextLayer->Update(InUpdateType);
+			NextLayer->Update(InSource, InUpdateType);
 		}
 	}
 	else if (UDMMaterialSlot* Slot = GetSlot())
 	{
-		Slot->Update(InUpdateType);
+		Slot->Update(InSource, InUpdateType);
 	}
 
-	Super::Update(InUpdateType);
+	Super::Update(InSource, InUpdateType);
 }
 
 void UDMMaterialLayerObject::PostEditorDuplicate(UDynamicMaterialModel* InMaterialModel, UDMMaterialComponent* InParent)
@@ -856,7 +856,7 @@ void UDMMaterialLayerObject::PostEditUndo()
 
 	MarkComponentDirty();
 
-	Update(EDMUpdateType::Structure);
+	Update(this, EDMUpdateType::Structure);
 }
 
 UDMMaterialComponent* UDMMaterialLayerObject::GetSubComponentByPath(FDMComponentPath& InPath, const FDMComponentPathSegment& InPathSegment) const
diff --git a/Engine/Plugins/Experimental/DynamicMaterial/Source/DynamicMaterialEditor/Private/Components/DMMaterialProperty.cpp b/Engine/Plugins/Experimental/DynamicMaterial/Source/DynamicMaterialEditor/Private/Components/DMMaterialProperty.cpp
--- a/Engine/Plugins/Experimental/DynamicMaterial/Source/DynamicMaterialEditor/Private/Components/DMMaterialProperty.cpp
+++ b/Engine/Plugins/Experimental/DynamicMaterial/Source/DynamicMaterialEditor/Private/Components/DMMaterialProperty.cpp
@@ -84,7 +84,7 @@ void UDMMaterialProperty::SetEnabled(bool bInEnabled)
 
 	bEnabled = bInEnabled;
 
-	Update(EDMUpdateType::Structure | EDMUpdateType::AllowParentUpdate);
+	Update(this, EDMUpdateType::Structure | EDMUpdateType::AllowParentUpdate);
 }
 
 FText UDMMaterialProperty::GetDescription() const
@@ -593,7 +593,7 @@ void UDMMaterialProperty::AddOutputProcessor(const TSharedRef<FDMMaterialBuildSt
 	MaterialPropertyPtr->OutputIndex = 0;
 }
 
-void UDMMaterialProperty::Update(EDMUpdateType InUpdateType)
+void UDMMaterialProperty::Update(UDMMaterialComponent* InSource, EDMUpdateType InUpdateType)
 {
 	if (!FDMUpdateGuard::CanUpdate())
 	{
@@ -610,7 +610,7 @@ void UDMMaterialProperty::Update(EDMUpdateType InUpdateType)
 		return;
 	}
 
-	Super::Update(InUpdateType);
+	Super::Update(InSource, InUpdateType);
 
 	if (EnumHasAnyFlags(InUpdateType, EDMUpdateType::Structure))
 	{
@@ -700,7 +700,7 @@ void UDMMaterialProperty::OnOutputProcessorUpdated()
 	{
 		if (OutputProcessor_PreUpdate)
 		{
-			Update(EDMUpdateType::Structure);
+			Update(this, EDMUpdateType::Structure);
 		}
 
 		OutputProcessor = nullptr;
@@ -744,7 +744,7 @@ void UDMMaterialProperty::OnOutputProcessorUpdated()
 		}
 	}
 
-	Update(EDMUpdateType::Structure);
+	Update(this, EDMUpdateType::Structure);
 }
 
 UDMMaterialComponent* UDMMaterialProperty::GetSubComponentByPath(FDMComponentPath& InPath, const FDMComponentPathSegment& InPathSegment) const
diff --git a/Engine/Plugins/Experimental/DynamicMaterial/Source/DynamicMaterialEditor/Private/Components/DMMaterialSlot.cpp b/Engine/Plugins/Experimental/DynamicMaterial/Source/DynamicMaterialEditor/Private/Components/DMMaterialSlot.cpp
--- a/Engine/Plugins/Experimental/DynamicMaterial/Source/DynamicMaterialEditor/Private/Components/DMMaterialSlot.cpp
+++ b/Engine/Plugins/Experimental/DynamicMaterial/Source/DynamicMaterialEditor/Private/Components/DMMaterialSlot.cpp
@@ -88,11 +88,11 @@ bool UDMMaterialSlot::MoveLayer(UDMMaterialLayerObject* InLayer, int32 InNewInde
 
 	if (UDMMaterialStage* Stage = LayerObjects[MinIndex]->GetFirstEnabledStage(EDMMaterialLayerStage::All))
 	{
-		Stage->Update(EDMUpdateType::Structure);
+		Stage->Update(this, EDMUpdateType::Structure);
 	}
 	else
 	{
-		Update(EDMUpdateType::Structure);
+		Update(this, EDMUpdateType::Structure);
 	}
 
 	OnLayersUpdateDelegate.Broadcast(this);
@@ -180,7 +180,7 @@ UDMMaterialLayerObject* UDMMaterialSlot::GetLastLayerForMaterialProperty(EDMMate
 	return nullptr;
 }
 
-void UDMMaterialSlot::Update(EDMUpdateType InUpdateType)
+void UDMMaterialSlot::Update(UDMMaterialComponent* InSource, EDMUpdateType InUpdateType)
 {
 	if (!FDMUpdateGuard::CanUpdate())
 	{
@@ -203,7 +203,7 @@ void UDMMaterialSlot::Update(EDMUpdateType InUpdateType)
 		UpdateMaterialProperties();
 	}
 
-	Super::Update(InUpdateType);
+	Super::Update(InSource, InUpdateType);
 
 	if (EnumHasAnyFlags(InUpdateType, EDMUpdateType::Structure))
 	{
@@ -448,7 +448,7 @@ UDMMaterialLayerObject* UDMMaterialSlot::AddDefaultLayer(EDMMaterialPropertyType
 
 	UpdateOutputConnectorTypes();
 
-	NewLayer->Update(EDMUpdateType::Structure);
+	NewLayer->Update(this, EDMUpdateType::Structure);
 
 	OnLayersUpdateDelegate.Broadcast(this);
 
@@ -492,7 +492,7 @@ UDMMaterialLayerObject* UDMMaterialSlot::AddLayer(EDMMaterialPropertyType InMate
 
 	UpdateOutputConnectorTypes();
 
-	NewLayer->Update(EDMUpdateType::Structure);
+	NewLayer->Update(this, EDMUpdateType::Structure);
 
 	OnLayersUpdateDelegate.Broadcast(this);
 
@@ -531,7 +531,7 @@ UDMMaterialLayerObject* UDMMaterialSlot::AddLayerWithMask(EDMMaterialPropertyTyp
 
 	UpdateOutputConnectorTypes();
 
-	NewLayer->Update(EDMUpdateType::Structure);
+	NewLayer->Update(this, EDMUpdateType::Structure);
 
 	OnLayersUpdateDelegate.Broadcast(this);
 
@@ -617,11 +617,11 @@ bool UDMMaterialSlot::PasteLayer(UDMMaterialLayerObject* InLayer)
 
 	if (UDMMaterialStage* Stage = InLayer->GetFirstEnabledStage(EDMMaterialLayerStage::All))
 	{
-		Stage->Update(EDMUpdateType::Structure);
+		Stage->Update(this, EDMUpdateType::Structure);
 	}
 	else
 	{
-		Update(EDMUpdateType::Structure);
+		Update(this, EDMUpdateType::Structure);
 	}
 
 	OnLayersUpdateDelegate.Broadcast(this);
@@ -698,11 +698,11 @@ bool UDMMaterialSlot::RemoveLayer(UDMMaterialLayerObject* InLayer)
 	{
 		if (UDMMaterialStage* Stage = LayerObjects[0]->GetFirstEnabledStage(EDMMaterialLayerStage::All))
 		{
-			Stage->Update(EDMUpdateType::Structure);
+			Stage->Update(this, EDMUpdateType::Structure);
 		}
 		else
 		{
-			Update(EDMUpdateType::Structure);
+			Update(this, EDMUpdateType::Structure);
 		}
 	}
 
@@ -1300,7 +1300,7 @@ bool UDMMaterialSlot::SetLayerMaterialPropertyAndReplaceOthers(UDMMaterialLayerO
 	{
 		if (UDMMaterialStage* Stage = CurrentSlot->LayerObjects[0]->GetFirstEnabledStage(EDMMaterialLayerStage::All))
 		{
-			Stage->Update(EDMUpdateType::Structure);
+			Stage->Update(this, EDMUpdateType::Structure);
 		}
 	}
 
@@ -1369,7 +1369,7 @@ bool UDMMaterialSlot::ChangeMaterialProperty(EDMMaterialPropertyType InPropertyF
 
 	if (FirstLayerObject)
 	{
-		FirstLayerObject->Update(EDMUpdateType::Structure);
+		FirstLayerObject->Update(this, EDMUpdateType::Structure);
 	}
 
 	return true;
@@ -1478,7 +1478,7 @@ void UDMMaterialSlot::PostEditUndo()
 
 	MarkComponentDirty();
 
-	Update(EDMUpdateType::Structure);
+	Update(this, EDMUpdateType::Structure);
 
 	// Fire all of these to make sure everything is updated.
 	OnPropertiesUpdateDelegate.Broadcast(this);
@@ -1533,11 +1533,11 @@ void UDMMaterialSlot::ConvertDeprecatedLayers(TArray<FDMMaterialLayer>& InLayers
 
 	if (!LayerObjects.IsEmpty())
 	{
-		LayerObjects[0]->Update(EDMUpdateType::Structure);
+		LayerObjects[0]->Update(this, EDMUpdateType::Structure);
 	}
 	else
 	{
-		Update(EDMUpdateType::Structure);
+		Update(this, EDMUpdateType::Structure);
 	}
 }
 
diff --git a/Engine/Plugins/Experimental/DynamicMaterial/Source/DynamicMaterialEditor/Private/Components/DMMaterialStage.cpp b/Engine/Plugins/Experimental/DynamicMaterial/Source/DynamicMaterialEditor/Private/Components/DMMaterialStage.cpp
--- a/Engine/Plugins/Experimental/DynamicMaterial/Source/DynamicMaterialEditor/Private/Components/DMMaterialStage.cpp
+++ b/Engine/Plugins/Experimental/DynamicMaterial/Source/DynamicMaterialEditor/Private/Components/DMMaterialStage.cpp
@@ -241,7 +241,7 @@ bool UDMMaterialStage::SetEnabled(bool bInEnabled)
 
 	bEnabled = bInEnabled;
 
-	Update(EDMUpdateType::Structure);
+	Update(this, EDMUpdateType::Structure);
 
 	return true;
 }
@@ -282,7 +282,7 @@ void UDMMaterialStage::SetSource(UDMMaterialStageSource* InSource)
 		Source->SetComponentState(EDMComponentLifetimeState::Added);
 	}
 
-	Update(EDMUpdateType::Structure);
+	Update(this, EDMUpdateType::Structure);
 }
 
 FText UDMMaterialStage::GetComponentDescription() const
@@ -348,7 +348,7 @@ bool UDMMaterialStage::IsInputMapped(int32 InputIndex) const
 	return false;
 }
 
-void UDMMaterialStage::Update(EDMUpdateType InUpdateType)
+void UDMMaterialStage::Update(UDMMaterialComponent* InSource, EDMUpdateType InUpdateType)
 {
 	if (!FDMUpdateGuard::CanUpdate())
 	{
@@ -371,15 +371,15 @@ void UDMMaterialStage::Update(EDMUpdateType InUpdateType)
 		VerifyAllInputMaps();
 	}
 
-	Super::Update(InUpdateType);
+	Super::Update(InSource, InUpdateType);
 
 	if (UDMMaterialStage* NextStage = GetNextStage())
 	{
-		NextStage->Update(InUpdateType);
+		NextStage->Update(InSource, InUpdateType);
 	}
 	else if (UDMMaterialLayerObject* Layer = GetLayer())
 	{
-		Layer->Update(InUpdateType);
+		Layer->Update(InSource, InUpdateType);
 	}	
 }
 
@@ -833,7 +833,7 @@ void UDMMaterialStage::UpdateInputMap(int32 InInputIdx, int32 InSourceIndex, int
 
 	RemoveUnusedInputs();
 
-	Source->Update(EDMUpdateType::Structure);
+	Source->Update(this, EDMUpdateType::Structure);
 }
 
 int32 UDMMaterialStage::FindIndex() const
@@ -1188,7 +1188,7 @@ void UDMMaterialStage::RemoveInput(UDMMaterialStageInput* InInput)
 
 	InInput->SetComponentState(EDMComponentLifetimeState::Removed);
 
-	Update(EDMUpdateType::Structure);
+	Update(this, EDMUpdateType::Structure);
 }
 
 void UDMMaterialStage::RemoveAllInputs()
@@ -1210,7 +1210,7 @@ void UDMMaterialStage::RemoveAllInputs()
 
 	Inputs.Empty();
 
-	Update(EDMUpdateType::Structure);
+	Update(this, EDMUpdateType::Structure);
 }
 
 const FDMMaterialStageConnectorChannel* UDMMaterialStage::FindInputChannel(UDMMaterialStageInput* InStageInput)
@@ -1259,7 +1259,7 @@ void UDMMaterialStage::PostEditUndo()
 
 	MarkComponentDirty();
 
-	Update(EDMUpdateType::Structure);
+	Update(this, EDMUpdateType::Structure);
 }
 
 void UDMMaterialStage::GeneratePreviewMaterial(UMaterial* InPreviewMaterial)
diff --git a/Engine/Plugins/Experimental/DynamicMaterial/Source/DynamicMaterialEditor/Private/Components/DMMaterialStageBlend.cpp b/Engine/Plugins/Experimental/DynamicMaterial/Source/DynamicMaterialEditor/Private/Components/DMMaterialStageBlend.cpp
--- a/Engine/Plugins/Experimental/DynamicMaterial/Source/DynamicMaterialEditor/Private/Components/DMMaterialStageBlend.cpp
+++ b/Engine/Plugins/Experimental/DynamicMaterial/Source/DynamicMaterialEditor/Private/Components/DMMaterialStageBlend.cpp
@@ -76,7 +76,7 @@ void UDMMaterialStageBlend::SetBaseChannelOverride(EAvaColorChannel InMaskChanne
 	BaseChannelOverride = InMaskChannel;
 	PushBaseChannelOverride();
 
-	Update(EDMUpdateType::Structure);
+	Update(this, EDMUpdateType::Structure);
 }
 
 int32 UDMMaterialStageBlend::GetDefaultBaseChannelOverrideOutputIndex() const
diff --git a/Engine/Plugins/Experimental/DynamicMaterial/Source/DynamicMaterialEditor/Private/Components/DMMaterialStageFunction.cpp b/Engine/Plugins/Experimental/DynamicMaterial/Source/DynamicMaterialEditor/Private/Components/DMMaterialStageFunction.cpp
--- a/Engine/Plugins/Experimental/DynamicMaterial/Source/DynamicMaterialEditor/Private/Components/DMMaterialStageFunction.cpp
+++ b/Engine/Plugins/Experimental/DynamicMaterial/Source/DynamicMaterialEditor/Private/Components/DMMaterialStageFunction.cpp
@@ -288,7 +288,7 @@ void UDMMaterialStageFunction::OnMaterialFunctionChanged()
 	DeinitFunction();
 	InitFunction();
 
-	Update(EDMUpdateType::Structure);
+	Update(this, EDMUpdateType::Structure);
 }
 
 bool UDMMaterialStageFunction::NeedsFunctionInit() const
diff --git a/Engine/Plugins/Experimental/DynamicMaterial/Source/DynamicMaterialEditor/Private/Components/DMMaterialStageGradient.cpp b/Engine/Plugins/Experimental/DynamicMaterial/Source/DynamicMaterialEditor/Private/Components/DMMaterialStageGradient.cpp
--- a/Engine/Plugins/Experimental/DynamicMaterial/Source/DynamicMaterialEditor/Private/Components/DMMaterialStageGradient.cpp
+++ b/Engine/Plugins/Experimental/DynamicMaterial/Source/DynamicMaterialEditor/Private/Components/DMMaterialStageGradient.cpp
@@ -47,7 +47,7 @@ bool UDMMaterialStageGradient::SetMaterialFunction(UMaterialFunctionInterface* I
 
 	MaterialFunction = InMaterialFunction;
 
-	Update(EDMUpdateType::Structure | EDMUpdateType::AllowParentUpdate);
+	Update(this, EDMUpdateType::Structure | EDMUpdateType::AllowParentUpdate);
 
 	return true;
 }
diff --git a/Engine/Plugins/Experimental/DynamicMaterial/Source/DynamicMaterialEditor/Private/Components/DMMaterialStageInput.cpp b/Engine/Plugins/Experimental/DynamicMaterial/Source/DynamicMaterialEditor/Private/Components/DMMaterialStageInput.cpp
--- a/Engine/Plugins/Experimental/DynamicMaterial/Source/DynamicMaterialEditor/Private/Components/DMMaterialStageInput.cpp
+++ b/Engine/Plugins/Experimental/DynamicMaterial/Source/DynamicMaterialEditor/Private/Components/DMMaterialStageInput.cpp
@@ -12,7 +12,7 @@
 
 const FString UDMMaterialStageInput::StageInputPrefixStr = FString(TEXT("DMMaterialStageInput"));
 
-void UDMMaterialStageInput::Update(EDMUpdateType InUpdateType)
+void UDMMaterialStageInput::Update(UDMMaterialComponent* InSource, EDMUpdateType InUpdateType)
 {
 	if (!FDMUpdateGuard::CanUpdate())
 	{
@@ -39,15 +39,15 @@ void UDMMaterialStageInput::Update(EDMUpdateType InUpdateType)
 
 	if (UDMMaterialStageThroughput* Throughput = Cast<UDMMaterialStageThroughput>(Stage->GetSource()))
 	{
-		Throughput->Update(InUpdateType);
+		Throughput->Update(InSource, InUpdateType);
 		Stage->InputUpdated(this, InUpdateType);
 	}
 	else
 	{
-		Stage->Update(InUpdateType);
+		Stage->Update(InSource, InUpdateType);
 	}
 
-	Super::Update(InUpdateType);
+	Super::Update(InSource, InUpdateType);
 }
 
 void UDMMaterialStageInput::GeneratePreviewMaterial(UMaterial* InPreviewMaterial)
diff --git a/Engine/Plugins/Experimental/DynamicMaterial/Source/DynamicMaterialEditor/Private/Components/DMMaterialStageSource.cpp b/Engine/Plugins/Experimental/DynamicMaterial/Source/DynamicMaterialEditor/Private/Components/DMMaterialStageSource.cpp
--- a/Engine/Plugins/Experimental/DynamicMaterial/Source/DynamicMaterialEditor/Private/Components/DMMaterialStageSource.cpp
+++ b/Engine/Plugins/Experimental/DynamicMaterial/Source/DynamicMaterialEditor/Private/Components/DMMaterialStageSource.cpp
@@ -22,7 +22,7 @@ UDMMaterialStage* UDMMaterialStageSource::GetStage() const
 	return Cast<UDMMaterialStage>(GetOuterSafe());
 }
 
-void UDMMaterialStageSource::Update(EDMUpdateType InUpdateType)
+void UDMMaterialStageSource::Update(UDMMaterialComponent* InSource, EDMUpdateType InUpdateType)
 {
 	if (!FDMUpdateGuard::CanUpdate())
 	{
@@ -47,9 +47,9 @@ void UDMMaterialStageSource::Update(EDMUpdateType InUpdateType)
 	UDMMaterialStage* Stage = GetStage();
 	check(Stage);
 
-	Stage->Update(InUpdateType);
+	Stage->Update(InSource, InUpdateType);
 
-	Super::Update(InUpdateType);
+	Super::Update(InSource, InUpdateType);
 }
 
 void UDMMaterialStageSource::OnComponentAdded()
@@ -61,7 +61,7 @@ void UDMMaterialStageSource::OnComponentAdded()
 
 	Super::OnComponentAdded();
 
-	Update(EDMUpdateType::Structure);
+	Update(this, EDMUpdateType::Structure);
 }
 
 void UDMMaterialStageSource::GetMaskAlphaBlendNode(const TSharedRef<FDMMaterialBuildState>& InBuildState, UMaterialExpression*& OutExpression, int32& OutOutputIndex, int32& OutOutputChannel) const
@@ -246,7 +246,7 @@ void UDMMaterialStageSource::PostEditUndo()
 	}
 
 	MarkComponentDirty();
-	Update(EDMUpdateType::Structure);
+	Update(this, EDMUpdateType::Structure);
 }
 
 UDMMaterialComponent* UDMMaterialStageSource::GetParentComponent() const
diff --git a/Engine/Plugins/Experimental/DynamicMaterial/Source/DynamicMaterialEditor/Private/Components/DMMaterialStageThroughputLayerBlend.cpp b/Engine/Plugins/Experimental/DynamicMaterial/Source/DynamicMaterialEditor/Private/Components/DMMaterialStageThroughputLayerBlend.cpp
--- a/Engine/Plugins/Experimental/DynamicMaterial/Source/DynamicMaterialEditor/Private/Components/DMMaterialStageThroughputLayerBlend.cpp
+++ b/Engine/Plugins/Experimental/DynamicMaterial/Source/DynamicMaterialEditor/Private/Components/DMMaterialStageThroughputLayerBlend.cpp
@@ -172,7 +172,7 @@ bool UDMMaterialStageThroughputLayerBlend::IsInputVisible(int32 InputIndex) cons
 	return (InputIndex == InputMaskSource);
 }
  
-void UDMMaterialStageThroughputLayerBlend::Update(EDMUpdateType InUpdateType)
+void UDMMaterialStageThroughputLayerBlend::Update(UDMMaterialComponent* InSource, EDMUpdateType InUpdateType)
 {
 	if (bBlockUpdate)
 	{
@@ -219,7 +219,7 @@ void UDMMaterialStageThroughputLayerBlend::Update(EDMUpdateType InUpdateType)
  
 	UpdateAlphaOnlyMasks(InUpdateType);
 
-	Super::Update(InUpdateType);
+	Super::Update(InSource, InUpdateType);
 
 	PullMaskChannelOverride();
 }
@@ -528,7 +528,7 @@ void UDMMaterialStageThroughputLayerBlend::UpdateLinkedInputStage(EDMUpdateType
 	if (bInputInCommon)
 	{
 		bBlockUpdate = true;
-		BaseStage->Update(InUpdateType);
+		BaseStage->Update(this, InUpdateType);
 		bBlockUpdate = false;
 	}
 }
@@ -790,7 +790,7 @@ void UDMMaterialStageThroughputLayerBlend::SetPremultiplyAlpha(bool bInValue)
  
 	bPremultiplyAlpha = bInValue;
  
-	Update(EDMUpdateType::Structure);
+	Update(this, EDMUpdateType::Structure);
 }
  
 int32 UDMMaterialStageThroughputLayerBlend::ResolveMaskInput(const TSharedRef<FDMMaterialBuildState>& InBuildState, int32 InputIndex, 
@@ -981,7 +981,7 @@ void UDMMaterialStageThroughputLayerBlend::UpdateAlphaOnlyMaskStatus()
 	bIsAlphaOnlyBlend = (Layer->IsStageEnabled(EDMMaterialLayerStage::Base) == false && Layer->IsStageEnabled(EDMMaterialLayerStage::Mask) == true);
 }
  
-void UDMMaterialStageThroughputLayerBlend::OnStageUpdated(UDMMaterialComponent* InComponent, EDMUpdateType InUpdateType)
+void UDMMaterialStageThroughputLayerBlend::OnStageUpdated(UDMMaterialComponent* InComponent, UDMMaterialComponent* InSource, EDMUpdateType InUpdateType)
 {
 	if (bBlockUpdate)
 	{
@@ -1066,7 +1066,7 @@ void UDMMaterialStageThroughputLayerBlend::UpdateAlphaOnlyMasks(EDMUpdateType In
 		// Calling this will inevitably recall this method when the chain updates.
 		// Block this update.
 		bBlockUpdate = true;
-		PreviousLayer->GetStage(EDMMaterialLayerStage::Mask)->Update(InUpdateType);
+		PreviousLayer->GetStage(EDMMaterialLayerStage::Mask)->Update(this, InUpdateType);
 		bBlockUpdate = false;
 	}
 }
@@ -1119,7 +1119,7 @@ void UDMMaterialStageThroughputLayerBlend::SetMaskChannelOverride(EAvaColorChann
 	MaskChannelOverride = InMaskChannel;
 	PushMaskChannelOverride();
 
-	Update(EDMUpdateType::Structure);
+	Update(this, EDMUpdateType::Structure);
 }
 
 bool UDMMaterialStageThroughputLayerBlend::CanUseMaskChannelOverride() const
diff --git a/Engine/Plugins/Experimental/DynamicMaterial/Source/DynamicMaterialEditor/Private/Components/MaterialStageExpressions/DMMSEMathBase.cpp b/Engine/Plugins/Experimental/DynamicMaterial/Source/DynamicMaterialEditor/Private/Components/MaterialStageExpressions/DMMSEMathBase.cpp
--- a/Engine/Plugins/Experimental/DynamicMaterial/Source/DynamicMaterialEditor/Private/Components/MaterialStageExpressions/DMMSEMathBase.cpp
+++ b/Engine/Plugins/Experimental/DynamicMaterial/Source/DynamicMaterialEditor/Private/Components/MaterialStageExpressions/DMMSEMathBase.cpp
@@ -207,7 +207,7 @@ bool UDMMaterialStageExpressionMathBase::CanInputAcceptType(int32 InInputIndex,
 	return true;
 }
 
-void UDMMaterialStageExpressionMathBase::Update(EDMUpdateType InUpdateType)
+void UDMMaterialStageExpressionMathBase::Update(UDMMaterialComponent* InSource, EDMUpdateType InUpdateType)
 {
 	if (!FDMUpdateGuard::CanUpdate())
 	{
@@ -284,7 +284,7 @@ void UDMMaterialStageExpressionMathBase::Update(EDMUpdateType InUpdateType)
 		}
 	}
 
-	Super::Update(InUpdateType);
+	Super::Update(InSource, InUpdateType);
 }
 
 void UDMMaterialStageExpressionMathBase::AddDefaultInput(int32 InInputIndex) const
diff --git a/Engine/Plugins/Experimental/DynamicMaterial/Source/DynamicMaterialEditor/Private/Components/MaterialStageExpressions/DMMSETextureSampleBase.cpp b/Engine/Plugins/Experimental/DynamicMaterial/Source/DynamicMaterialEditor/Private/Components/MaterialStageExpressions/DMMSETextureSampleBase.cpp
--- a/Engine/Plugins/Experimental/DynamicMaterial/Source/DynamicMaterialEditor/Private/Components/MaterialStageExpressions/DMMSETextureSampleBase.cpp
+++ b/Engine/Plugins/Experimental/DynamicMaterial/Source/DynamicMaterialEditor/Private/Components/MaterialStageExpressions/DMMSETextureSampleBase.cpp
@@ -247,7 +247,7 @@ void UDMMaterialStageExpressionTextureSampleBase::SetClampTextureEnabled(bool bI
 
 	bClampTexture = bInValue;
 
-	Update(EDMUpdateType::Value);
+	Update(this, EDMUpdateType::Value);
 }
 
 void UDMMaterialStageExpressionTextureSampleBase::PostEditChangeProperty(FPropertyChangedEvent& InPropertyChangedEvent)
@@ -260,7 +260,7 @@ void UDMMaterialStageExpressionTextureSampleBase::PostEditChangeProperty(FProper
 
 	if (PropertyName == ClampTextureName)
 	{
-		Update(EDMUpdateType::Structure);
+		Update(this, EDMUpdateType::Structure);
 	}
 }
 
diff --git a/Engine/Plugins/Experimental/DynamicMaterial/Source/DynamicMaterialEditor/Private/Components/MaterialStageExpressions/DMMSEWorldPositionNoise.cpp b/Engine/Plugins/Experimental/DynamicMaterial/Source/DynamicMaterialEditor/Private/Components/MaterialStageExpressions/DMMSEWorldPositionNoise.cpp
--- a/Engine/Plugins/Experimental/DynamicMaterial/Source/DynamicMaterialEditor/Private/Components/MaterialStageExpressions/DMMSEWorldPositionNoise.cpp
+++ b/Engine/Plugins/Experimental/DynamicMaterial/Source/DynamicMaterialEditor/Private/Components/MaterialStageExpressions/DMMSEWorldPositionNoise.cpp
@@ -111,7 +111,7 @@ void UDMMaterialStageExpressionWorldPositionNoise::SetLocationType(EDMLocationTy
 
 	LocationType = InLocationType;
 
-	Update(EDMUpdateType::Structure);
+	Update(this, EDMUpdateType::Structure);
 }
 
 void UDMMaterialStageExpressionWorldPositionNoise::SetShaderOffset(EWorldPositionIncludedOffsets InShaderOffset)
@@ -123,7 +123,7 @@ void UDMMaterialStageExpressionWorldPositionNoise::SetShaderOffset(EWorldPositio
 
 	ShaderOffset = InShaderOffset;
 
-	Update(EDMUpdateType::Structure);
+	Update(this, EDMUpdateType::Structure);
 }
 
 void UDMMaterialStageExpressionWorldPositionNoise::SetNoiseFunction(EVectorNoiseFunction InNoiseFunction)
@@ -135,7 +135,7 @@ void UDMMaterialStageExpressionWorldPositionNoise::SetNoiseFunction(EVectorNoise
 
 	NoiseFunction = InNoiseFunction;
 
-	Update(EDMUpdateType::Structure);
+	Update(this, EDMUpdateType::Structure);
 }
 
 void UDMMaterialStageExpressionWorldPositionNoise::SetQuality(int32 InQuality)
@@ -147,7 +147,7 @@ void UDMMaterialStageExpressionWorldPositionNoise::SetQuality(int32 InQuality)
 
 	Quality = InQuality;
 
-	Update(EDMUpdateType::Structure);
+	Update(this, EDMUpdateType::Structure);
 }
 
 void UDMMaterialStageExpressionWorldPositionNoise::SetTiling(bool bInTiling)
@@ -159,7 +159,7 @@ void UDMMaterialStageExpressionWorldPositionNoise::SetTiling(bool bInTiling)
 
 	bTiling = bInTiling;
 
-	Update(EDMUpdateType::Structure);
+	Update(this, EDMUpdateType::Structure);
 }
 
 void UDMMaterialStageExpressionWorldPositionNoise::SetTileSize(int32 InTileSize)
@@ -171,7 +171,7 @@ void UDMMaterialStageExpressionWorldPositionNoise::SetTileSize(int32 InTileSize)
 
 	TileSize = InTileSize;
 
-	Update(EDMUpdateType::Structure);
+	Update(this, EDMUpdateType::Structure);
 }
 
 void UDMMaterialStageExpressionWorldPositionNoise::AddDefaultInput(int32 InInputIndex) const
@@ -262,7 +262,7 @@ void UDMMaterialStageExpressionWorldPositionNoise::PostEditChangeProperty(FPrope
 		|| PropertyName == TilingName
 		|| PropertyName == TileSizeName)
 	{
-		Update(EDMUpdateType::Structure);
+		Update(this, EDMUpdateType::Structure);
 	}
 }
 
diff --git a/Engine/Plugins/Experimental/DynamicMaterial/Source/DynamicMaterialEditor/Private/Components/MaterialStageGradients/DMMSGLinear.cpp b/Engine/Plugins/Experimental/DynamicMaterial/Source/DynamicMaterialEditor/Private/Components/MaterialStageGradients/DMMSGLinear.cpp
--- a/Engine/Plugins/Experimental/DynamicMaterial/Source/DynamicMaterialEditor/Private/Components/MaterialStageGradients/DMMSGLinear.cpp
+++ b/Engine/Plugins/Experimental/DynamicMaterial/Source/DynamicMaterialEditor/Private/Components/MaterialStageGradients/DMMSGLinear.cpp
@@ -85,7 +85,7 @@ void UDMMaterialStageGradientLinear::OnTilingChanged()
 
 	if (!bWasUpdateCalled)
 	{
-		Update(EDMUpdateType::Structure);
+		Update(this, EDMUpdateType::Structure);
 	}
 }
 
diff --git a/Engine/Plugins/Experimental/DynamicMaterial/Source/DynamicMaterialEditor/Private/Components/MaterialStageInputs/DMMSISlot.cpp b/Engine/Plugins/Experimental/DynamicMaterial/Source/DynamicMaterialEditor/Private/Components/MaterialStageInputs/DMMSISlot.cpp
--- a/Engine/Plugins/Experimental/DynamicMaterial/Source/DynamicMaterialEditor/Private/Components/MaterialStageInputs/DMMSISlot.cpp
+++ b/Engine/Plugins/Experimental/DynamicMaterial/Source/DynamicMaterialEditor/Private/Components/MaterialStageInputs/DMMSISlot.cpp
@@ -184,7 +184,7 @@ void UDMMaterialStageInputSlot::SetSlot(UDMMaterialSlot* InSlot)
 
 	UpdateOutputConnectors();
 
-	Update(EDMUpdateType::Structure);
+	Update(this, EDMUpdateType::Structure);
 }
 
 void UDMMaterialStageInputSlot::SetMaterialProperty(EDMMaterialPropertyType InMaterialProperty)
@@ -203,7 +203,7 @@ void UDMMaterialStageInputSlot::SetMaterialProperty(EDMMaterialPropertyType InMa
 
 	UpdateOutputConnectors();
 
-	Update(EDMUpdateType::Structure);
+	Update(this, EDMUpdateType::Structure);
 }
 
 void UDMMaterialStageInputSlot::OnComponentRemoved()
@@ -333,19 +333,19 @@ UDMMaterialStageInputSlot::UDMMaterialStageInputSlot()
 {
 }
 
-void UDMMaterialStageInputSlot::OnSlotUpdated(UDMMaterialComponent* InComponent, EDMUpdateType InUpdateType)
+void UDMMaterialStageInputSlot::OnSlotUpdated(UDMMaterialComponent* InComponent, UDMMaterialComponent* InSource, EDMUpdateType InUpdateType)
 {
 	if (!IsComponentValid())
 	{
 		return;
 	}
 
-	if (Slot != InComponent)
+	if (Slot != InSource)
 	{
 		return;
 	}
 
-	Update(InUpdateType);
+	Update(InSource, InUpdateType);
 }
 
 void UDMMaterialStageInputSlot::OnSlotConnectorsUpdated(UDMMaterialSlot* InSlot)
@@ -370,7 +370,7 @@ void UDMMaterialStageInputSlot::OnSlotConnectorsUpdated(UDMMaterialSlot* InSlot)
 	{
 		if (UDMMaterialStage* MaskStage = Layer->GetStage(EDMMaterialLayerStage::Mask))
 		{
-			MaskStage->Update(EDMUpdateType::Structure);
+			MaskStage->Update(InSlot, EDMUpdateType::Structure);
 		}
 	}
 	else
diff --git a/Engine/Plugins/Experimental/DynamicMaterial/Source/DynamicMaterialEditor/Private/Components/MaterialStageInputs/DMMSITextureUV.cpp b/Engine/Plugins/Experimental/DynamicMaterial/Source/DynamicMaterialEditor/Private/Components/MaterialStageInputs/DMMSITextureUV.cpp
--- a/Engine/Plugins/Experimental/DynamicMaterial/Source/DynamicMaterialEditor/Private/Components/MaterialStageInputs/DMMSITextureUV.cpp
+++ b/Engine/Plugins/Experimental/DynamicMaterial/Source/DynamicMaterialEditor/Private/Components/MaterialStageInputs/DMMSITextureUV.cpp
@@ -291,16 +291,16 @@ UDMMaterialComponent* UDMMaterialStageInputTextureUV::GetSubComponentByPath(FDMC
 	return Super::GetSubComponentByPath(InPath, InPathSegment);
 }
 
-void UDMMaterialStageInputTextureUV::OnTextureUVUpdated(UDMMaterialComponent* InComponent, EDMUpdateType InUpdateType)
+void UDMMaterialStageInputTextureUV::OnTextureUVUpdated(UDMMaterialComponent* InComponent, UDMMaterialComponent* InSource, EDMUpdateType InUpdateType)
 {
 	if (!IsComponentValid())
 	{
 		return;
 	}
 
-	if (InComponent == TextureUV && EnumHasAnyFlags(InUpdateType, EDMUpdateType::Structure))
+	if (InSource == TextureUV && EnumHasAnyFlags(InUpdateType, EDMUpdateType::Structure))
 	{
-		Update(InUpdateType);
+		Update(InSource, InUpdateType);
 	}
 }
 
diff --git a/Engine/Plugins/Experimental/DynamicMaterial/Source/DynamicMaterialEditor/Private/Components/MaterialStageInputs/DMMSIThroughput.cpp b/Engine/Plugins/Experimental/DynamicMaterial/Source/DynamicMaterialEditor/Private/Components/MaterialStageInputs/DMMSIThroughput.cpp
--- a/Engine/Plugins/Experimental/DynamicMaterial/Source/DynamicMaterialEditor/Private/Components/MaterialStageInputs/DMMSIThroughput.cpp
+++ b/Engine/Plugins/Experimental/DynamicMaterial/Source/DynamicMaterialEditor/Private/Components/MaterialStageInputs/DMMSIThroughput.cpp
@@ -59,7 +59,7 @@ void UDMMaterialStageInputThroughput::SetMaterialStageThroughputClass(TSubclassO
 		}
 	}
 
-	Update(EDMUpdateType::Structure);
+	Update(this, EDMUpdateType::Structure);
 }
 
 UDMMaterialStageThroughput* UDMMaterialStageInputThroughput::GetMaterialStageThroughput() const
@@ -284,16 +284,16 @@ UDMMaterialStageInputThroughput::UDMMaterialStageInputThroughput()
 	EditableProperties.Add(GET_MEMBER_NAME_CHECKED(UDMMaterialStageInputThroughput, SubStage));
 }
 
-void UDMMaterialStageInputThroughput::OnSubStageUpdated(UDMMaterialComponent* InComponent, EDMUpdateType InUpdateType)
+void UDMMaterialStageInputThroughput::OnSubStageUpdated(UDMMaterialComponent* InComponent, UDMMaterialComponent* InSource, EDMUpdateType InUpdateType)
 {
 	if (!IsComponentValid())
 	{
 		return;
 	}
 
-	if (InComponent == SubStage)
+	if (InSource == SubStage)
 	{
-		Update(InUpdateType);
+		Update(InSource, InUpdateType);
 	}
 }
 
diff --git a/Engine/Plugins/Experimental/DynamicMaterial/Source/DynamicMaterialEditor/Private/Components/MaterialStageInputs/DMMSIValue.cpp b/Engine/Plugins/Experimental/DynamicMaterial/Source/DynamicMaterialEditor/Private/Components/MaterialStageInputs/DMMSIValue.cpp
--- a/Engine/Plugins/Experimental/DynamicMaterial/Source/DynamicMaterialEditor/Private/Components/MaterialStageInputs/DMMSIValue.cpp
+++ b/Engine/Plugins/Experimental/DynamicMaterial/Source/DynamicMaterialEditor/Private/Components/MaterialStageInputs/DMMSIValue.cpp
@@ -449,7 +449,7 @@ void UDMMaterialStageInputValue::SetValue(UDMMaterialValue* InValue)
 
 	UpdateOutputConnectors();
 
-	Update(EDMUpdateType::Structure);
+	Update(this, EDMUpdateType::Structure);
 }
 
 void UDMMaterialStageInputValue::GenerateExpressions(const TSharedRef<FDMMaterialBuildState>& InBuildState) const
@@ -759,7 +759,7 @@ void UDMMaterialStageInputValue::OnValueUpdated(UDMMaterialComponent* InComponen
 		return;
 	}
 
-	Update(EDMUpdateType::Value);
+	Update(InComponent, EDMUpdateType::Value);
 }
 
 void UDMMaterialStageInputValue::UpdateOutputConnectors()
diff --git a/Engine/Plugins/Experimental/DynamicMaterial/Source/DynamicMaterialEditor/Private/Model/DynamicMaterialModelEditorOnlyData.cpp b/Engine/Plugins/Experimental/DynamicMaterial/Source/DynamicMaterialEditor/Private/Model/DynamicMaterialModelEditorOnlyData.cpp
--- a/Engine/Plugins/Experimental/DynamicMaterial/Source/DynamicMaterialEditor/Private/Model/DynamicMaterialModelEditorOnlyData.cpp
+++ b/Engine/Plugins/Experimental/DynamicMaterial/Source/DynamicMaterialEditor/Private/Model/DynamicMaterialModelEditorOnlyData.cpp
@@ -824,7 +824,7 @@ bool UDynamicMaterialModelEditorOnlyData::AddTextureSet(UDMTextureSet* InTexture
 			}
 		}
 
-		Layer->Update(EDMUpdateType::Structure);
+		Layer->Update(Layer, EDMUpdateType::Structure);
 	}
 
 	return bMadeChange;
@@ -1339,7 +1339,7 @@ UDMMaterialSlot* UDynamicMaterialModelEditorOnlyData::AddSlotForMaterialProperty
 		Property->OnSlotAdded(NewSlot);
 	}
 
-	NewSlot->Update(EDMUpdateType::Structure);
+	NewSlot->Update(NewSlot, EDMUpdateType::Structure);
 
 	OnSlotListUpdateDelegate.Broadcast(MaterialModel);
 
diff --git a/Engine/Plugins/Experimental/DynamicMaterial/Source/DynamicMaterialEditor/Private/UI/Widgets/Editor/SDMMaterialComponentEditor.cpp b/Engine/Plugins/Experimental/DynamicMaterial/Source/DynamicMaterialEditor/Private/UI/Widgets/Editor/SDMMaterialComponentEditor.cpp
--- a/Engine/Plugins/Experimental/DynamicMaterial/Source/DynamicMaterialEditor/Private/UI/Widgets/Editor/SDMMaterialComponentEditor.cpp
+++ b/Engine/Plugins/Experimental/DynamicMaterial/Source/DynamicMaterialEditor/Private/UI/Widgets/Editor/SDMMaterialComponentEditor.cpp
@@ -50,7 +50,7 @@ UDMMaterialComponent* SDMMaterialComponentEditor::GetComponent() const
 	return Cast<UDMMaterialComponent>(ObjectWeak.Get());
 }
 
-void SDMMaterialComponentEditor::OnComponentUpdated(UDMMaterialComponent* InComponent, EDMUpdateType InUpdateType)
+void SDMMaterialComponentEditor::OnComponentUpdated(UDMMaterialComponent* InComponent, UDMMaterialComponent* InSource, EDMUpdateType InUpdateType)
 {
 	if (EnumHasAnyFlags(InUpdateType, EDMUpdateType::Structure))
 	{
diff --git a/Engine/Plugins/Experimental/DynamicMaterial/Source/DynamicMaterialEditor/Private/UI/Widgets/Editor/SDMMaterialComponentEditor.h b/Engine/Plugins/Experimental/DynamicMaterial/Source/DynamicMaterialEditor/Private/UI/Widgets/Editor/SDMMaterialComponentEditor.h
--- a/Engine/Plugins/Experimental/DynamicMaterial/Source/DynamicMaterialEditor/Private/UI/Widgets/Editor/SDMMaterialComponentEditor.h
+++ b/Engine/Plugins/Experimental/DynamicMaterial/Source/DynamicMaterialEditor/Private/UI/Widgets/Editor/SDMMaterialComponentEditor.h
@@ -23,7 +23,7 @@ public:
 	UDMMaterialComponent* GetComponent() const;
 
 protected:
-	void OnComponentUpdated(UDMMaterialComponent* InComponent, EDMUpdateType InUpdateType);
+	void OnComponentUpdated(UDMMaterialComponent* InComponent, UDMMaterialComponent* InSource, EDMUpdateType InUpdateType);
 
 	//~ Begin SDMObjectEditorWidgetBase
 	virtual TArray<FDMPropertyHandle> GetPropertyRows() override;
diff --git a/Engine/Plugins/Experimental/DynamicMaterial/Source/DynamicMaterialEditor/Private/UI/Widgets/Editor/SlotEditor/SDMMaterialSlotLayerEffectView.cpp b/Engine/Plugins/Experimental/DynamicMaterial/Source/DynamicMaterialEditor/Private/UI/Widgets/Editor/SlotEditor/SDMMaterialSlotLayerEffectView.cpp
--- a/Engine/Plugins/Experimental/DynamicMaterial/Source/DynamicMaterialEditor/Private/UI/Widgets/Editor/SlotEditor/SDMMaterialSlotLayerEffectView.cpp
+++ b/Engine/Plugins/Experimental/DynamicMaterial/Source/DynamicMaterialEditor/Private/UI/Widgets/Editor/SlotEditor/SDMMaterialSlotLayerEffectView.cpp
@@ -47,7 +47,7 @@ SDMMaterialSlotLayerEffectView::~SDMMaterialSlotLayerEffectView()
 		{
 			if (UDMMaterialEffectStack* EffectStack = Layer->GetEffectStack())
 			{
-				EffectStack->GetOnUpdate().AddSP(this, &SDMMaterialSlotLayerEffectView::OnEffectStackUpdate);
+				EffectStack->GetOnUpdate().RemoveAll(this);
 			}
 		}
 	}
@@ -232,7 +232,7 @@ void SDMMaterialSlotLayerEffectView::OnUndo()
 	RequestListRefresh();
 }
 
-void SDMMaterialSlotLayerEffectView::OnEffectStackUpdate(UDMMaterialComponent* InComponent, EDMUpdateType InUpdateType)
+void SDMMaterialSlotLayerEffectView::OnEffectStackUpdate(UDMMaterialComponent* InComponent, UDMMaterialComponent* InSource, EDMUpdateType InUpdateType)
 {
 	if (EnumHasAnyFlags(InUpdateType, EDMUpdateType::Structure))
 	{
diff --git a/Engine/Plugins/Experimental/DynamicMaterial/Source/DynamicMaterialEditor/Private/UI/Widgets/Editor/SlotEditor/SDMMaterialSlotLayerEffectView.h b/Engine/Plugins/Experimental/DynamicMaterial/Source/DynamicMaterialEditor/Private/UI/Widgets/Editor/SlotEditor/SDMMaterialSlotLayerEffectView.h
--- a/Engine/Plugins/Experimental/DynamicMaterial/Source/DynamicMaterialEditor/Private/UI/Widgets/Editor/SlotEditor/SDMMaterialSlotLayerEffectView.h
+++ b/Engine/Plugins/Experimental/DynamicMaterial/Source/DynamicMaterialEditor/Private/UI/Widgets/Editor/SlotEditor/SDMMaterialSlotLayerEffectView.h
@@ -53,7 +53,7 @@ protected:
 	/** Events */
 	void OnUndo();
 
-	void OnEffectStackUpdate(UDMMaterialComponent* InComponent, EDMUpdateType InUpdateType);
+	void OnEffectStackUpdate(UDMMaterialComponent* InComponent, UDMMaterialComponent* InSource, EDMUpdateType InUpdateType);
 
 	TSharedPtr<SWidget> CreateEffectItemContextMenu();
 
diff --git a/Engine/Plugins/Experimental/DynamicMaterial/Source/DynamicMaterialEditor/Private/UI/Widgets/Visualizers/SDMMaterialComponentPreview.cpp b/Engine/Plugins/Experimental/DynamicMaterial/Source/DynamicMaterialEditor/Private/UI/Widgets/Visualizers/SDMMaterialComponentPreview.cpp
--- a/Engine/Plugins/Experimental/DynamicMaterial/Source/DynamicMaterialEditor/Private/UI/Widgets/Visualizers/SDMMaterialComponentPreview.cpp
+++ b/Engine/Plugins/Experimental/DynamicMaterial/Source/DynamicMaterialEditor/Private/UI/Widgets/Visualizers/SDMMaterialComponentPreview.cpp
@@ -84,7 +84,7 @@ void SDMMaterialComponentPreview::Construct(const FArguments& InArgs, const TSha
 	if (ensure(IsValid(InComponent)))
 	{
 		InComponent->GetOnUpdate().AddSP(this, &SDMMaterialComponentPreview::OnComponentUpdated);
-		OnComponentUpdated(InComponent, EDMUpdateType::Structure);
+		OnComponentUpdated(InComponent, InComponent, EDMUpdateType::Structure);
 	}
 
 	ChildSlot
@@ -105,44 +105,56 @@ void SDMMaterialComponentPreview::Tick(const FGeometry& AllottedGeometry, const
 	}	
 }
 
-void SDMMaterialComponentPreview::OnComponentUpdated(UDMMaterialComponent* InComponent, EDMUpdateType InUpdateType)
+void SDMMaterialComponentPreview::OnComponentUpdated(UDMMaterialComponent* InComponent, UDMMaterialComponent* InSource, EDMUpdateType InUpdateType)
 {
-	if (UDMMaterialStage* Stage = Cast<UDMMaterialStage>(InComponent))
+	UDMMaterialStage* Stage = Cast<UDMMaterialStage>(InComponent);
+
+	if (!Stage)
 	{
-		if (Stage == ComponentWeak.Get() && IsValid(Stage) && Stage->IsComponentValid())
-		{
-			if (TSharedPtr<SDMMaterialEditor> EditorWidget = EditorWidgetWeak.Pin())
-			{
-				UMaterial* PreviewMaterialBase = PreviewMaterialBaseWeak.Get();
-
-				if (!PreviewMaterialBase)
-				{
-					PreviewMaterialBase = EditorWidget->GetPreviewMaterialManager()->CreatePreviewMaterial(Stage);
-				}
-
-				if (EnumHasAnyFlags(InUpdateType, EDMUpdateType::Structure))
-				{
-					Stage->GeneratePreviewMaterial(PreviewMaterialBase);
-
-					EditorWidget->GetPreviewMaterialManager()->FreePreviewMaterialDynamic(PreviewMaterialBase);
-					PreviewMaterialDynamicWeak = EditorWidget->GetPreviewMaterialManager()->CreatePreviewMaterialDynamic(PreviewMaterialBase);
-
-					UDynamicMaterialModelBase* MaterialModelBase = EditorWidget->GetMaterialModelBase();
-
-					if (UDynamicMaterialModel* MaterialModel = Cast<UDynamicMaterialModel>(MaterialModelBase))
-					{
-						MaterialModel->ApplyComponents(PreviewMaterialDynamicWeak.Get());
-					}
-					else if (UDynamicMaterialModelDynamic* MaterialModelDynamic = Cast<UDynamicMaterialModelDynamic>(MaterialModelBase))
-					{
-						MaterialModelDynamic->ApplyComponents(PreviewMaterialDynamicWeak.Get());
-					}
-
-					Brush.SetMaterial(PreviewMaterialDynamicWeak.Get());
-				}
-			}
-		}
+		return;
+	}
+
+	if (Stage != ComponentWeak.Get() || !IsValid(Stage) || !Stage->IsComponentValid())
+	{
+		return;
+	}
+
+	TSharedPtr<SDMMaterialEditor> EditorWidget = EditorWidgetWeak.Pin();
+
+	if (!EditorWidget.IsValid())
+	{
+		return;
+	}
+
+	UMaterial* PreviewMaterialBase = PreviewMaterialBaseWeak.Get();
+
+	if (!PreviewMaterialBase)
+	{
+		PreviewMaterialBase = EditorWidget->GetPreviewMaterialManager()->CreatePreviewMaterial(Stage);
+	}
+
+	if (!EnumHasAnyFlags(InUpdateType, EDMUpdateType::Structure))
+	{
+		return;
+	}
+
+	Stage->GeneratePreviewMaterial(PreviewMaterialBase);
+
+	EditorWidget->GetPreviewMaterialManager()->FreePreviewMaterialDynamic(PreviewMaterialBase);
+	PreviewMaterialDynamicWeak = EditorWidget->GetPreviewMaterialManager()->CreatePreviewMaterialDynamic(PreviewMaterialBase);
+
+	UDynamicMaterialModelBase* MaterialModelBase = EditorWidget->GetMaterialModelBase();
+
+	if (UDynamicMaterialModel* MaterialModel = Cast<UDynamicMaterialModel>(MaterialModelBase))
+	{
+		MaterialModel->ApplyComponents(PreviewMaterialDynamicWeak.Get());
 	}
+	else if (UDynamicMaterialModelDynamic* MaterialModelDynamic = Cast<UDynamicMaterialModelDynamic>(MaterialModelBase))
+	{
+		MaterialModelDynamic->ApplyComponents(PreviewMaterialDynamicWeak.Get());
+	}
+
+	Brush.SetMaterial(PreviewMaterialDynamicWeak.Get());
 }
 
 void SDMMaterialComponentPreview::OnValueUpdated(UDynamicMaterialModel* InMaterialModel, UDMMaterialValue* InValue)
diff --git a/Engine/Plugins/Experimental/DynamicMaterial/Source/DynamicMaterialEditor/Private/UI/Widgets/Visualizers/SDMMaterialComponentPreview.h b/Engine/Plugins/Experimental/DynamicMaterial/Source/DynamicMaterialEditor/Private/UI/Widgets/Visualizers/SDMMaterialComponentPreview.h
--- a/Engine/Plugins/Experimental/DynamicMaterial/Source/DynamicMaterialEditor/Private/UI/Widgets/Visualizers/SDMMaterialComponentPreview.h
+++ b/Engine/Plugins/Experimental/DynamicMaterial/Source/DynamicMaterialEditor/Private/UI/Widgets/Visualizers/SDMMaterialComponentPreview.h
@@ -46,7 +46,7 @@ protected:
 	TWeakObjectPtr<UMaterialInstanceDynamic> PreviewMaterialDynamicWeak;
 	FSlateMaterialBrush Brush;
 
-	void OnComponentUpdated(UDMMaterialComponent* InComponent, EDMUpdateType InUpdateType);
+	void OnComponentUpdated(UDMMaterialComponent* InComponent, UDMMaterialComponent* InSource, EDMUpdateType InUpdateType);
 
 	void OnValueUpdated(UDynamicMaterialModel* InMaterialModel, UDMMaterialValue* InValue);
 
diff --git a/Engine/Plugins/Experimental/DynamicMaterial/Source/DynamicMaterialEditor/Private/Utils/DMMaterialSlotFunctionLibrary.cpp b/Engine/Plugins/Experimental/DynamicMaterial/Source/DynamicMaterialEditor/Private/Utils/DMMaterialSlotFunctionLibrary.cpp
--- a/Engine/Plugins/Experimental/DynamicMaterial/Source/DynamicMaterialEditor/Private/Utils/DMMaterialSlotFunctionLibrary.cpp
+++ b/Engine/Plugins/Experimental/DynamicMaterial/Source/DynamicMaterialEditor/Private/Utils/DMMaterialSlotFunctionLibrary.cpp
@@ -111,11 +111,11 @@ UDMMaterialLayerObject* UDMMaterialSlotFunctionLibrary::AddNewLayer(UDMMaterialS
 
 	if (IsValid(Source))
 	{
-		Source->Update(EDMUpdateType::Structure);
+		Source->Update(Layer, EDMUpdateType::Structure);
 	}
 	else if (bValidLayer)
 	{
-		Layer->Update(EDMUpdateType::Structure);
+		Layer->Update(Layer, EDMUpdateType::Structure);
 	}
 
 	return Layer;
diff --git a/Engine/Plugins/Experimental/DynamicMaterial/Source/DynamicMaterialEditor/Public/Components/DMMaterialEffect.h b/Engine/Plugins/Experimental/DynamicMaterial/Source/DynamicMaterialEditor/Public/Components/DMMaterialEffect.h
--- a/Engine/Plugins/Experimental/DynamicMaterial/Source/DynamicMaterialEditor/Public/Components/DMMaterialEffect.h
+++ b/Engine/Plugins/Experimental/DynamicMaterial/Source/DynamicMaterialEditor/Public/Components/DMMaterialEffect.h
@@ -79,7 +79,7 @@ public:
 	DYNAMICMATERIALEDITOR_API virtual UDMMaterialComponent* GetParentComponent() const override;
 	DYNAMICMATERIALEDITOR_API virtual FString GetComponentPathComponent() const override;
 	DYNAMICMATERIALEDITOR_API virtual FText GetComponentDescription() const override;
-	DYNAMICMATERIALEDITOR_API virtual void Update(EDMUpdateType InUpdateType) override;
+	DYNAMICMATERIALEDITOR_API virtual void Update(UDMMaterialComponent* InSource, EDMUpdateType InUpdateType) override;
 	DYNAMICMATERIALEDITOR_API virtual void PostEditorDuplicate(UDynamicMaterialModel* InMaterialModel, UDMMaterialComponent* InParent) override;
 	//~ End UDMMaterialComponent
 
diff --git a/Engine/Plugins/Experimental/DynamicMaterial/Source/DynamicMaterialEditor/Public/Components/DMMaterialEffectStack.h b/Engine/Plugins/Experimental/DynamicMaterial/Source/DynamicMaterialEditor/Public/Components/DMMaterialEffectStack.h
--- a/Engine/Plugins/Experimental/DynamicMaterial/Source/DynamicMaterialEditor/Public/Components/DMMaterialEffectStack.h
+++ b/Engine/Plugins/Experimental/DynamicMaterial/Source/DynamicMaterialEditor/Public/Components/DMMaterialEffectStack.h
@@ -138,7 +138,7 @@ public:
 	DYNAMICMATERIALEDITOR_API virtual UDMMaterialComponent* GetParentComponent() const override;
 	DYNAMICMATERIALEDITOR_API virtual FString GetComponentPathComponent() const override;
 	DYNAMICMATERIALEDITOR_API virtual FText GetComponentDescription() const override;
-	DYNAMICMATERIALEDITOR_API virtual void Update(EDMUpdateType InUpdateType) override;
+	DYNAMICMATERIALEDITOR_API virtual void Update(UDMMaterialComponent* InSource, EDMUpdateType InUpdateType) override;
 	DYNAMICMATERIALEDITOR_API virtual void PostEditorDuplicate(UDynamicMaterialModel* InMaterialModel, UDMMaterialComponent* InParent) override;
 	//~ End UDMMaterialComponent
 
diff --git a/Engine/Plugins/Experimental/DynamicMaterial/Source/DynamicMaterialEditor/Public/Components/DMMaterialLayer.h b/Engine/Plugins/Experimental/DynamicMaterial/Source/DynamicMaterialEditor/Public/Components/DMMaterialLayer.h
--- a/Engine/Plugins/Experimental/DynamicMaterial/Source/DynamicMaterialEditor/Public/Components/DMMaterialLayer.h
+++ b/Engine/Plugins/Experimental/DynamicMaterial/Source/DynamicMaterialEditor/Public/Components/DMMaterialLayer.h
@@ -152,7 +152,7 @@ public:
 	DYNAMICMATERIALEDITOR_API virtual UDMMaterialComponent* GetParentComponent() const override;
 	DYNAMICMATERIALEDITOR_API virtual FString GetComponentPathComponent() const override;
 	DYNAMICMATERIALEDITOR_API virtual FText GetComponentDescription() const override;
-	DYNAMICMATERIALEDITOR_API virtual void Update(EDMUpdateType InUpdateType) override;
+	DYNAMICMATERIALEDITOR_API virtual void Update(UDMMaterialComponent* InSource, EDMUpdateType InUpdateType) override;
 	DYNAMICMATERIALEDITOR_API virtual void PostEditorDuplicate(UDynamicMaterialModel* InMaterialModel, UDMMaterialComponent* InParent) override;
 	//~ End UDMMaterialComponent
 
diff --git a/Engine/Plugins/Experimental/DynamicMaterial/Source/DynamicMaterialEditor/Public/Components/DMMaterialProperty.h b/Engine/Plugins/Experimental/DynamicMaterial/Source/DynamicMaterialEditor/Public/Components/DMMaterialProperty.h
--- a/Engine/Plugins/Experimental/DynamicMaterial/Source/DynamicMaterialEditor/Public/Components/DMMaterialProperty.h
+++ b/Engine/Plugins/Experimental/DynamicMaterial/Source/DynamicMaterialEditor/Public/Components/DMMaterialProperty.h
@@ -197,7 +197,7 @@ public:
 	virtual void AddOutputProcessor(const TSharedRef<FDMMaterialBuildState>& InBuildState) const;
 
 	//~ Begin UDMMaterialComponent
-	DYNAMICMATERIALEDITOR_API virtual void Update(EDMUpdateType InUpdateType) override;
+	DYNAMICMATERIALEDITOR_API virtual void Update(UDMMaterialComponent* InSource, EDMUpdateType InUpdateType) override;
 	DYNAMICMATERIALEDITOR_API virtual FString GetComponentPathComponent() const override;
 	DYNAMICMATERIALEDITOR_API virtual void PostEditorDuplicate(UDynamicMaterialModel* InMaterialModel, UDMMaterialComponent* InParent) override;
 	DYNAMICMATERIALEDITOR_API virtual void PreEditChange(FProperty* InPropertyAboutToChange) override;
diff --git a/Engine/Plugins/Experimental/DynamicMaterial/Source/DynamicMaterialEditor/Public/Components/DMMaterialSlot.h b/Engine/Plugins/Experimental/DynamicMaterial/Source/DynamicMaterialEditor/Public/Components/DMMaterialSlot.h
--- a/Engine/Plugins/Experimental/DynamicMaterial/Source/DynamicMaterialEditor/Public/Components/DMMaterialSlot.h
+++ b/Engine/Plugins/Experimental/DynamicMaterial/Source/DynamicMaterialEditor/Public/Components/DMMaterialSlot.h
@@ -149,7 +149,7 @@ public:
 	DYNAMICMATERIALEDITOR_API bool ChangeMaterialProperty(EDMMaterialPropertyType InPropertyFrom, EDMMaterialPropertyType InPropertyTo);
 
 	//~ Begin UDMMaterialComponent
-	DYNAMICMATERIALEDITOR_API virtual void Update(EDMUpdateType InUpdateType) override;
+	DYNAMICMATERIALEDITOR_API virtual void Update(UDMMaterialComponent* InSource, EDMUpdateType InUpdateType) override;
 	DYNAMICMATERIALEDITOR_API virtual FString GetComponentPathComponent() const override;
 	virtual UDMMaterialComponent* GetParentComponent() const override { return nullptr; }
 	DYNAMICMATERIALEDITOR_API virtual void PostEditorDuplicate(UDynamicMaterialModel* InMaterialModel, UDMMaterialComponent* InParent) override;
diff --git a/Engine/Plugins/Experimental/DynamicMaterial/Source/DynamicMaterialEditor/Public/Components/DMMaterialStage.h b/Engine/Plugins/Experimental/DynamicMaterial/Source/DynamicMaterialEditor/Public/Components/DMMaterialStage.h
--- a/Engine/Plugins/Experimental/DynamicMaterial/Source/DynamicMaterialEditor/Public/Components/DMMaterialStage.h
+++ b/Engine/Plugins/Experimental/DynamicMaterial/Source/DynamicMaterialEditor/Public/Components/DMMaterialStage.h
@@ -233,7 +233,7 @@ public:
 
 	//~ Begin UDMMaterialComponent
 	DYNAMICMATERIALEDITOR_API virtual FText GetComponentDescription() const override;
-	DYNAMICMATERIALEDITOR_API virtual void Update(EDMUpdateType InUpdateType) override;
+	DYNAMICMATERIALEDITOR_API virtual void Update(UDMMaterialComponent* InSource, EDMUpdateType InUpdateType) override;
 	DYNAMICMATERIALEDITOR_API virtual FString GetComponentPathComponent() const override;
 	DYNAMICMATERIALEDITOR_API virtual UDMMaterialComponent* GetParentComponent() const override;
 	DYNAMICMATERIALEDITOR_API virtual void PostEditorDuplicate(UDynamicMaterialModel* InMaterialModel, UDMMaterialComponent* InParent) override;
diff --git a/Engine/Plugins/Experimental/DynamicMaterial/Source/DynamicMaterialEditor/Public/Components/DMMaterialStageInput.h b/Engine/Plugins/Experimental/DynamicMaterial/Source/DynamicMaterialEditor/Public/Components/DMMaterialStageInput.h
--- a/Engine/Plugins/Experimental/DynamicMaterial/Source/DynamicMaterialEditor/Public/Components/DMMaterialStageInput.h
+++ b/Engine/Plugins/Experimental/DynamicMaterial/Source/DynamicMaterialEditor/Public/Components/DMMaterialStageInput.h
@@ -25,7 +25,7 @@ public:
 		
 	//~ Begin UDMMaterialComponent
 	DYNAMICMATERIALEDITOR_API virtual FString GetComponentPathComponent() const override;
-	DYNAMICMATERIALEDITOR_API virtual void Update(EDMUpdateType InUpdateType) override;
+	DYNAMICMATERIALEDITOR_API virtual void Update(UDMMaterialComponent* InSource, EDMUpdateType InUpdateType) override;
 	//~ End UDMMaterialComponent
 
 protected:
diff --git a/Engine/Plugins/Experimental/DynamicMaterial/Source/DynamicMaterialEditor/Public/Components/DMMaterialStageSource.h b/Engine/Plugins/Experimental/DynamicMaterial/Source/DynamicMaterialEditor/Public/Components/DMMaterialStageSource.h
--- a/Engine/Plugins/Experimental/DynamicMaterial/Source/DynamicMaterialEditor/Public/Components/DMMaterialStageSource.h
+++ b/Engine/Plugins/Experimental/DynamicMaterial/Source/DynamicMaterialEditor/Public/Components/DMMaterialStageSource.h
@@ -74,7 +74,7 @@ public:
 	//~ End FNotifyHook
 
 	//~ Begin UDMMaterialComponent
-	DYNAMICMATERIALEDITOR_API virtual void Update(EDMUpdateType InUpdateType) override;
+	DYNAMICMATERIALEDITOR_API virtual void Update(UDMMaterialComponent* InSource, EDMUpdateType InUpdateType) override;
 	DYNAMICMATERIALEDITOR_API virtual UDMMaterialComponent* GetParentComponent() const override;
 	//~ End UDMMaterialComponent
 
diff --git a/Engine/Plugins/Experimental/DynamicMaterial/Source/DynamicMaterialEditor/Public/Components/DMMaterialStageThroughputLayerBlend.h b/Engine/Plugins/Experimental/DynamicMaterial/Source/DynamicMaterialEditor/Public/Components/DMMaterialStageThroughputLayerBlend.h
--- a/Engine/Plugins/Experimental/DynamicMaterial/Source/DynamicMaterialEditor/Public/Components/DMMaterialStageThroughputLayerBlend.h
+++ b/Engine/Plugins/Experimental/DynamicMaterial/Source/DynamicMaterialEditor/Public/Components/DMMaterialStageThroughputLayerBlend.h
@@ -40,7 +40,7 @@ public:
 
 	//~ Begin UDMMaterialComponent
 	DYNAMICMATERIALEDITOR_API virtual void OnComponentAdded() override;
-	DYNAMICMATERIALEDITOR_API virtual void Update(EDMUpdateType InUpdateType) override;
+	DYNAMICMATERIALEDITOR_API virtual void Update(UDMMaterialComponent* InSource, EDMUpdateType InUpdateType) override;
 	DYNAMICMATERIALEDITOR_API virtual void GenerateExpressions(const TSharedRef<FDMMaterialBuildState>& InBuildState) const override;
 	DYNAMICMATERIALEDITOR_API virtual void PostEditorDuplicate(UDynamicMaterialModel* InMaterialModel, UDMMaterialComponent* InParent) override;
 	DYNAMICMATERIALEDITOR_API virtual bool IsPropertyVisible(FName InProperty) const override;
@@ -106,7 +106,7 @@ protected:
 	DYNAMICMATERIALEDITOR_API virtual void UpdateAlphaOnlyMaskStatus();
 
 	/** Called when the owning stage is updated. */
-	DYNAMICMATERIALEDITOR_API virtual void OnStageUpdated(UDMMaterialComponent* InComponent, EDMUpdateType InUpdateType);
+	DYNAMICMATERIALEDITOR_API virtual void OnStageUpdated(UDMMaterialComponent* InComponent, UDMMaterialComponent* InSource, EDMUpdateType InUpdateType);
 
 	/** Update this and all parent layer mask stages that are only-only (no based). */
 	DYNAMICMATERIALEDITOR_API virtual void UpdateAlphaOnlyMasks(EDMUpdateType InUpdateType);
diff --git a/Engine/Plugins/Experimental/DynamicMaterial/Source/DynamicMaterialEditor/Public/Components/MaterialStageExpressions/DMMSEMathBase.h b/Engine/Plugins/Experimental/DynamicMaterial/Source/DynamicMaterialEditor/Public/Components/MaterialStageExpressions/DMMSEMathBase.h
--- a/Engine/Plugins/Experimental/DynamicMaterial/Source/DynamicMaterialEditor/Public/Components/MaterialStageExpressions/DMMSEMathBase.h
+++ b/Engine/Plugins/Experimental/DynamicMaterial/Source/DynamicMaterialEditor/Public/Components/MaterialStageExpressions/DMMSEMathBase.h
@@ -27,7 +27,7 @@ public:
 	//~ End UDMMaterialStageSource
 
 	//~ Begin UDMMaterialComponent
-	DYNAMICMATERIALEDITOR_API virtual void Update(EDMUpdateType InUpdateType) override;
+	DYNAMICMATERIALEDITOR_API virtual void Update(UDMMaterialComponent* InSource, EDMUpdateType InUpdateType) override;
 	//~ End UDMMaterialComponent
 
 protected:
diff --git a/Engine/Plugins/Experimental/DynamicMaterial/Source/DynamicMaterialEditor/Public/Components/MaterialStageInputs/DMMSISlot.h b/Engine/Plugins/Experimental/DynamicMaterial/Source/DynamicMaterialEditor/Public/Components/MaterialStageInputs/DMMSISlot.h
--- a/Engine/Plugins/Experimental/DynamicMaterial/Source/DynamicMaterialEditor/Public/Components/MaterialStageInputs/DMMSISlot.h
+++ b/Engine/Plugins/Experimental/DynamicMaterial/Source/DynamicMaterialEditor/Public/Components/MaterialStageInputs/DMMSISlot.h
@@ -77,7 +77,7 @@ protected:
 
 	virtual ~UDMMaterialStageInputSlot() override = default;
 
-	void OnSlotUpdated(UDMMaterialComponent* InComponent, EDMUpdateType InUpdateType);
+	void OnSlotUpdated(UDMMaterialComponent* InComponent, UDMMaterialComponent* InSource, EDMUpdateType InUpdateType);
 	void OnSlotConnectorsUpdated(UDMMaterialSlot* InSlot);
 	void OnSlotRemoved(UDMMaterialComponent* InComponent, EDMComponentLifetimeState InLifetimeState);
 	void OnParentSlotRemoved(UDMMaterialComponent* InComponent, EDMComponentLifetimeState InLifetimeState);
diff --git a/Engine/Plugins/Experimental/DynamicMaterial/Source/DynamicMaterialEditor/Public/Components/MaterialStageInputs/DMMSITextureUV.h b/Engine/Plugins/Experimental/DynamicMaterial/Source/DynamicMaterialEditor/Public/Components/MaterialStageInputs/DMMSITextureUV.h
--- a/Engine/Plugins/Experimental/DynamicMaterial/Source/DynamicMaterialEditor/Public/Components/MaterialStageInputs/DMMSITextureUV.h
+++ b/Engine/Plugins/Experimental/DynamicMaterial/Source/DynamicMaterialEditor/Public/Components/MaterialStageInputs/DMMSITextureUV.h
@@ -70,7 +70,7 @@ protected:
 	UPROPERTY(EditInstanceOnly, BlueprintReadOnly, Instanced, Category = "Material Designer")
 	TObjectPtr<UDMTextureUV> TextureUV;
 
-	void OnTextureUVUpdated(UDMMaterialComponent* InComponent, EDMUpdateType InUpdateType);
+	void OnTextureUVUpdated(UDMMaterialComponent* InComponent, UDMMaterialComponent* InSource, EDMUpdateType InUpdateType);
 
 	void InitTextureUV();
 
diff --git a/Engine/Plugins/Experimental/DynamicMaterial/Source/DynamicMaterialEditor/Public/Components/MaterialStageInputs/DMMSIThroughput.h b/Engine/Plugins/Experimental/DynamicMaterial/Source/DynamicMaterialEditor/Public/Components/MaterialStageInputs/DMMSIThroughput.h
--- a/Engine/Plugins/Experimental/DynamicMaterial/Source/DynamicMaterialEditor/Public/Components/MaterialStageInputs/DMMSIThroughput.h
+++ b/Engine/Plugins/Experimental/DynamicMaterial/Source/DynamicMaterialEditor/Public/Components/MaterialStageInputs/DMMSIThroughput.h
@@ -56,7 +56,7 @@ protected:
 	UPROPERTY(VisibleInstanceOnly, BlueprintReadOnly, Instanced, Category = "Material Designer")
 	TObjectPtr<UDMMaterialSubStage> SubStage;
 
-	void OnSubStageUpdated(UDMMaterialComponent* InComponent, EDMUpdateType InUpdateType);
+	void OnSubStageUpdated(UDMMaterialComponent* InComponent, UDMMaterialComponent* InSource, EDMUpdateType InUpdateType);
 
 	void SetMaterialStageThroughputClass(TSubclassOf<UDMMaterialStageThroughput> InMaterialStageThroughputClass);
 
