Material Designer
- Added a property preview page.

diff --git a/Engine/Plugins/Experimental/DynamicMaterial/Source/DynamicMaterialEditor/Private/Components/DMMaterialProperty.cpp b/Engine/Plugins/Experimental/DynamicMaterial/Source/DynamicMaterialEditor/Private/Components/DMMaterialProperty.cpp
--- a/Engine/Plugins/Experimental/DynamicMaterial/Source/DynamicMaterialEditor/Private/Components/DMMaterialProperty.cpp
+++ b/Engine/Plugins/Experimental/DynamicMaterial/Source/DynamicMaterialEditor/Private/Components/DMMaterialProperty.cpp
@@ -89,7 +89,7 @@ void UDMMaterialProperty::SetEnabled(bool bInEnabled)
 
 FText UDMMaterialProperty::GetDescription() const
 {
-	return StaticEnum<EDMMaterialPropertyType>()->GetDisplayNameTextByValue(static_cast<int64>(MaterialProperty));
+	return UE::DynamicMaterialEditor::Private::GetMaterialPropertyShortDisplayName(MaterialProperty);
 }
 
 bool UDMMaterialProperty::IsMaterialPin() const
diff --git a/Engine/Plugins/Experimental/DynamicMaterial/Source/DynamicMaterialEditor/Private/Components/MaterialStageInputs/DMMSISlot.cpp b/Engine/Plugins/Experimental/DynamicMaterial/Source/DynamicMaterialEditor/Private/Components/MaterialStageInputs/DMMSISlot.cpp
--- a/Engine/Plugins/Experimental/DynamicMaterial/Source/DynamicMaterialEditor/Private/Components/MaterialStageInputs/DMMSISlot.cpp
+++ b/Engine/Plugins/Experimental/DynamicMaterial/Source/DynamicMaterialEditor/Private/Components/MaterialStageInputs/DMMSISlot.cpp
@@ -11,6 +11,7 @@
 #include "DMValueDefinition.h"
 #include "Model/DMMaterialBuildState.h"
 #include "Model/DynamicMaterialModelEditorOnlyData.h"
+#include "Utils/DMPrivate.h"
 
 #define LOCTEXT_NAMESPACE "DMMaterialStageInputSlot"
 
@@ -435,7 +436,7 @@ void UDMMaterialStageInputSlot::UpdateOutputConnectors()
 
 		if (MaterialProperties.Contains(MaterialProperty))
 		{
-			const FText MaterialPropertyName = StaticEnum<EDMMaterialPropertyType>()->GetDisplayNameTextByValue(static_cast<int64>(MaterialProperty));
+			const FText MaterialPropertyName = UE::DynamicMaterialEditor::Private::GetMaterialPropertyShortDisplayName(MaterialProperty);
 			const TArray<EDMValueType>& OutputTypes = Slot->GetOutputConnectorTypesForMaterialProperty(MaterialProperty);
 
 			for (int32 OutputTypeIdx = 0; OutputTypeIdx < OutputTypes.Num(); ++OutputTypeIdx)
diff --git a/Engine/Plugins/Experimental/DynamicMaterial/Source/DynamicMaterialEditor/Private/Model/DMMaterialBuildUtils.cpp b/Engine/Plugins/Experimental/DynamicMaterial/Source/DynamicMaterialEditor/Private/Model/DMMaterialBuildUtils.cpp
--- a/Engine/Plugins/Experimental/DynamicMaterial/Source/DynamicMaterialEditor/Private/Model/DMMaterialBuildUtils.cpp
+++ b/Engine/Plugins/Experimental/DynamicMaterial/Source/DynamicMaterialEditor/Private/Model/DMMaterialBuildUtils.cpp
@@ -88,7 +88,7 @@ UMaterialExpression* FDMMaterialBuildUtils::CreateExpressionParameter(TSubclassO
 				{
 					const EDMMaterialPropertyType PropertyType = Property->GetMaterialProperty();
 					PropertyIndex = static_cast<int32>(PropertyType);
-					PropertyName = StaticEnum<EDMMaterialPropertyType>()->GetDisplayNameTextByValue(static_cast<int64>(PropertyType)).ToString();
+					PropertyName = UE::DynamicMaterialEditor::Private::GetMaterialPropertyShortDisplayName(PropertyType).ToString();
 					break;
 				}
 
diff --git a/Engine/Plugins/Experimental/DynamicMaterial/Source/DynamicMaterialEditor/Private/UI/PropertyGenerators/DMMaterialModelPropertyRowGenerator.cpp b/Engine/Plugins/Experimental/DynamicMaterial/Source/DynamicMaterialEditor/Private/UI/PropertyGenerators/DMMaterialModelPropertyRowGenerator.cpp
--- a/Engine/Plugins/Experimental/DynamicMaterial/Source/DynamicMaterialEditor/Private/UI/PropertyGenerators/DMMaterialModelPropertyRowGenerator.cpp
+++ b/Engine/Plugins/Experimental/DynamicMaterial/Source/DynamicMaterialEditor/Private/UI/PropertyGenerators/DMMaterialModelPropertyRowGenerator.cpp
@@ -9,6 +9,7 @@
 #include "Model/DynamicMaterialModelEditorOnlyData.h"
 #include "UI/Utils/DMWidgetStatics.h"
 #include "UI/Widgets/Editor/SDMMaterialGlobalSettingsEditor.h"
+#include "Utils/DMPrivate.h"
 
 #define LOCTEXT_NAMESPACE "DMMaterialModelPropertyRowGenerator"
 
@@ -112,8 +113,6 @@ void FDMMaterialModelPropertyRowGenerator::AddGlobalMaterialParameterValue(EDMMa
 			{
 				if (UDMMaterialValueFloat1* AlphaValue = Cast<UDMMaterialValueFloat1>(MaterialProperty->GetComponent(UDynamicMaterialModelEditorOnlyData::AlphaValueName)))
 				{
-					UEnum* MaterialPropertyEnum = StaticEnum<EDMMaterialPropertyType>();
-
 					AddGlobalValue(
 						InGlobalSettingEditorWidget,
 						InMaterialModelBase,
@@ -121,7 +120,7 @@ void FDMMaterialModelPropertyRowGenerator::AddGlobalMaterialParameterValue(EDMMa
 						AlphaValue,
 						FText::Format(
 							LOCTEXT("PropertyFormat", "Global {0}"),
-							MaterialPropertyEnum->GetDisplayNameTextByValue(static_cast<int64>(InProperty))
+							UE::DynamicMaterialEditor::Private::GetMaterialPropertyShortDisplayName(InProperty)
 						)
 					);
 				}
diff --git a/Engine/Plugins/Experimental/DynamicMaterial/Source/DynamicMaterialEditor/Private/UI/Widgets/Editor/EditorLayouts/SDMMaterialEditor_LeftBase.cpp b/Engine/Plugins/Experimental/DynamicMaterial/Source/DynamicMaterialEditor/Private/UI/Widgets/Editor/EditorLayouts/SDMMaterialEditor_LeftBase.cpp
--- a/Engine/Plugins/Experimental/DynamicMaterial/Source/DynamicMaterialEditor/Private/UI/Widgets/Editor/EditorLayouts/SDMMaterialEditor_LeftBase.cpp
+++ b/Engine/Plugins/Experimental/DynamicMaterial/Source/DynamicMaterialEditor/Private/UI/Widgets/Editor/EditorLayouts/SDMMaterialEditor_LeftBase.cpp
@@ -7,6 +7,7 @@
 #include "UI/Widgets/Editor/SDMMaterialComponentEditor.h"
 #include "UI/Widgets/Editor/SDMMaterialGlobalSettingsEditor.h"
 #include "UI/Widgets/Editor/SDMMaterialPreview.h"
+#include "UI/Widgets/Editor/SDMMaterialPropertyPreviews.h"
 #include "UI/Widgets/Editor/SDMMaterialPropertySelector.h"
 #include "UI/Widgets/Editor/SDMMaterialSlotEditor.h"
 #include "Widgets/Layout/SBorder.h"
@@ -43,7 +44,7 @@ void SDMMaterialEditor_LeftBase::EditComponent(UDMMaterialComponent* InComponent
 		return;
 	}
 
-	if (bGlobalSettingsMode)
+	if (EditMode != EDMMaterialEditorMode::EditSlot)
 	{
 		RightSlot.Invalidate();
 	}
@@ -53,12 +54,12 @@ void SDMMaterialEditor_LeftBase::EditComponent(UDMMaterialComponent* InComponent
 
 void SDMMaterialEditor_LeftBase::EditGlobalSettings(bool bInForceRefresh)
 {
-	if (bGlobalSettingsMode && !bInForceRefresh)
+	if (EditMode == EDMMaterialEditorMode::GlobalSettings && !bInForceRefresh)
 	{
 		return;
 	}
 
-	if (!bGlobalSettingsMode)
+	if (EditMode != EDMMaterialEditorMode::GlobalSettings)
 	{
 		RightSlot.Invalidate();
 	}
@@ -66,6 +67,21 @@ void SDMMaterialEditor_LeftBase::EditGlobalSettings(bool bInForceRefresh)
 	SDMMaterialEditor::EditGlobalSettings(bInForceRefresh);
 }
 
+void SDMMaterialEditor_LeftBase::ShowPropertyPreviews(bool bInForceRefresh)
+{
+	if (EditMode == EDMMaterialEditorMode::PropertyPreviews && !bInForceRefresh)
+	{
+		return;
+	}
+
+	if (EditMode != EDMMaterialEditorMode::PropertyPreviews)
+	{
+		RightSlot.Invalidate();
+	}
+
+	SDMMaterialEditor::ShowPropertyPreviews(bInForceRefresh);
+}
+
 void SDMMaterialEditor_LeftBase::ValidateSlots_Main()
 {
 	if (LeftSlot.HasBeenInvalidated())
@@ -152,22 +168,38 @@ TSharedRef<SWidget> SDMMaterialEditor_LeftBase::CreateSlot_Right()
 
 	const bool bHasSlotToEdit = SlotToEdit.IsValid();
 
-	if (!bGlobalSettingsMode && !bHasSlotToEdit)
+	if (EditMode == EDMMaterialEditorMode::EditSlot && !bHasSlotToEdit)
 	{
-		bGlobalSettingsMode = true;
+		EditMode = EDMMaterialEditorMode::GlobalSettings;
 	}
 	else if (bHasSlotToEdit)
 	{
-		bGlobalSettingsMode = false;
+		EditMode = EDMMaterialEditorMode::EditSlot;
+	}
+
+	TSharedPtr<SWidget> Content;
+
+	switch (EditMode)
+	{
+		default:
+		case EDMMaterialEditorMode::GlobalSettings:
+			Content = CreateSlot_Right_GlobalSettings();
+			break;
+
+		case EDMMaterialEditorMode::PropertyPreviews:
+			Content = CreateSlot_Right_PropertyPreviews();
+			break;
+
+		case EDMMaterialEditorMode::EditSlot:
+			Content = CreateSlot_Right_EditSlot();
+			break;
 	}
 
 	return SNew(SBorder)
 		.BorderImage(FAppStyle::GetBrush(EditorDarkBackground))
 		.Padding(FMargin(0.f, 5.f))
 		[
-			bGlobalSettingsMode
-				? CreateSlot_Right_GlobalSettings()
-				: CreateSlot_Right_Slot()
+			Content.ToSharedRef()
 		];
 }
 
@@ -195,7 +227,31 @@ TSharedRef<SWidget> SDMMaterialEditor_LeftBase::CreateSlot_Right_GlobalSettings(
 	return NewRight;
 }
 
-TSharedRef<SWidget> SDMMaterialEditor_LeftBase::CreateSlot_Right_Slot()
+TSharedRef<SWidget> SDMMaterialEditor_LeftBase::CreateSlot_Right_PropertyPreviews()
+{
+	using namespace UE::DynamicMaterialEditor::Private;
+
+	SScrollBox::FSlot* PropertyPreviewsSlotPtr = nullptr;
+
+	TSharedRef<SBorder> NewRight = SNew(SBorder)
+		.BorderImage(FAppStyle::GetBrush(EditorLightBackground))
+		.Padding(0.f)
+		[
+			SNew(SScrollBox)
+			+ SScrollBox::Slot()
+			.Expose(PropertyPreviewsSlotPtr)
+			.VAlign(EVerticalAlignment::VAlign_Fill)
+			[
+				SNullWidget::NullWidget
+			]
+		];
+
+	MaterialPropertyPreviewsSlot = TDMWidgetSlot<SDMMaterialPropertyPreviews>(PropertyPreviewsSlotPtr, CreateSlot_MaterialPropertyPreviews());
+
+	return NewRight;
+}
+
+TSharedRef<SWidget> SDMMaterialEditor_LeftBase::CreateSlot_Right_EditSlot()
 {
 	using namespace UE::DynamicMaterialEditor::Private;
 
diff --git a/Engine/Plugins/Experimental/DynamicMaterial/Source/DynamicMaterialEditor/Private/UI/Widgets/Editor/EditorLayouts/SDMMaterialEditor_LeftBase.h b/Engine/Plugins/Experimental/DynamicMaterial/Source/DynamicMaterialEditor/Private/UI/Widgets/Editor/EditorLayouts/SDMMaterialEditor_LeftBase.h
--- a/Engine/Plugins/Experimental/DynamicMaterial/Source/DynamicMaterialEditor/Private/UI/Widgets/Editor/EditorLayouts/SDMMaterialEditor_LeftBase.h
+++ b/Engine/Plugins/Experimental/DynamicMaterial/Source/DynamicMaterialEditor/Private/UI/Widgets/Editor/EditorLayouts/SDMMaterialEditor_LeftBase.h
@@ -23,6 +23,7 @@ public:
 	virtual void EditSlot(UDMMaterialSlot* InSlot, bool bInForceRefresh = false) override;
 	virtual void EditComponent(UDMMaterialComponent* InComponent, bool bInForceRefresh = false) override;
 	virtual void EditGlobalSettings(bool bInForceRefresh = false) override;
+	virtual void ShowPropertyPreviews(bool bInForceRefresh = false) override;
 	//~ End SDMMaterialEditor
 
 protected:
@@ -35,7 +36,9 @@ protected:
 
 	TSharedRef<SWidget> CreateSlot_Right_GlobalSettings();
 
-	TSharedRef<SWidget> CreateSlot_Right_Slot();
+	TSharedRef<SWidget> CreateSlot_Right_PropertyPreviews();
+
+	TSharedRef<SWidget> CreateSlot_Right_EditSlot();
 
 	//~ Begin SDMMaterialEditor
 	virtual void ValidateSlots_Main() override;
diff --git a/Engine/Plugins/Experimental/DynamicMaterial/Source/DynamicMaterialEditor/Private/UI/Widgets/Editor/EditorLayouts/SDMMaterialEditor_TopBase.cpp b/Engine/Plugins/Experimental/DynamicMaterial/Source/DynamicMaterialEditor/Private/UI/Widgets/Editor/EditorLayouts/SDMMaterialEditor_TopBase.cpp
--- a/Engine/Plugins/Experimental/DynamicMaterial/Source/DynamicMaterialEditor/Private/UI/Widgets/Editor/EditorLayouts/SDMMaterialEditor_TopBase.cpp
+++ b/Engine/Plugins/Experimental/DynamicMaterial/Source/DynamicMaterialEditor/Private/UI/Widgets/Editor/EditorLayouts/SDMMaterialEditor_TopBase.cpp
@@ -6,6 +6,7 @@
 #include "UI/Widgets/Editor/SDMMaterialComponentEditor.h"
 #include "UI/Widgets/Editor/SDMMaterialGlobalSettingsEditor.h"
 #include "UI/Widgets/Editor/SDMMaterialPreview.h"
+#include "UI/Widgets/Editor/SDMMaterialPropertyPreviews.h"
 #include "UI/Widgets/Editor/SDMMaterialPropertySelector.h"
 #include "UI/Widgets/Editor/SDMMaterialSlotEditor.h"
 #include "Widgets/Layout/SBorder.h"
@@ -41,7 +42,7 @@ void SDMMaterialEditor_TopBase::EditComponent(UDMMaterialComponent* InComponent,
 		return;
 	}
 
-	if (bGlobalSettingsMode)
+	if (EditMode != EDMMaterialEditorMode::EditSlot)
 	{
 		BottomSlot.Invalidate();
 	}
@@ -51,12 +52,12 @@ void SDMMaterialEditor_TopBase::EditComponent(UDMMaterialComponent* InComponent,
 
 void SDMMaterialEditor_TopBase::EditGlobalSettings(bool bInForceRefresh)
 {
-	if (bGlobalSettingsMode && !bInForceRefresh)
+	if (EditMode == EDMMaterialEditorMode::GlobalSettings && !bInForceRefresh)
 	{
 		return;
 	}
 
-	if (!bGlobalSettingsMode)
+	if (EditMode != EDMMaterialEditorMode::GlobalSettings)
 	{
 		BottomSlot.Invalidate();
 	}
@@ -64,6 +65,21 @@ void SDMMaterialEditor_TopBase::EditGlobalSettings(bool bInForceRefresh)
 	SDMMaterialEditor::EditGlobalSettings(bInForceRefresh);
 }
 
+void SDMMaterialEditor_TopBase::ShowPropertyPreviews(bool bInForceRefresh)
+{
+	if (EditMode == EDMMaterialEditorMode::PropertyPreviews && !bInForceRefresh)
+	{
+		return;
+	}
+
+	if (EditMode != EDMMaterialEditorMode::PropertyPreviews)
+	{
+		BottomSlot.Invalidate();
+	}
+
+	SDMMaterialEditor::ShowPropertyPreviews(bInForceRefresh);
+}
+
 void SDMMaterialEditor_TopBase::ValidateSlots_Main()
 {
 	if (TopSlot.HasBeenInvalidated())
@@ -150,22 +166,38 @@ TSharedRef<SWidget> SDMMaterialEditor_TopBase::CreateSlot_Bottom()
 
 	const bool bHasSlotToEdit = SlotToEdit.IsValid();
 
-	if (!bGlobalSettingsMode && !bHasSlotToEdit)
+	if (EditMode == EDMMaterialEditorMode::EditSlot && !bHasSlotToEdit)
 	{
-		bGlobalSettingsMode = true;
+		EditMode = EDMMaterialEditorMode::GlobalSettings;
 	}
 	else if (bHasSlotToEdit)
 	{
-		bGlobalSettingsMode = false;
+		EditMode = EDMMaterialEditorMode::EditSlot;
+	}
+
+	TSharedPtr<SWidget> Content;
+
+	switch (EditMode)
+	{
+		default:
+		case EDMMaterialEditorMode::GlobalSettings:
+			Content = CreateSlot_Bottom_GlobalSettings();
+			break;
+
+		case EDMMaterialEditorMode::PropertyPreviews:
+			Content = CreateSlot_Bottom_PropertyPreviews();
+			break;
+
+		case EDMMaterialEditorMode::EditSlot:
+			Content = CreateSlot_Bottom_EditSlot();
+			break;
 	}
 
 	return SNew(SBorder)
 		.BorderImage(FAppStyle::GetBrush(EditorDarkBackground))
 		.Padding(FMargin(0.f, 5.f))
 		[
-			bGlobalSettingsMode
-				? CreateSlot_Bottom_GlobalSettings()
-				: CreateSlot_Bottom_Slot()
+			Content.ToSharedRef()
 		];
 }
 
@@ -193,7 +225,31 @@ TSharedRef<SWidget> SDMMaterialEditor_TopBase::CreateSlot_Bottom_GlobalSettings(
 	return NewBottom;
 }
 
-TSharedRef<SWidget> SDMMaterialEditor_TopBase::CreateSlot_Bottom_Slot()
+TSharedRef<SWidget> SDMMaterialEditor_TopBase::CreateSlot_Bottom_PropertyPreviews()
+{
+	using namespace UE::DynamicMaterialEditor::Private;
+
+	SScrollBox::FSlot* PropertyPreviewsSlotPtr = nullptr;
+
+	TSharedRef<SBorder> NewBottom = SNew(SBorder)
+		.BorderImage(FAppStyle::GetBrush(EditorLightBackground))
+		.Padding(0.f)
+		[
+			SNew(SScrollBox)
+			+ SScrollBox::Slot()
+			.Expose(PropertyPreviewsSlotPtr)
+			.VAlign(EVerticalAlignment::VAlign_Fill)
+			[
+				SNullWidget::NullWidget
+			]
+		];
+
+	MaterialPropertyPreviewsSlot = TDMWidgetSlot<SDMMaterialPropertyPreviews>(PropertyPreviewsSlotPtr, CreateSlot_MaterialPropertyPreviews());
+
+	return NewBottom;
+}
+
+TSharedRef<SWidget> SDMMaterialEditor_TopBase::CreateSlot_Bottom_EditSlot()
 {
 	using namespace UE::DynamicMaterialEditor::Private;
 
diff --git a/Engine/Plugins/Experimental/DynamicMaterial/Source/DynamicMaterialEditor/Private/UI/Widgets/Editor/EditorLayouts/SDMMaterialEditor_TopBase.h b/Engine/Plugins/Experimental/DynamicMaterial/Source/DynamicMaterialEditor/Private/UI/Widgets/Editor/EditorLayouts/SDMMaterialEditor_TopBase.h
--- a/Engine/Plugins/Experimental/DynamicMaterial/Source/DynamicMaterialEditor/Private/UI/Widgets/Editor/EditorLayouts/SDMMaterialEditor_TopBase.h
+++ b/Engine/Plugins/Experimental/DynamicMaterial/Source/DynamicMaterialEditor/Private/UI/Widgets/Editor/EditorLayouts/SDMMaterialEditor_TopBase.h
@@ -23,6 +23,7 @@ public:
 	virtual void EditSlot(UDMMaterialSlot* InSlot, bool bInForceRefresh = false) override;
 	virtual void EditComponent(UDMMaterialComponent* InComponent, bool bInForceRefresh = false) override;
 	virtual void EditGlobalSettings(bool bInForceRefresh = false) override;
+	virtual void ShowPropertyPreviews(bool bInForceRefresh = false) override;
 	//~ End SDMMaterialEditor
 
 protected:
@@ -35,7 +36,9 @@ protected:
 
 	TSharedRef<SWidget> CreateSlot_Bottom_GlobalSettings();
 
-	TSharedRef<SWidget> CreateSlot_Bottom_Slot();
+	TSharedRef<SWidget> CreateSlot_Bottom_PropertyPreviews();
+
+	TSharedRef<SWidget> CreateSlot_Bottom_EditSlot();
 
 	//~ Begin SDMMaterialEditor
 	virtual void ValidateSlots_Main() override;
diff --git a/Engine/Plugins/Experimental/DynamicMaterial/Source/DynamicMaterialEditor/Private/UI/Widgets/Editor/PropertySelectorLayouts/SDMMaterialPropertySelector_Vertical.cpp b/Engine/Plugins/Experimental/DynamicMaterial/Source/DynamicMaterialEditor/Private/UI/Widgets/Editor/PropertySelectorLayouts/SDMMaterialPropertySelector_Vertical.cpp
--- a/Engine/Plugins/Experimental/DynamicMaterial/Source/DynamicMaterialEditor/Private/UI/Widgets/Editor/PropertySelectorLayouts/SDMMaterialPropertySelector_Vertical.cpp
+++ b/Engine/Plugins/Experimental/DynamicMaterial/Source/DynamicMaterialEditor/Private/UI/Widgets/Editor/PropertySelectorLayouts/SDMMaterialPropertySelector_Vertical.cpp
@@ -23,17 +23,8 @@ void SDMMaterialPropertySelector_Vertical::Construct(const FArguments& InArgs, c
 
 TSharedRef<SWidget> SDMMaterialPropertySelector_Vertical::CreateSlot_SelectButton(EDMMaterialPropertyType InMaterialProperty)
 {
-	UEnum* PropertyEnum = StaticEnum<EDMMaterialPropertyType>();
-
-	const FText ButtonText = InMaterialProperty == EDMMaterialPropertyType::None
-		? LOCTEXT("GlobalSettings", "Global Settings")
-		: PropertyEnum->GetDisplayNameTextByValue(static_cast<int64>(InMaterialProperty));
-
-	const FText Format = LOCTEXT("PropertySelectFormat", "Edit the {0} property.");
-
-	const FText ToolTip = (InMaterialProperty == EDMMaterialPropertyType::None)
-		? LOCTEXT("GeneralSettingsToolTip", "Edit the material global settings.")
-		: FText::Format(Format, PropertyEnum->GetDisplayNameTextByValue(static_cast<int64>(InMaterialProperty)));
+	const FText ButtonText = GetSelectButtonText(InMaterialProperty, /* Short Name */ false);
+	const FText ToolTip = GetButtonToolTip(InMaterialProperty);
 
 	return SNew(SCheckBox)
 		.Style(FAppStyle::Get(), "DetailsView.SectionButton")
diff --git a/Engine/Plugins/Experimental/DynamicMaterial/Source/DynamicMaterialEditor/Private/UI/Widgets/Editor/PropertySelectorLayouts/SDMMaterialPropertySelector_VerticalBase.cpp b/Engine/Plugins/Experimental/DynamicMaterial/Source/DynamicMaterialEditor/Private/UI/Widgets/Editor/PropertySelectorLayouts/SDMMaterialPropertySelector_VerticalBase.cpp
--- a/Engine/Plugins/Experimental/DynamicMaterial/Source/DynamicMaterialEditor/Private/UI/Widgets/Editor/PropertySelectorLayouts/SDMMaterialPropertySelector_VerticalBase.cpp
+++ b/Engine/Plugins/Experimental/DynamicMaterial/Source/DynamicMaterialEditor/Private/UI/Widgets/Editor/PropertySelectorLayouts/SDMMaterialPropertySelector_VerticalBase.cpp
@@ -42,6 +42,13 @@ TSharedRef<SWidget> SDMMaterialPropertySelector_VerticalBase::CreateSlot_Propert
 
 	++Row;
 
+	NewSlotList->AddSlot(PropertySelectorColumns::Select, Row)
+		[
+			CreateSlot_SelectButton(EDMMaterialPropertyType::Any)
+		];
+
+	++Row;
+
 	for (const TPair<EDMMaterialPropertyType, UDMMaterialProperty*>& PropertyPair : EditorOnlyData->GetMaterialProperties())
 	{
 		if (IsCustomMaterialProperty(PropertyPair.Key))
diff --git a/Engine/Plugins/Experimental/DynamicMaterial/Source/DynamicMaterialEditor/Private/UI/Widgets/Editor/PropertySelectorLayouts/SDMMaterialPropertySelector_VerticalSlim.cpp b/Engine/Plugins/Experimental/DynamicMaterial/Source/DynamicMaterialEditor/Private/UI/Widgets/Editor/PropertySelectorLayouts/SDMMaterialPropertySelector_VerticalSlim.cpp
--- a/Engine/Plugins/Experimental/DynamicMaterial/Source/DynamicMaterialEditor/Private/UI/Widgets/Editor/PropertySelectorLayouts/SDMMaterialPropertySelector_VerticalSlim.cpp
+++ b/Engine/Plugins/Experimental/DynamicMaterial/Source/DynamicMaterialEditor/Private/UI/Widgets/Editor/PropertySelectorLayouts/SDMMaterialPropertySelector_VerticalSlim.cpp
@@ -24,17 +24,8 @@ void SDMMaterialPropertySelector_VerticalSlim::Construct(const FArguments& InArg
 
 TSharedRef<SWidget> SDMMaterialPropertySelector_VerticalSlim::CreateSlot_SelectButton(EDMMaterialPropertyType InMaterialProperty)
 {
-	UEnum* PropertyEnum = StaticEnum<EDMMaterialPropertyType>();
-
-	const FText ButtonText = InMaterialProperty == EDMMaterialPropertyType::None
-		? LOCTEXT("GlobalSettings", "Global Settings")
-		: UE::DynamicMaterialEditor::Private::GetMaterialPropertyShortDisplayName(InMaterialProperty);
-
-	const FText Format = LOCTEXT("PropertySelectFormat", "Edit the {0} property.");
-
-	const FText ToolTip = (InMaterialProperty == EDMMaterialPropertyType::None)
-		? LOCTEXT("GeneralSettingsToolTip", "Edit the material global settings.")
-		: FText::Format(Format, PropertyEnum->GetDisplayNameTextByValue(static_cast<int64>(InMaterialProperty)));
+	const FText ButtonText = GetSelectButtonText(InMaterialProperty, /* Short Name */ true);
+	const FText ToolTip = GetButtonToolTip(InMaterialProperty);
 
 	return SNew(SCheckBox)
 		.Style(FAppStyle::Get(), "DetailsView.SectionButton")
diff --git a/Engine/Plugins/Experimental/DynamicMaterial/Source/DynamicMaterialEditor/Private/UI/Widgets/Editor/PropertySelectorLayouts/SDMMaterialPropertySelector_WrapBase.cpp b/Engine/Plugins/Experimental/DynamicMaterial/Source/DynamicMaterialEditor/Private/UI/Widgets/Editor/PropertySelectorLayouts/SDMMaterialPropertySelector_WrapBase.cpp
--- a/Engine/Plugins/Experimental/DynamicMaterial/Source/DynamicMaterialEditor/Private/UI/Widgets/Editor/PropertySelectorLayouts/SDMMaterialPropertySelector_WrapBase.cpp
+++ b/Engine/Plugins/Experimental/DynamicMaterial/Source/DynamicMaterialEditor/Private/UI/Widgets/Editor/PropertySelectorLayouts/SDMMaterialPropertySelector_WrapBase.cpp
@@ -55,6 +55,22 @@ TSharedRef<SWidget> SDMMaterialPropertySelector_WrapBase::CreateSlot_PropertyLis
 			]
 		];
 
+	NewSlotList->AddSlot()
+		[
+			SNew(SHorizontalBox)
+			+ SHorizontalBox::Slot()
+			.AutoWidth()
+			[
+				SNew(SBox)
+				.WidthOverride(20.f)
+			]
+			+ SHorizontalBox::Slot()
+			.AutoWidth()
+			[
+				CreateSlot_SelectButton(EDMMaterialPropertyType::Any)
+			]
+		];
+
 	for (const TPair<EDMMaterialPropertyType, UDMMaterialProperty*>& PropertyPair : EditorOnlyData->GetMaterialProperties())
 	{
 		if (IsCustomMaterialProperty(PropertyPair.Key))
@@ -83,17 +99,8 @@ TSharedRef<SWidget> SDMMaterialPropertySelector_WrapBase::CreateSlot_PropertyLis
 
 TSharedRef<SWidget> SDMMaterialPropertySelector_WrapBase::CreateSlot_SelectButton(EDMMaterialPropertyType InMaterialProperty)
 {
-	UEnum* PropertyEnum = StaticEnum<EDMMaterialPropertyType>();
-
-	const FText ButtonText = InMaterialProperty == EDMMaterialPropertyType::None
-		? LOCTEXT("GlobalSettings", "Global Settings")
-		: UE::DynamicMaterialEditor::Private::GetMaterialPropertyShortDisplayName(InMaterialProperty);
-
-	const FText Format = LOCTEXT("PropertySelectFormat", "Edit the {0} property.");
-
-	const FText ToolTip = (InMaterialProperty == EDMMaterialPropertyType::None)
-		? LOCTEXT("GeneralSettingsToolTip", "Edit the material global settings.")
-		: FText::Format(Format, PropertyEnum->GetDisplayNameTextByValue(static_cast<int64>(InMaterialProperty)));
+	const FText ButtonText = GetSelectButtonText(InMaterialProperty, /* Short Name */ true);
+	const FText ToolTip = GetButtonToolTip(InMaterialProperty);
 
 	return SNew(SCheckBox)
 		.Style(FAppStyle::Get(), "DetailsView.SectionButton")
@@ -125,6 +132,7 @@ TSharedRef<SWidget> SDMMaterialPropertySelector_WrapBase::CreateSlot_SelectButto
 					SNew(STextBlock)
 					.Font(IDetailLayoutBuilder::GetDetailFont())
 					.Text(ButtonText)
+					.Justification(ETextJustify::Center)
 				]
 			]
 		];
diff --git a/Engine/Plugins/Experimental/DynamicMaterial/Source/DynamicMaterialEditor/Private/UI/Widgets/Editor/PropertySelectorLayouts/SDMMaterialPropertySelector_WrapSlim.cpp b/Engine/Plugins/Experimental/DynamicMaterial/Source/DynamicMaterialEditor/Private/UI/Widgets/Editor/PropertySelectorLayouts/SDMMaterialPropertySelector_WrapSlim.cpp
--- a/Engine/Plugins/Experimental/DynamicMaterial/Source/DynamicMaterialEditor/Private/UI/Widgets/Editor/PropertySelectorLayouts/SDMMaterialPropertySelector_WrapSlim.cpp
+++ b/Engine/Plugins/Experimental/DynamicMaterial/Source/DynamicMaterialEditor/Private/UI/Widgets/Editor/PropertySelectorLayouts/SDMMaterialPropertySelector_WrapSlim.cpp
@@ -36,6 +36,11 @@ TSharedRef<SWidget> SDMMaterialPropertySelector_WrapSlim::CreateSlot_PropertyLis
 			CreateSlot_SelectButton(EDMMaterialPropertyType::None)
 		];
 
+	NewSlotList->AddSlot()
+		[
+			CreateSlot_SelectButton(EDMMaterialPropertyType::Any)
+		];
+
 	for (const TPair<EDMMaterialPropertyType, UDMMaterialProperty*>& PropertyPair : EditorOnlyData->GetMaterialProperties())
 	{
 		if (!PropertyPair.Value || !PropertyPair.Value->IsEnabled())
diff --git a/Engine/Plugins/Experimental/DynamicMaterial/Source/DynamicMaterialEditor/Private/UI/Widgets/Editor/SDMMaterialPropertyPreviews.cpp b/Engine/Plugins/Experimental/DynamicMaterial/Source/DynamicMaterialEditor/Private/UI/Widgets/Editor/SDMMaterialPropertyPreviews.cpp
new file mode 100644
--- /dev/null
+++ b/Engine/Plugins/Experimental/DynamicMaterial/Source/DynamicMaterialEditor/Private/UI/Widgets/Editor/SDMMaterialPropertyPreviews.cpp
@@ -0,0 +1,343 @@
+// Copyright Epic Games, Inc. All Rights Reserved.
+
+#include "UI/Widgets/Editor/SDMMaterialPropertyPreviews.h"
+
+#include "Components/DMMaterialProperty.h"
+#include "Components/DMMaterialSlot.h"
+#include "DetailLayoutBuilder.h"
+#include "Model/DynamicMaterialModelEditorOnlyData.h"
+#include "UI/Widgets/Editor/SDMMaterialPropertySelector.h"
+#include "UI/Widgets/SDMMaterialEditor.h"
+#include "UI/Widgets/Visualizers/SDMMaterialComponentPreview.h"
+#include "Utils/DMPrivate.h"
+#include "Widgets/Input/SCheckBox.h"
+#include "Widgets/Layout/SBox.h"
+#include "Widgets/Layout/SWrapBox.h"
+#include "Widgets/SBoxPanel.h"
+#include "Widgets/SNullWidget.h"
+#include "Widgets/Text/STextBlock.h"
+
+#define LOCTEXT_NAMESPACE "SDMMaterialPropertyPreviews"
+
+void SDMMaterialPropertyPreviews::PrivateRegisterAttributes(FSlateAttributeDescriptor::FInitializer&)
+{
+}
+
+void SDMMaterialPropertyPreviews::Construct(const FArguments& InArgs, const TSharedRef<SDMMaterialEditor>& InEditorWidget)
+{
+	EditorWidgetWeak = InEditorWidget;
+
+	Content = TDMWidgetSlot<SWidget>(SharedThis(this), 0, CreateSlot_Content());
+}
+
+void SDMMaterialPropertyPreviews::Validate()
+{
+	if (Content.HasBeenInvalidated())
+	{
+		Content << CreateSlot_Content();
+	}
+}
+
+TSharedRef<SWidget> SDMMaterialPropertyPreviews::CreateSlot_Content()
+{
+	TSharedPtr<SDMMaterialEditor> EditorWidget = EditorWidgetWeak.Pin();
+
+	if (!EditorWidget.IsValid())
+	{
+		return SNullWidget::NullWidget;
+	}
+
+	UDynamicMaterialModel* MaterialModel = EditorWidget->GetMaterialModel();
+
+	if (!MaterialModel)
+	{
+		return SNullWidget::NullWidget;
+	}
+
+	UDynamicMaterialModelEditorOnlyData* EditorOnlyData = UDynamicMaterialModelEditorOnlyData::Get(MaterialModel);
+
+	if (!EditorOnlyData)
+	{
+		return SNullWidget::NullWidget;
+	}
+
+	TSharedRef<SWrapBox> WrapBox = SNew(SWrapBox)
+		.InnerSlotPadding(FVector2D(5.f, 5.f))
+		.UseAllottedSize(true);
+
+	// Active properties first.
+	UE::DynamicMaterial::ForEachMaterialPropertyType(
+		[this, &WrapBox, EditorOnlyData](EDMMaterialPropertyType InMaterialProperty)
+		{
+			if (UE::DynamicMaterialEditor::Private::IsCustomMaterialProperty(InMaterialProperty))
+			{
+				return EDMIterationResult::Continue;
+			}
+
+			if (UDMMaterialProperty* MaterialProperty = EditorOnlyData->GetMaterialProperty(InMaterialProperty))
+			{
+				UDMMaterialSlot* Slot = EditorOnlyData->GetSlotForEnabledMaterialProperty(InMaterialProperty);
+
+				const bool bIsActive = Slot && MaterialProperty->IsEnabled() && MaterialProperty->IsValidForModel(*EditorOnlyData);
+
+				if (bIsActive)
+				{
+					AddPropertyPreview(WrapBox, InMaterialProperty, Slot);
+				}
+			}
+
+			return EDMIterationResult::Continue;
+		}
+	);
+
+	// Now inactive properties
+	UE::DynamicMaterial::ForEachMaterialPropertyType(
+		[this, &WrapBox, EditorOnlyData](EDMMaterialPropertyType InMaterialProperty)
+		{
+			if (UE::DynamicMaterialEditor::Private::IsCustomMaterialProperty(InMaterialProperty))
+			{
+				return EDMIterationResult::Continue;
+			}
+
+			if (UDMMaterialProperty* MaterialProperty = EditorOnlyData->GetMaterialProperty(InMaterialProperty))
+			{
+				UDMMaterialSlot* Slot = EditorOnlyData->GetSlotForEnabledMaterialProperty(InMaterialProperty);
+
+				const bool bIsActive = Slot && MaterialProperty->IsEnabled();
+				const bool bIsValid = MaterialProperty->IsValidForModel(*EditorOnlyData);
+
+				if (!bIsActive && bIsValid)
+				{
+					AddPropertyPreview(WrapBox, InMaterialProperty, Slot);
+				}
+			}
+
+			return EDMIterationResult::Continue;
+		}
+	);
+
+	// Now invalid properties
+	UE::DynamicMaterial::ForEachMaterialPropertyType(
+		[this, &WrapBox, EditorOnlyData](EDMMaterialPropertyType InMaterialProperty)
+		{
+			if (UE::DynamicMaterialEditor::Private::IsCustomMaterialProperty(InMaterialProperty))
+			{
+				return EDMIterationResult::Continue;
+			}
+
+			if (UDMMaterialProperty* MaterialProperty = EditorOnlyData->GetMaterialProperty(InMaterialProperty))
+			{
+				UDMMaterialSlot* Slot = EditorOnlyData->GetSlotForEnabledMaterialProperty(InMaterialProperty);
+
+				const bool bIsActive = Slot && MaterialProperty->IsEnabled();
+				const bool bIsValid = MaterialProperty->IsValidForModel(*EditorOnlyData);
+
+				if (!bIsActive && !bIsValid)
+				{
+					AddPropertyPreview(WrapBox, InMaterialProperty, Slot);
+				}
+			}
+
+			return EDMIterationResult::Continue;
+		}
+	);
+
+	return SNew(SBox)
+		.Padding(5.f)
+		[
+			WrapBox
+		];
+}
+
+void SDMMaterialPropertyPreviews::AddPropertyPreview(const TSharedRef<SWrapBox>& InContainer, EDMMaterialPropertyType InMaterialProperty,
+	UDMMaterialSlot* InSlot)
+{
+	TSharedPtr<SDMMaterialEditor> EditorWidget = EditorWidgetWeak.Pin();
+
+	if (!EditorWidget.IsValid())
+	{
+		return;
+	}
+
+	TSharedPtr<SWidget> PreviewWidget;
+
+	if (InSlot)
+	{
+		PreviewWidget = SNew(SDMMaterialComponentPreview, EditorWidget.ToSharedRef(), InSlot)
+			.PreviewSize(FVector2D(60.f, 60.f));
+	}
+	else
+	{
+		PreviewWidget = SNew(SBox)
+			.WidthOverride(60.f)
+			.HeightOverride(60.f);
+	}
+
+	InContainer->AddSlot()
+		[
+			SNew(SVerticalBox)
+
+			+ SVerticalBox::Slot()
+			.AutoHeight()
+			[
+				PreviewWidget.ToSharedRef()
+			]
+
+			+ SVerticalBox::Slot()
+			.AutoHeight()
+			.Padding(0.f, 3.f, 0.f, 0.f)
+			[
+				SNew(SHorizontalBox)
+				+ SHorizontalBox::Slot()
+				.AutoWidth()
+				.VAlign(EVerticalAlignment::VAlign_Center)
+				[
+					CreateSlot_EnabledButton(InMaterialProperty)
+				]
+				+ SHorizontalBox::Slot()
+				.FillWidth(1.f)
+				.MaxWidth(40.f)
+				.VAlign(EVerticalAlignment::VAlign_Center)
+				.Padding(3.f, 0.f, 0.f, 0.f)
+				[
+					CreateSlot_PropertyName(InMaterialProperty)
+				]
+			]
+		];
+}
+
+TSharedRef<SWidget> SDMMaterialPropertyPreviews::CreateSlot_EnabledButton(EDMMaterialPropertyType InMaterialProperty)
+{
+	const FText Format = LOCTEXT("PropertyEnableFormat", "Toggle the {0} property.\n\nProperty must be valid for the Material Type.");
+	const FText ToolTip = FText::Format(Format, SDMMaterialPropertySelector::GetSelectButtonText(InMaterialProperty, /* Short Name */ false));
+
+	return SNew(SCheckBox)
+		.IsEnabled(this, &SDMMaterialPropertyPreviews::GetPropertyEnabledEnabled, InMaterialProperty)
+		.IsChecked(this, &SDMMaterialPropertyPreviews::GetPropertyEnabledState, InMaterialProperty)
+		.OnCheckStateChanged(this, &SDMMaterialPropertyPreviews::OnPropertyEnabledStateChanged, InMaterialProperty)
+		.ToolTipText(ToolTip);
+}
+
+TSharedRef<SWidget> SDMMaterialPropertyPreviews::CreateSlot_PropertyName(EDMMaterialPropertyType InMaterialProperty)
+{
+	return SNew(STextBlock)
+		.Font(IDetailLayoutBuilder::GetDetailFont())
+		.Text(SDMMaterialPropertySelector::GetSelectButtonText(InMaterialProperty, /* Short Name */ true));
+}
+
+bool SDMMaterialPropertyPreviews::GetPropertyEnabledEnabled(EDMMaterialPropertyType InMaterialProperty) const
+{
+	TSharedPtr<SDMMaterialEditor> EditorWidget = EditorWidgetWeak.Pin();
+
+	if (!EditorWidget.IsValid())
+	{
+		return false;
+	}
+
+	UDynamicMaterialModel* MaterialModel = EditorWidget->GetMaterialModel();
+
+	if (!MaterialModel)
+	{
+		return false;
+	}
+
+	UDynamicMaterialModelEditorOnlyData* EditorOnlyData = UDynamicMaterialModelEditorOnlyData::Get(MaterialModel);
+
+	if (!EditorOnlyData)
+	{
+		return false;
+	}
+
+	UDMMaterialProperty* Property = EditorOnlyData->GetMaterialProperty(InMaterialProperty);
+
+	if (!Property)
+	{
+		return false;
+	}
+
+	return Property->IsValidForModel(*EditorOnlyData);
+}
+
+ECheckBoxState SDMMaterialPropertyPreviews::GetPropertyEnabledState(EDMMaterialPropertyType InMaterialProperty) const
+{
+	TSharedPtr<SDMMaterialEditor> EditorWidget = EditorWidgetWeak.Pin();
+
+	if (!EditorWidget.IsValid())
+	{
+		return ECheckBoxState::Unchecked;
+	}
+
+	UDynamicMaterialModel* MaterialModel = EditorWidget->GetMaterialModel();
+
+	if (!MaterialModel)
+	{
+		return ECheckBoxState::Unchecked;
+	}
+
+	UDynamicMaterialModelEditorOnlyData* EditorOnlyData = UDynamicMaterialModelEditorOnlyData::Get(MaterialModel);
+
+	if (!EditorOnlyData)
+	{
+		return ECheckBoxState::Unchecked;
+	}
+
+	UDMMaterialProperty* Property = EditorOnlyData->GetMaterialProperty(InMaterialProperty);
+
+	if (!Property)
+	{
+		return ECheckBoxState::Unchecked;
+	}
+
+	return (Property->IsEnabled() && EditorOnlyData->GetSlotForMaterialProperty(InMaterialProperty))
+		? ECheckBoxState::Checked
+		: ECheckBoxState::Unchecked;
+}
+
+void SDMMaterialPropertyPreviews::OnPropertyEnabledStateChanged(ECheckBoxState InState, EDMMaterialPropertyType InMaterialProperty)
+{
+	TSharedPtr<SDMMaterialEditor> EditorWidget = EditorWidgetWeak.Pin();
+
+	if (!EditorWidget.IsValid())
+	{
+		return;
+	}
+
+	UDynamicMaterialModel* MaterialModel = EditorWidget->GetMaterialModel();
+
+	if (!MaterialModel)
+	{
+		return;
+	}
+
+	UDynamicMaterialModelEditorOnlyData* EditorOnlyData = UDynamicMaterialModelEditorOnlyData::Get(MaterialModel);
+
+	if (!EditorOnlyData)
+	{
+		return;
+	}
+
+	UDMMaterialProperty* MaterialProperty = EditorOnlyData->GetMaterialProperty(InMaterialProperty);
+
+	if (!MaterialProperty)
+	{
+		return;
+	}
+
+	const bool bEnabled = InState == ECheckBoxState::Checked;
+
+	MaterialProperty->SetEnabled(bEnabled);
+
+	if (bEnabled)
+	{
+		UDMMaterialSlot* Slot = EditorOnlyData->GetSlotForMaterialProperty(InMaterialProperty);
+
+		if (!Slot)
+		{
+			EditorOnlyData->AddSlotForMaterialProperty(InMaterialProperty);
+		}
+	}
+
+	Content.Invalidate();
+}
+
+#undef LOCTEXT_NAMESPACE
diff --git a/Engine/Plugins/Experimental/DynamicMaterial/Source/DynamicMaterialEditor/Private/UI/Widgets/Editor/SDMMaterialPropertyPreviews.h b/Engine/Plugins/Experimental/DynamicMaterial/Source/DynamicMaterialEditor/Private/UI/Widgets/Editor/SDMMaterialPropertyPreviews.h
new file mode 100644
--- /dev/null
+++ b/Engine/Plugins/Experimental/DynamicMaterial/Source/DynamicMaterialEditor/Private/UI/Widgets/Editor/SDMMaterialPropertyPreviews.h
@@ -0,0 +1,46 @@
+// Copyright Epic Games, Inc. All Rights Reserved.
+
+#pragma once
+
+#include "Widgets/SCompoundWidget.h"
+
+#include "UI/Utils/DMWidgetSlot.h"
+
+class SDMMaterialEditor;
+class SWrapBox;
+class UDMMaterialProperty;
+class UDMMaterialSlot;
+enum class ECheckBoxState : uint8;
+enum class EDMMaterialPropertyType : uint8;
+
+class SDMMaterialPropertyPreviews : public SCompoundWidget
+{
+	SLATE_DECLARE_WIDGET(SDMMaterialPropertyPreviews, SCompoundWidget)
+
+	SLATE_BEGIN_ARGS(SDMMaterialPropertyPreviews) {}
+	SLATE_END_ARGS()
+
+public:
+	virtual ~SDMMaterialPropertyPreviews() override = default;
+
+	void Construct(const FArguments& InArgs, const TSharedRef<SDMMaterialEditor>& InEditorWidget);
+
+	void Validate();
+
+protected:
+	TWeakPtr<SDMMaterialEditor> EditorWidgetWeak;
+
+	TDMWidgetSlot<SWidget> Content;
+
+	TSharedRef<SWidget> CreateSlot_Content();
+
+	void AddPropertyPreview(const TSharedRef<SWrapBox>& InContainer, EDMMaterialPropertyType InMaterialProperty, UDMMaterialSlot* InSlot);
+
+	TSharedRef<SWidget> CreateSlot_EnabledButton(EDMMaterialPropertyType InMaterialProperty);
+
+	TSharedRef<SWidget> CreateSlot_PropertyName(EDMMaterialPropertyType InMaterialProperty);
+
+	bool GetPropertyEnabledEnabled(EDMMaterialPropertyType InMaterialProperty) const;
+	ECheckBoxState GetPropertyEnabledState(EDMMaterialPropertyType InMaterialProperty) const;
+	void OnPropertyEnabledStateChanged(ECheckBoxState InState, EDMMaterialPropertyType InMaterialProperty);
+};
diff --git a/Engine/Plugins/Experimental/DynamicMaterial/Source/DynamicMaterialEditor/Private/UI/Widgets/Editor/SDMMaterialPropertySelector.cpp b/Engine/Plugins/Experimental/DynamicMaterial/Source/DynamicMaterialEditor/Private/UI/Widgets/Editor/SDMMaterialPropertySelector.cpp
--- a/Engine/Plugins/Experimental/DynamicMaterial/Source/DynamicMaterialEditor/Private/UI/Widgets/Editor/SDMMaterialPropertySelector.cpp
+++ b/Engine/Plugins/Experimental/DynamicMaterial/Source/DynamicMaterialEditor/Private/UI/Widgets/Editor/SDMMaterialPropertySelector.cpp
@@ -65,7 +65,21 @@ void SDMMaterialPropertySelector::SetGlobalSettings()
 		return;
 	}
 
-	EditorWidget->EditSlot(nullptr);
+	EditorWidget->EditGlobalSettings();
+}
+
+void SDMMaterialPropertySelector::SetPropertyPreviews()
+{
+	SelectedProperty = EDMMaterialPropertyType::Any;
+
+	TSharedPtr<SDMMaterialEditor> EditorWidget = GetEditorWidget();
+
+	if (!EditorWidget.IsValid())
+	{
+		return;
+	}
+
+	EditorWidget->ShowPropertyPreviews();
 }
 
 void SDMMaterialPropertySelector::SetSelectedProperty(EDMMaterialPropertyType InMaterialProperty)
@@ -114,9 +128,7 @@ TSharedPtr<SDMMaterialSlotEditor> SDMMaterialPropertySelector::GetSlotEditorWidg
 TSharedRef<SWidget> SDMMaterialPropertySelector::CreateSlot_EnabledButton(EDMMaterialPropertyType InMaterialProperty)
 {
 	const FText Format = LOCTEXT("PropertyEnableFormat", "Toggle the {0} property.\n\nProperty must be valid for the Material Type.");
-	UEnum* MaterialPropertyEnum = StaticEnum<EDMMaterialPropertyType>();
-
-	const FText ToolTip = FText::Format(Format, MaterialPropertyEnum->GetDisplayNameTextByValue(static_cast<int64>(InMaterialProperty)));
+	const FText ToolTip = FText::Format(Format, GetSelectButtonText(InMaterialProperty, /* Short Name */ false));
 
 	return SNew(SCheckBox)
 		.IsEnabled(this, &SDMMaterialPropertySelector::GetPropertyEnabledEnabled, InMaterialProperty)
@@ -125,78 +137,43 @@ TSharedRef<SWidget> SDMMaterialPropertySelector::CreateSlot_EnabledButton(EDMMat
 		.ToolTipText(ToolTip);
 }
 
-TSharedRef<SWidget> SDMMaterialPropertySelector::CreateSlot_SelectButton(EDMMaterialPropertyType InMaterialProperty)
+FText SDMMaterialPropertySelector::GetSelectButtonText(EDMMaterialPropertyType InMaterialProperty, bool bInShortName)
 {
-	EDMMaterialEditorLayout Layout = EDMMaterialEditorLayout::First;
-
-	if (UDynamicMaterialEditorSettings* Settings = UDynamicMaterialEditorSettings::Get())
+	switch (InMaterialProperty)
 	{
-		Layout = Settings->Layout;
+		case EDMMaterialPropertyType::None:
+			return bInShortName
+				? LOCTEXT("GlobalSettingsShort", "Global")
+				: LOCTEXT("GlobalSettings", "Global Settings");
+
+		case EDMMaterialPropertyType::Any:
+			return bInShortName
+				? LOCTEXT("Properties", "Props")
+				: LOCTEXT("PropertyPreviews", "Properties");
+
+		default:
+			return bInShortName
+				? UE::DynamicMaterialEditor::Private::GetMaterialPropertyShortDisplayName(InMaterialProperty)
+				: UE::DynamicMaterialEditor::Private::GetMaterialPropertyLongDisplayName(InMaterialProperty);
 	}
+}
 
-	UEnum* PropertyEnum = StaticEnum<EDMMaterialPropertyType>();
-
-	auto UseShortName = [Layout]()
-		{
-			switch (Layout)
-			{
-				case EDMMaterialEditorLayout::LeftSlim:
-					return true;
-
-				default:
-					return false;
-			}
-		};
-
-	const bool bUseShortName = UseShortName();
-
-	const FText ButtonText = InMaterialProperty == EDMMaterialPropertyType::None
-		? LOCTEXT("GlobalSettings", "Global Settings")
-		: (bUseShortName 
-			? UE::DynamicMaterialEditor::Private::GetMaterialPropertyShortDisplayName(InMaterialProperty)
-			: PropertyEnum->GetDisplayNameTextByValue(static_cast<int64>(InMaterialProperty)));
-
-	const FText Format = LOCTEXT("PropertySelectFormat", "Edit the {0} property.");
-	UEnum* MaterialPropertyEnum = StaticEnum<EDMMaterialPropertyType>();
+FText SDMMaterialPropertySelector::GetButtonToolTip(EDMMaterialPropertyType InMaterialProperty)
+{
+	switch (InMaterialProperty)
+	{
+		case EDMMaterialPropertyType::None:
+			return LOCTEXT("GeneralSettingsToolTip", "Edit the Material Global Settings.");
 
-	const FText ToolTip = (InMaterialProperty == EDMMaterialPropertyType::None)
-		? LOCTEXT("GeneralSettingsToolTip", "Edit the material global settings.")
-		: FText::Format(Format, MaterialPropertyEnum->GetDisplayNameTextByValue(static_cast<int64>(InMaterialProperty)));
+		case EDMMaterialPropertyType::Any:
+			return LOCTEXT("PropertyPreviewsToolTip", "Preview and toggle all the Material Properties.");
 
-	return SNew(SCheckBox)
-		.Style(FAppStyle::Get(), "DetailsView.SectionButton")
-		.HAlign(EHorizontalAlignment::HAlign_Center)
-		.Padding(0.f)
-		.IsEnabled(this, &SDMMaterialPropertySelector::GetPropertySelectEnabled, InMaterialProperty)
-		.IsChecked(this, &SDMMaterialPropertySelector::GetPropertySelectState, InMaterialProperty)
-		.OnCheckStateChanged(this, &SDMMaterialPropertySelector::OnPropertySelectStateChanged, InMaterialProperty)
-		.ToolTipText(ToolTip)
-		.Content()
-		[
-			SNew(SBox)
-			.WidthOverride(bUseShortName ? 42.f : 135.f)
-			.Clipping(EWidgetClipping::ClipToBounds)
-			[
-				SNew(SHorizontalBox)
-				+SHorizontalBox::Slot()
-				.VAlign(VAlign_Center)
-				.AutoWidth()
-				[
-					SNew(SImage)
-					.Image(FAppStyle::Get().GetBrush("FilterBar.FilterImage"))
-					.ColorAndOpacity(this, &SDMMaterialPropertySelector::GetPropertySelectButtonChipColor, InMaterialProperty)
-				]
-				+SHorizontalBox::Slot()
-				.Padding(bUseShortName ? 2.f : 10.f, 6.f)
-				.VAlign(VAlign_Center)
-				.FillWidth(1.f)
-				[
-					SNew(STextBlock)
-					.Font(IDetailLayoutBuilder::GetDetailFont())
-					.Text(ButtonText)
-				]
-			]
-		];
+		default:
+		{
+			const FText Format = LOCTEXT("PropertySelectFormat", "Edit the {0} property.");
+			return FText::Format(Format, GetSelectButtonText(InMaterialProperty, /* Short Name */ false));
+		}
+	}
 }
 
 bool SDMMaterialPropertySelector::IsPropertyEnabled(EDMMaterialPropertyType InMaterialProperty) const
@@ -313,7 +290,7 @@ void SDMMaterialPropertySelector::OnPropertyEnabledStateChanged(ECheckBoxState I
 
 bool SDMMaterialPropertySelector::GetPropertySelectEnabled(EDMMaterialPropertyType InMaterialProperty) const
 {
-	if (InMaterialProperty == EDMMaterialPropertyType::None)
+	if (InMaterialProperty == EDMMaterialPropertyType::None || InMaterialProperty == EDMMaterialPropertyType::Any)
 	{
 		return true;
 	}
@@ -336,6 +313,12 @@ void SDMMaterialPropertySelector::OnPropertySelectStateChanged(ECheckBoxState In
 		return;
 	}
 
+	if (InMaterialProperty == EDMMaterialPropertyType::Any)
+	{
+		SetPropertyPreviews();
+		return;
+	}
+
 	UDynamicMaterialModelEditorOnlyData* EditorOnlyData = GetEditorOnlyData();
 
 	if (!EditorOnlyData)
@@ -360,7 +343,7 @@ void SDMMaterialPropertySelector::OnPropertySelectStateChanged(ECheckBoxState In
 
 FSlateColor SDMMaterialPropertySelector::GetPropertySelectButtonChipColor(EDMMaterialPropertyType InMaterialProperty) const
 {
-	if (InMaterialProperty == EDMMaterialPropertyType::None)
+	if (InMaterialProperty == EDMMaterialPropertyType::None || InMaterialProperty == EDMMaterialPropertyType::Any)
 	{
 		return FStyleColors::AccentGreen;
 	}
diff --git a/Engine/Plugins/Experimental/DynamicMaterial/Source/DynamicMaterialEditor/Private/UI/Widgets/Editor/SDMMaterialPropertySelector.h b/Engine/Plugins/Experimental/DynamicMaterial/Source/DynamicMaterialEditor/Private/UI/Widgets/Editor/SDMMaterialPropertySelector.h
--- a/Engine/Plugins/Experimental/DynamicMaterial/Source/DynamicMaterialEditor/Private/UI/Widgets/Editor/SDMMaterialPropertySelector.h
+++ b/Engine/Plugins/Experimental/DynamicMaterial/Source/DynamicMaterialEditor/Private/UI/Widgets/Editor/SDMMaterialPropertySelector.h
@@ -33,6 +33,10 @@ class SDMMaterialPropertySelector : public SCompoundWidget
 	SLATE_END_ARGS()
 
 public:
+	static FText GetSelectButtonText(EDMMaterialPropertyType InMaterialProperty, bool bInShortName);
+
+	static FText GetButtonToolTip(EDMMaterialPropertyType InMaterialProperty);
+
 	virtual ~SDMMaterialPropertySelector() override = default;
 
 	void Construct(const FArguments& InArgs, const TSharedRef<SDMMaterialEditor> InEditorWidget);
@@ -43,6 +47,8 @@ public:
 
 	void SetGlobalSettings();
 
+	void SetPropertyPreviews();
+
 	void SetSelectedProperty(EDMMaterialPropertyType InMaterialProperty);
 
 protected:
diff --git a/Engine/Plugins/Experimental/DynamicMaterial/Source/DynamicMaterialEditor/Private/UI/Widgets/SDMMaterialEditor.cpp b/Engine/Plugins/Experimental/DynamicMaterial/Source/DynamicMaterialEditor/Private/UI/Widgets/SDMMaterialEditor.cpp
--- a/Engine/Plugins/Experimental/DynamicMaterial/Source/DynamicMaterialEditor/Private/UI/Widgets/SDMMaterialEditor.cpp
+++ b/Engine/Plugins/Experimental/DynamicMaterial/Source/DynamicMaterialEditor/Private/UI/Widgets/SDMMaterialEditor.cpp
@@ -19,6 +19,7 @@
 #include "UI/Widgets/Editor/SDMMaterialComponentEditor.h"
 #include "UI/Widgets/Editor/SDMMaterialGlobalSettingsEditor.h"
 #include "UI/Widgets/Editor/SDMMaterialPreview.h"
+#include "UI/Widgets/Editor/SDMMaterialPropertyPreviews.h"
 #include "UI/Widgets/Editor/SDMMaterialPropertySelector.h"
 #include "UI/Widgets/Editor/SDMMaterialSlotEditor.h"
 #include "UI/Widgets/Editor/SDMStatusBar.h"
@@ -39,7 +40,7 @@ void SDMMaterialEditor::PrivateRegisterAttributes(FSlateAttributeDescriptor::FIn
 SDMMaterialEditor::SDMMaterialEditor()
 	: CommandList(MakeShared<FUICommandList>())
 	, PreviewMaterialManager(MakeShared<FDMPreviewMaterialManager>())
-	, bGlobalSettingsMode(true)
+	, EditMode(EDMMaterialEditorMode::GlobalSettings)
 {
 }
 
@@ -51,7 +52,7 @@ SDMMaterialEditor::~SDMMaterialEditor()
 void SDMMaterialEditor::Construct(const FArguments& InArgs, const TSharedRef<SDMMaterialDesigner>& InDesignerWidget)
 {
 	DesignerWidgetWeak = InDesignerWidget;
-	bGlobalSettingsMode = true;
+	EditMode = EDMMaterialEditorMode::GlobalSettings;
 	PropertyToSelect.Reset();
 
 	SetCanTick(false);
@@ -147,9 +148,9 @@ AActor* SDMMaterialEditor::GetMaterialActor() const
 	return nullptr;
 }
 
-bool SDMMaterialEditor::IsEditingGlobalSettings() const
+EDMMaterialEditorMode SDMMaterialEditor::GetEditMode() const
 {
-	return bGlobalSettingsMode;
+	return EditMode;
 }
 
 void SDMMaterialEditor::SetMaterialActor(AActor* InActor)
@@ -215,7 +216,7 @@ void SDMMaterialEditor::EditSlot(UDMMaterialSlot* InSlot, bool bInForceRefresh)
 	ComponentEditorSlot.Invalidate();
 	ComponentToEdit.Reset();
 
-	bGlobalSettingsMode = !InSlot;
+	EditMode = InSlot ? EDMMaterialEditorMode::EditSlot : EDMMaterialEditorMode::GlobalSettings;
 
 	if (InSlot)
 	{
@@ -237,13 +238,15 @@ void SDMMaterialEditor::EditComponent(UDMMaterialComponent* InComponent, bool bI
 		return;
 	}
 
-	if (bGlobalSettingsMode)
+	if (EditMode != EDMMaterialEditorMode::EditSlot)
 	{
 		SlotEditorSlot.Invalidate();
 		SplitterSlot = nullptr;
+		GlobalSettingsEditorSlot.Invalidate();
+		MaterialPropertyPreviewsSlot.Invalidate();
 	}
 
-	bGlobalSettingsMode = false;
+	EditMode = EDMMaterialEditorMode::EditSlot;
 
 	ComponentEditorSlot.Invalidate();
 	ComponentToEdit = InComponent;
@@ -251,23 +254,44 @@ void SDMMaterialEditor::EditComponent(UDMMaterialComponent* InComponent, bool bI
 
 void SDMMaterialEditor::EditGlobalSettings(bool bInForceRefresh)
 {
-	if (bGlobalSettingsMode && !bInForceRefresh)
+	if (EditMode == EDMMaterialEditorMode::GlobalSettings && !bInForceRefresh)
 	{
 		return;
 	}
 
-	if (!bGlobalSettingsMode)
+	if (EditMode != EDMMaterialEditorMode::GlobalSettings)
 	{
 		SlotEditorSlot.Invalidate();
 		SplitterSlot = nullptr;
 		ComponentEditorSlot.Invalidate();
+		MaterialPropertyPreviewsSlot.Invalidate();
 	}
 
-	bGlobalSettingsMode = true;
+	EditMode = EDMMaterialEditorMode::GlobalSettings;
 
 	GlobalSettingsEditorSlot.Invalidate();
 }
 
+void SDMMaterialEditor::ShowPropertyPreviews(bool bInForceRefresh)
+{
+	if (EditMode == EDMMaterialEditorMode::PropertyPreviews && !bInForceRefresh)
+	{
+		return;
+	}
+
+	if (EditMode != EDMMaterialEditorMode::PropertyPreviews)
+	{
+		SlotEditorSlot.Invalidate();
+		SplitterSlot = nullptr;
+		ComponentEditorSlot.Invalidate();
+		GlobalSettingsEditorSlot.Invalidate();
+	}
+
+	EditMode = EDMMaterialEditorMode::PropertyPreviews;
+
+	MaterialPropertyPreviewsSlot.Invalidate();
+}
+
 void SDMMaterialEditor::Validate()
 {
 	if (!FDynamicMaterialModule::AreUObjectsSafe())
@@ -477,7 +501,7 @@ void SDMMaterialEditor::ValidateSlots()
 			PropertySelectorSlot << CreateSlot_PropertySelector();
 		}
 
-		if (bGlobalSettingsMode)
+		if (EditMode == EDMMaterialEditorMode::GlobalSettings)
 		{
 			if (GlobalSettingsEditorSlot.HasBeenInvalidated())
 			{
@@ -488,6 +512,17 @@ void SDMMaterialEditor::ValidateSlots()
 				GlobalSettingsEditorSlot->Validate();
 			}
 		}
+		else if (EditMode == EDMMaterialEditorMode::PropertyPreviews)
+		{
+			if (MaterialPropertyPreviewsSlot.HasBeenInvalidated())
+			{
+				MaterialPropertyPreviewsSlot << CreateSlot_MaterialPropertyPreviews();
+			}
+			else
+			{
+				MaterialPropertyPreviewsSlot->Validate();
+			}
+		}
 		else
 		{
 			if (SlotEditorSlot.HasBeenInvalidated())
@@ -588,6 +623,11 @@ TSharedRef<SDMMaterialGlobalSettingsEditor> SDMMaterialEditor::CreateSlot_Global
 	return SNew(SDMMaterialGlobalSettingsEditor, SharedThis(this), GetMaterialModelBase());
 }
 
+TSharedRef<SDMMaterialPropertyPreviews> SDMMaterialEditor::CreateSlot_MaterialPropertyPreviews()
+{
+	return SNew(SDMMaterialPropertyPreviews, SharedThis(this));
+}
+
 TSharedRef<SDMMaterialPreview> SDMMaterialEditor::CreateSlot_Preview()
 {
 	return SNew(SDMMaterialPreview, SharedThis(this), GetMaterialModelBase());
diff --git a/Engine/Plugins/Experimental/DynamicMaterial/Source/DynamicMaterialEditor/Private/UI/Widgets/SDMMaterialEditor.h b/Engine/Plugins/Experimental/DynamicMaterial/Source/DynamicMaterialEditor/Private/UI/Widgets/SDMMaterialEditor.h
--- a/Engine/Plugins/Experimental/DynamicMaterial/Source/DynamicMaterialEditor/Private/UI/Widgets/SDMMaterialEditor.h
+++ b/Engine/Plugins/Experimental/DynamicMaterial/Source/DynamicMaterialEditor/Private/UI/Widgets/SDMMaterialEditor.h
@@ -16,6 +16,7 @@ class SDMMaterialComponentEditor;
 class SDMMaterialDesigner;
 class SDMMaterialGlobalSettingsEditor;
 class SDMMaterialPreview;
+class SDMMaterialPropertyPreviews;
 class SDMMaterialPropertySelector;
 class SDMMaterialSlotEditor;
 class SDMStatusBar;
@@ -34,6 +35,13 @@ namespace UE::DynamicMaterialEditor::Private
 	inline const TCHAR* EditorLightBackground = TEXT("Brushes.Header");
 }
 
+enum class EDMMaterialEditorMode : uint8
+{
+	GlobalSettings,
+	PropertyPreviews,
+	EditSlot
+};
+
 class SDMMaterialEditor : public SCompoundWidget, public FSelfRegisteringEditorUndoClient
 {
 	SLATE_DECLARE_WIDGET(SDMMaterialEditor, SCompoundWidget)
@@ -67,7 +75,7 @@ public:
 	AActor* GetMaterialActor() const;
 
 	/** Widget Components */
-	bool IsEditingGlobalSettings() const;
+	EDMMaterialEditorMode GetEditMode() const;
 
 	const TSharedRef<FUICommandList>& GetCommandList() const;
 
@@ -86,6 +94,8 @@ public:
 
 	virtual void EditGlobalSettings(bool bInForceRefresh = false);
 
+	virtual void ShowPropertyPreviews(bool bInForceRefresh = false);
+
 	void Validate();
 
 	DECLARE_MULTICAST_DELEGATE_TwoParams(FOnEditedSlotChanged, const TSharedRef<SDMMaterialSlotEditor>&, UDMMaterialSlot*);
@@ -113,6 +123,7 @@ protected:
 	TDMWidgetSlot<SDMMaterialPreview> MaterialPreviewSlot;
 	TDMWidgetSlot<SDMMaterialPropertySelector> PropertySelectorSlot;
 	TDMWidgetSlot<SDMMaterialGlobalSettingsEditor> GlobalSettingsEditorSlot;
+	TDMWidgetSlot<SDMMaterialPropertyPreviews> MaterialPropertyPreviewsSlot;
 	FSlotBase* SplitterSlot;
 	TDMWidgetSlot<SDMMaterialSlotEditor> SlotEditorSlot;
 	TDMWidgetSlot<SDMMaterialComponentEditor> ComponentEditorSlot;
@@ -127,7 +138,7 @@ protected:
 	TOptional<EDMMaterialPropertyType> PropertyToSelect;
 	TWeakObjectPtr<UDMMaterialSlot> SlotToEdit;
 	TWeakObjectPtr<UDMMaterialComponent> ComponentToEdit;
-	bool bGlobalSettingsMode;
+	EDMMaterialEditorMode EditMode;
 
 	FOnEditedSlotChanged OnEditedSlotChanged;
 	FOnEditedComponentChanged OnEditedComponentChanged;
@@ -164,6 +175,8 @@ protected:
 
 	TSharedRef<SDMMaterialGlobalSettingsEditor> CreateSlot_GlobalSettingsEditor();
 
+	TSharedRef<SDMMaterialPropertyPreviews> CreateSlot_MaterialPropertyPreviews();
+
 	TSharedRef<SDMMaterialPreview> CreateSlot_Preview();
 
 	TSharedRef<SDMMaterialPropertySelector> CreateSlot_PropertySelector();
diff --git a/Engine/Plugins/Experimental/DynamicMaterial/Source/DynamicMaterialEditor/Private/Utils/DMPrivate.cpp b/Engine/Plugins/Experimental/DynamicMaterial/Source/DynamicMaterialEditor/Private/Utils/DMPrivate.cpp
--- a/Engine/Plugins/Experimental/DynamicMaterial/Source/DynamicMaterialEditor/Private/Utils/DMPrivate.cpp
+++ b/Engine/Plugins/Experimental/DynamicMaterial/Source/DynamicMaterialEditor/Private/Utils/DMPrivate.cpp
@@ -94,6 +94,11 @@ namespace UE::DynamicMaterialEditor::Private
 		}
 	}
 
+	FText GetMaterialPropertyLongDisplayName(EDMMaterialPropertyType InMaterialProperty)
+	{
+		return StaticEnum<EDMMaterialPropertyType>()->GetDisplayNameTextByValue(static_cast<int64>(InMaterialProperty));
+	}
+
 	FText GetMaterialPropertyShortDisplayName(EDMMaterialPropertyType InMaterialProperty)
 	{
 		constexpr const TCHAR* ShortNameName = TEXT("ShortName");
@@ -107,7 +112,7 @@ namespace UE::DynamicMaterialEditor::Private
 			return FText::FromString(ShortName);
 		}
 
-		return PropertyEnum->GetDisplayNameTextByValue(static_cast<int64>(InMaterialProperty));
+		return GetMaterialPropertyLongDisplayName(InMaterialProperty);
 	}
 }
 
diff --git a/Engine/Plugins/Experimental/DynamicMaterial/Source/DynamicMaterialEditor/Private/Utils/DMPrivate.h b/Engine/Plugins/Experimental/DynamicMaterial/Source/DynamicMaterialEditor/Private/Utils/DMPrivate.h
--- a/Engine/Plugins/Experimental/DynamicMaterial/Source/DynamicMaterialEditor/Private/Utils/DMPrivate.h
+++ b/Engine/Plugins/Experimental/DynamicMaterial/Source/DynamicMaterialEditor/Private/Utils/DMPrivate.h
@@ -34,6 +34,8 @@ namespace UE::DynamicMaterialEditor::Private
 
 	void LogError(const FString& InMessage, bool bInToast = true, const UObject* InSource = nullptr);
 
+	FText GetMaterialPropertyLongDisplayName(EDMMaterialPropertyType InMaterialProperty);
+
 	FText GetMaterialPropertyShortDisplayName(EDMMaterialPropertyType InMaterialProperty);
 }
 
