Material Designer
- UI flow issues

diff --git a/Engine/Plugins/Experimental/DynamicMaterial/Source/DynamicMaterialEditor/Private/Components/DMMaterialProperty.cpp b/Engine/Plugins/Experimental/DynamicMaterial/Source/DynamicMaterialEditor/Private/Components/DMMaterialProperty.cpp
--- a/Engine/Plugins/Experimental/DynamicMaterial/Source/DynamicMaterialEditor/Private/Components/DMMaterialProperty.cpp
+++ b/Engine/Plugins/Experimental/DynamicMaterial/Source/DynamicMaterialEditor/Private/Components/DMMaterialProperty.cpp
@@ -617,7 +617,7 @@ void UDMMaterialProperty::Update(UDMMaterialComponent* InSource, EDMUpdateType I
 		UDynamicMaterialModelEditorOnlyData* ModelEditorOnlyData = GetMaterialModelEditorOnlyData();
 		check(ModelEditorOnlyData);
 
-		ModelEditorOnlyData->RequestMaterialBuild();
+		ModelEditorOnlyData->OnPropertyUpdate(this);
 	}
 }
 
diff --git a/Engine/Plugins/Experimental/DynamicMaterial/Source/DynamicMaterialEditor/Private/Model/DynamicMaterialModelEditorOnlyData.cpp b/Engine/Plugins/Experimental/DynamicMaterial/Source/DynamicMaterialEditor/Private/Model/DynamicMaterialModelEditorOnlyData.cpp
--- a/Engine/Plugins/Experimental/DynamicMaterial/Source/DynamicMaterialEditor/Private/Model/DynamicMaterialModelEditorOnlyData.cpp
+++ b/Engine/Plugins/Experimental/DynamicMaterial/Source/DynamicMaterialEditor/Private/Model/DynamicMaterialModelEditorOnlyData.cpp
@@ -654,6 +654,13 @@ void UDynamicMaterialModelEditorOnlyData::OnValueListUpdate()
 	OnValueListUpdateDelegate.Broadcast(MaterialModel);
 }
 
+void UDynamicMaterialModelEditorOnlyData::OnPropertyUpdate(UDMMaterialProperty* InProperty)
+{
+	RequestMaterialBuild();
+
+	OnPropertyUpdateDelegate.Broadcast(MaterialModel);
+}
+
 TSharedRef<FDMMaterialBuildState> UDynamicMaterialModelEditorOnlyData::CreateBuildState(UMaterial* InMaterialToBuild, bool bInDirtyAssets) const
 {
 	check(MaterialModel);
diff --git a/Engine/Plugins/Experimental/DynamicMaterial/Source/DynamicMaterialEditor/Private/UI/PropertyGenerators/DMMaterialModelPropertyRowGenerator.cpp b/Engine/Plugins/Experimental/DynamicMaterial/Source/DynamicMaterialEditor/Private/UI/PropertyGenerators/DMMaterialModelPropertyRowGenerator.cpp
--- a/Engine/Plugins/Experimental/DynamicMaterial/Source/DynamicMaterialEditor/Private/UI/PropertyGenerators/DMMaterialModelPropertyRowGenerator.cpp
+++ b/Engine/Plugins/Experimental/DynamicMaterial/Source/DynamicMaterialEditor/Private/UI/PropertyGenerators/DMMaterialModelPropertyRowGenerator.cpp
@@ -120,7 +120,7 @@ void FDMMaterialModelPropertyRowGenerator::AddGlobalMaterialParameterValue(EDMMa
 						AlphaValue,
 						FText::Format(
 							LOCTEXT("PropertyFormat", "Global {0}"),
-							UE::DynamicMaterialEditor::Private::GetMaterialPropertyShortDisplayName(InProperty)
+							UE::DynamicMaterialEditor::Private::GetMaterialPropertyLongDisplayName(InProperty)
 						)
 					);
 				}
diff --git a/Engine/Plugins/Experimental/DynamicMaterial/Source/DynamicMaterialEditor/Private/UI/Utils/DMWidgetSlot.cpp b/Engine/Plugins/Experimental/DynamicMaterial/Source/DynamicMaterialEditor/Private/UI/Utils/DMWidgetSlot.cpp
--- a/Engine/Plugins/Experimental/DynamicMaterial/Source/DynamicMaterialEditor/Private/UI/Utils/DMWidgetSlot.cpp
+++ b/Engine/Plugins/Experimental/DynamicMaterial/Source/DynamicMaterialEditor/Private/UI/Utils/DMWidgetSlot.cpp
@@ -7,23 +7,48 @@
 #include "Widgets/SWidget.h"
 #include "Widgets/SNullWidget.h"
 
+FDMWidgetSlot::FDMWidgetSlot(FSlotBase* InSlot, const TSharedRef<SWidget>& InWidget)
+{
+	if (InSlot)
+	{
+		SetSlot(InSlot);
+	}
+
+	AssignWidget(InWidget);
+}
+
 FSlotBase* FDMWidgetSlot::GetSlot() const
 {
-	return Slot;
+	if (Owner.IsValid())
+	{
+		return Slot;
+	}
+
+	return nullptr;
 }
 
 void FDMWidgetSlot::SetSlot(FSlotBase* InSlot)
 {
-	if (Slot)
+	if (FSlotBase* ValidSlot = GetSlot())
 	{
-		Slot->DetachWidget();
+		ValidSlot->DetachWidget();
 	}
 
+	Owner.Reset();
+
 	Slot = InSlot;
 
-	if (Slot && Widget.IsValid())
+	if (InSlot)
 	{
-		Slot->AttachWidget(Widget.ToSharedRef());
+		if (SWidget* SlotOwner = Slot->GetOwnerWidget())
+		{
+			Owner = SlotOwner->AsShared();
+		}
+
+		if (Widget.IsValid())
+		{
+			Slot->AttachWidget(Widget.ToSharedRef());
+		}
 	}
 }
 
@@ -52,9 +77,9 @@ void FDMWidgetSlot::ClearWidget()
 	Widget.Reset();
 	bInvalidated = true;
 
-	if (Slot)
+	if (FSlotBase* ValidSlot = GetSlot())
 	{
-		Slot->DetachWidget();
+		ValidSlot->DetachWidget();
 	}
 }
 
@@ -73,9 +98,9 @@ void FDMWidgetSlot::AssignWidget(const TSharedRef<SWidget>& InWidget)
 	Widget = InWidget;
 	bInvalidated = InWidget == SNullWidget::NullWidget;
 
-	if (Slot)
+	if (FSlotBase* ValidSlot = GetSlot())
 	{
-		Slot->AttachWidget(InWidget);
+		ValidSlot->AttachWidget(InWidget);
 	}
 }
 
diff --git a/Engine/Plugins/Experimental/DynamicMaterial/Source/DynamicMaterialEditor/Private/UI/Utils/DMWidgetSlot.h b/Engine/Plugins/Experimental/DynamicMaterial/Source/DynamicMaterialEditor/Private/UI/Utils/DMWidgetSlot.h
--- a/Engine/Plugins/Experimental/DynamicMaterial/Source/DynamicMaterialEditor/Private/UI/Utils/DMWidgetSlot.h
+++ b/Engine/Plugins/Experimental/DynamicMaterial/Source/DynamicMaterialEditor/Private/UI/Utils/DMWidgetSlot.h
@@ -32,11 +32,13 @@ struct FDMWidgetSlot
 	bool operator==(const TSharedRef<SWidget>& InWidget) const;
 
 protected:
+	TWeakPtr<SWidget> Owner;
 	FSlotBase* Slot = nullptr;
-	TSharedPtr<SWidget> Widget = nullptr;
+	TSharedPtr<SWidget> Widget;
 	bool bInvalidated = true;
 
 	FDMWidgetSlot() = default;
+	FDMWidgetSlot(FSlotBase* InSlot, const TSharedRef<SWidget>& InWidget);
 
 	FSlotBase* FindSlot(const TSharedRef<SWidget>& InParentWidget, int32 InChildSlot) const;
 
@@ -51,17 +53,13 @@ struct TDMWidgetSlot : public FDMWidgetSlot
 	TDMWidgetSlot() = default;
 
 	TDMWidgetSlot(FSlotBase* InSlot, const TSharedRef<FWidgetType>& InWidget)
+		: FDMWidgetSlot(InSlot, InWidget)
 	{
-		Slot = InSlot;
-
-		AssignWidget(InWidget);
 	}
 
 	TDMWidgetSlot(const TSharedRef<SWidget>& InParentWidget, int32 InChildSlot, const TSharedRef<FWidgetType>& InWidget)
+		: FDMWidgetSlot(FindSlot(InParentWidget, InChildSlot), InWidget)
 	{
-		Slot = FindSlot(InParentWidget, InChildSlot);
-
-		AssignWidget(InWidget);
 	}
 
 	void operator<<(const TSharedRef<FWidgetType>& InWidget)
diff --git a/Engine/Plugins/Experimental/DynamicMaterial/Source/DynamicMaterialEditor/Private/UI/Widgets/Editor/PropertySelectorLayouts/SDMMaterialPropertySelector_Vertical.cpp b/Engine/Plugins/Experimental/DynamicMaterial/Source/DynamicMaterialEditor/Private/UI/Widgets/Editor/PropertySelectorLayouts/SDMMaterialPropertySelector_Vertical.cpp
--- a/Engine/Plugins/Experimental/DynamicMaterial/Source/DynamicMaterialEditor/Private/UI/Widgets/Editor/PropertySelectorLayouts/SDMMaterialPropertySelector_Vertical.cpp
+++ b/Engine/Plugins/Experimental/DynamicMaterial/Source/DynamicMaterialEditor/Private/UI/Widgets/Editor/PropertySelectorLayouts/SDMMaterialPropertySelector_Vertical.cpp
@@ -21,18 +21,19 @@ void SDMMaterialPropertySelector_Vertical::Construct(const FArguments& InArgs, c
 	);
 }
 
-TSharedRef<SWidget> SDMMaterialPropertySelector_Vertical::CreateSlot_SelectButton(EDMMaterialPropertyType InMaterialProperty)
+TSharedRef<SWidget> SDMMaterialPropertySelector_Vertical::CreateSlot_SelectButton(EDMMaterialEditorMode InEditMode, 
+	EDMMaterialPropertyType InMaterialProperty)
 {
-	const FText ButtonText = GetSelectButtonText(InMaterialProperty, /* Short Name */ false);
-	const FText ToolTip = GetButtonToolTip(InMaterialProperty);
+	const FText ButtonText = GetSelectButtonText(InEditMode, InMaterialProperty, /* Short Name */ false);
+	const FText ToolTip = GetButtonToolTip(InEditMode, InMaterialProperty);
 
 	return SNew(SCheckBox)
 		.Style(FAppStyle::Get(), "DetailsView.SectionButton")
 		.HAlign(EHorizontalAlignment::HAlign_Center)
 		.Padding(0.f)
-		.IsEnabled(this, &SDMMaterialPropertySelector_Vertical::GetPropertySelectEnabled, InMaterialProperty)
-		.IsChecked(this, &SDMMaterialPropertySelector_Vertical::GetPropertySelectState, InMaterialProperty)
-		.OnCheckStateChanged(this, &SDMMaterialPropertySelector_Vertical::OnPropertySelectStateChanged, InMaterialProperty)
+		.IsEnabled(this, &SDMMaterialPropertySelector_Vertical::GetPropertySelectEnabled, InEditMode, InMaterialProperty)
+		.IsChecked(this, &SDMMaterialPropertySelector_Vertical::GetPropertySelectState, InEditMode, InMaterialProperty)
+		.OnCheckStateChanged(this, &SDMMaterialPropertySelector_Vertical::OnPropertySelectStateChanged, InEditMode, InMaterialProperty)
 		.ToolTipText(ToolTip)
 		.Content()
 		[
@@ -47,7 +48,7 @@ TSharedRef<SWidget> SDMMaterialPropertySelector_Vertical::CreateSlot_SelectButto
 				[
 					SNew(SImage)
 					.Image(FAppStyle::Get().GetBrush("FilterBar.FilterImage"))
-					.ColorAndOpacity(this, &SDMMaterialPropertySelector_Vertical::GetPropertySelectButtonChipColor, InMaterialProperty)
+					.ColorAndOpacity(this, &SDMMaterialPropertySelector_Vertical::GetPropertySelectButtonChipColor, InEditMode, InMaterialProperty)
 				]
 				+SHorizontalBox::Slot()
 				.Padding(10.f, 6.f)
diff --git a/Engine/Plugins/Experimental/DynamicMaterial/Source/DynamicMaterialEditor/Private/UI/Widgets/Editor/PropertySelectorLayouts/SDMMaterialPropertySelector_Vertical.h b/Engine/Plugins/Experimental/DynamicMaterial/Source/DynamicMaterialEditor/Private/UI/Widgets/Editor/PropertySelectorLayouts/SDMMaterialPropertySelector_Vertical.h
--- a/Engine/Plugins/Experimental/DynamicMaterial/Source/DynamicMaterialEditor/Private/UI/Widgets/Editor/PropertySelectorLayouts/SDMMaterialPropertySelector_Vertical.h
+++ b/Engine/Plugins/Experimental/DynamicMaterial/Source/DynamicMaterialEditor/Private/UI/Widgets/Editor/PropertySelectorLayouts/SDMMaterialPropertySelector_Vertical.h
@@ -16,6 +16,6 @@ public:
 
 protected:
 	//~ Begin SDMMaterialPropertySelector
-	virtual TSharedRef<SWidget> CreateSlot_SelectButton(EDMMaterialPropertyType InMaterialProperty) override;
+	virtual TSharedRef<SWidget> CreateSlot_SelectButton(EDMMaterialEditorMode InEditMode, EDMMaterialPropertyType InMaterialProperty) override;
 	//~ End SDMMaterialPropertySelector
 };
diff --git a/Engine/Plugins/Experimental/DynamicMaterial/Source/DynamicMaterialEditor/Private/UI/Widgets/Editor/PropertySelectorLayouts/SDMMaterialPropertySelector_VerticalBase.cpp b/Engine/Plugins/Experimental/DynamicMaterial/Source/DynamicMaterialEditor/Private/UI/Widgets/Editor/PropertySelectorLayouts/SDMMaterialPropertySelector_VerticalBase.cpp
--- a/Engine/Plugins/Experimental/DynamicMaterial/Source/DynamicMaterialEditor/Private/UI/Widgets/Editor/PropertySelectorLayouts/SDMMaterialPropertySelector_VerticalBase.cpp
+++ b/Engine/Plugins/Experimental/DynamicMaterial/Source/DynamicMaterialEditor/Private/UI/Widgets/Editor/PropertySelectorLayouts/SDMMaterialPropertySelector_VerticalBase.cpp
@@ -4,6 +4,7 @@
 
 #include "DMDefs.h"
 #include "Model/DynamicMaterialModelEditorOnlyData.h"
+#include "UI/Widgets/SDMMaterialEditor.h"
 #include "Utils/DMPrivate.h"
 #include "Widgets/Layout/SBox.h"
 #include "Widgets/Layout/SGridPanel.h"
@@ -37,14 +38,14 @@ TSharedRef<SWidget> SDMMaterialPropertySelector_VerticalBase::CreateSlot_Propert
 
 	NewSlotList->AddSlot(PropertySelectorColumns::Select, Row)
 		[
-			CreateSlot_SelectButton(EDMMaterialPropertyType::None)
+			CreateSlot_SelectButton(EDMMaterialEditorMode::GlobalSettings, EDMMaterialPropertyType::None)
 		];
 
 	++Row;
 
 	NewSlotList->AddSlot(PropertySelectorColumns::Select, Row)
 		[
-			CreateSlot_SelectButton(EDMMaterialPropertyType::Any)
+			CreateSlot_SelectButton(EDMMaterialEditorMode::PropertyPreviews, EDMMaterialPropertyType::None)
 		];
 
 	++Row;
@@ -63,7 +64,7 @@ TSharedRef<SWidget> SDMMaterialPropertySelector_VerticalBase::CreateSlot_Propert
 
 		NewSlotList->AddSlot(PropertySelectorColumns::Select, Row)
 			[
-				CreateSlot_SelectButton(PropertyPair.Key)
+				CreateSlot_SelectButton(EDMMaterialEditorMode::EditSlot, PropertyPair.Key)
 			];
 
 		++Row;
diff --git a/Engine/Plugins/Experimental/DynamicMaterial/Source/DynamicMaterialEditor/Private/UI/Widgets/Editor/PropertySelectorLayouts/SDMMaterialPropertySelector_VerticalSlim.cpp b/Engine/Plugins/Experimental/DynamicMaterial/Source/DynamicMaterialEditor/Private/UI/Widgets/Editor/PropertySelectorLayouts/SDMMaterialPropertySelector_VerticalSlim.cpp
--- a/Engine/Plugins/Experimental/DynamicMaterial/Source/DynamicMaterialEditor/Private/UI/Widgets/Editor/PropertySelectorLayouts/SDMMaterialPropertySelector_VerticalSlim.cpp
+++ b/Engine/Plugins/Experimental/DynamicMaterial/Source/DynamicMaterialEditor/Private/UI/Widgets/Editor/PropertySelectorLayouts/SDMMaterialPropertySelector_VerticalSlim.cpp
@@ -22,18 +22,19 @@ void SDMMaterialPropertySelector_VerticalSlim::Construct(const FArguments& InArg
 	);
 }
 
-TSharedRef<SWidget> SDMMaterialPropertySelector_VerticalSlim::CreateSlot_SelectButton(EDMMaterialPropertyType InMaterialProperty)
+TSharedRef<SWidget> SDMMaterialPropertySelector_VerticalSlim::CreateSlot_SelectButton(EDMMaterialEditorMode InEditMode, 
+	EDMMaterialPropertyType InMaterialProperty)
 {
-	const FText ButtonText = GetSelectButtonText(InMaterialProperty, /* Short Name */ true);
-	const FText ToolTip = GetButtonToolTip(InMaterialProperty);
+	const FText ButtonText = GetSelectButtonText(InEditMode, InMaterialProperty, /* Short Name */ true);
+	const FText ToolTip = GetButtonToolTip(InEditMode, InMaterialProperty);
 
 	return SNew(SCheckBox)
 		.Style(FAppStyle::Get(), "DetailsView.SectionButton")
 		.HAlign(EHorizontalAlignment::HAlign_Center)
 		.Padding(0.f)
-		.IsEnabled(this, &SDMMaterialPropertySelector_VerticalSlim::GetPropertySelectEnabled, InMaterialProperty)
-		.IsChecked(this, &SDMMaterialPropertySelector_VerticalSlim::GetPropertySelectState, InMaterialProperty)
-		.OnCheckStateChanged(this, &SDMMaterialPropertySelector_VerticalSlim::OnPropertySelectStateChanged, InMaterialProperty)
+		.IsEnabled(this, &SDMMaterialPropertySelector_VerticalSlim::GetPropertySelectEnabled, InEditMode, InMaterialProperty)
+		.IsChecked(this, &SDMMaterialPropertySelector_VerticalSlim::GetPropertySelectState, InEditMode, InMaterialProperty)
+		.OnCheckStateChanged(this, &SDMMaterialPropertySelector_VerticalSlim::OnPropertySelectStateChanged, InEditMode, InMaterialProperty)
 		.ToolTipText(ToolTip)
 		.Content()
 		[
@@ -48,7 +49,7 @@ TSharedRef<SWidget> SDMMaterialPropertySelector_VerticalSlim::CreateSlot_SelectB
 				[
 					SNew(SImage)
 					.Image(FAppStyle::Get().GetBrush("FilterBar.FilterImage"))
-					.ColorAndOpacity(this, &SDMMaterialPropertySelector_VerticalSlim::GetPropertySelectButtonChipColor, InMaterialProperty)
+					.ColorAndOpacity(this, &SDMMaterialPropertySelector_VerticalSlim::GetPropertySelectButtonChipColor, InEditMode, InMaterialProperty)
 				]
 				+SHorizontalBox::Slot()
 				.Padding(2.f, 6.f)
diff --git a/Engine/Plugins/Experimental/DynamicMaterial/Source/DynamicMaterialEditor/Private/UI/Widgets/Editor/PropertySelectorLayouts/SDMMaterialPropertySelector_VerticalSlim.h b/Engine/Plugins/Experimental/DynamicMaterial/Source/DynamicMaterialEditor/Private/UI/Widgets/Editor/PropertySelectorLayouts/SDMMaterialPropertySelector_VerticalSlim.h
--- a/Engine/Plugins/Experimental/DynamicMaterial/Source/DynamicMaterialEditor/Private/UI/Widgets/Editor/PropertySelectorLayouts/SDMMaterialPropertySelector_VerticalSlim.h
+++ b/Engine/Plugins/Experimental/DynamicMaterial/Source/DynamicMaterialEditor/Private/UI/Widgets/Editor/PropertySelectorLayouts/SDMMaterialPropertySelector_VerticalSlim.h
@@ -16,6 +16,6 @@ public:
 
 protected:
 	//~ Begin SDMMaterialPropertySelector
-	virtual TSharedRef<SWidget> CreateSlot_SelectButton(EDMMaterialPropertyType InMaterialProperty) override;
+	virtual TSharedRef<SWidget> CreateSlot_SelectButton(EDMMaterialEditorMode InEditMode, EDMMaterialPropertyType InMaterialProperty) override;
 	//~ End SDMMaterialPropertySelector
 };
diff --git a/Engine/Plugins/Experimental/DynamicMaterial/Source/DynamicMaterialEditor/Private/UI/Widgets/Editor/PropertySelectorLayouts/SDMMaterialPropertySelector_WrapBase.cpp b/Engine/Plugins/Experimental/DynamicMaterial/Source/DynamicMaterialEditor/Private/UI/Widgets/Editor/PropertySelectorLayouts/SDMMaterialPropertySelector_WrapBase.cpp
--- a/Engine/Plugins/Experimental/DynamicMaterial/Source/DynamicMaterialEditor/Private/UI/Widgets/Editor/PropertySelectorLayouts/SDMMaterialPropertySelector_WrapBase.cpp
+++ b/Engine/Plugins/Experimental/DynamicMaterial/Source/DynamicMaterialEditor/Private/UI/Widgets/Editor/PropertySelectorLayouts/SDMMaterialPropertySelector_WrapBase.cpp
@@ -6,6 +6,7 @@
 #include "DetailLayoutBuilder.h"
 #include "DMDefs.h"
 #include "Model/DynamicMaterialModelEditorOnlyData.h"
+#include "UI/Widgets/SDMMaterialEditor.h"
 #include "Utils/DMPrivate.h"
 #include "Widgets/Images/SImage.h"
 #include "Widgets/Input/SCheckBox.h"
@@ -51,7 +52,7 @@ TSharedRef<SWidget> SDMMaterialPropertySelector_WrapBase::CreateSlot_PropertyLis
 			+ SHorizontalBox::Slot()
 			.AutoWidth()
 			[
-				CreateSlot_SelectButton(EDMMaterialPropertyType::None)
+				CreateSlot_SelectButton(EDMMaterialEditorMode::GlobalSettings, EDMMaterialPropertyType::None)
 			]
 		];
 
@@ -67,7 +68,7 @@ TSharedRef<SWidget> SDMMaterialPropertySelector_WrapBase::CreateSlot_PropertyLis
 			+ SHorizontalBox::Slot()
 			.AutoWidth()
 			[
-				CreateSlot_SelectButton(EDMMaterialPropertyType::Any)
+				CreateSlot_SelectButton(EDMMaterialEditorMode::PropertyPreviews, EDMMaterialPropertyType::None)
 			]
 		];
 
@@ -89,7 +90,7 @@ TSharedRef<SWidget> SDMMaterialPropertySelector_WrapBase::CreateSlot_PropertyLis
 				+ SHorizontalBox::Slot()
 				.AutoWidth()
 				[
-					CreateSlot_SelectButton(PropertyPair.Key)
+					CreateSlot_SelectButton(EDMMaterialEditorMode::EditSlot, PropertyPair.Key)
 				]
 			];
 	}
@@ -97,18 +98,19 @@ TSharedRef<SWidget> SDMMaterialPropertySelector_WrapBase::CreateSlot_PropertyLis
 	return NewSlotList;
 }
 
-TSharedRef<SWidget> SDMMaterialPropertySelector_WrapBase::CreateSlot_SelectButton(EDMMaterialPropertyType InMaterialProperty)
+TSharedRef<SWidget> SDMMaterialPropertySelector_WrapBase::CreateSlot_SelectButton(EDMMaterialEditorMode InEditMode, 
+	EDMMaterialPropertyType InMaterialProperty)
 {
-	const FText ButtonText = GetSelectButtonText(InMaterialProperty, /* Short Name */ true);
-	const FText ToolTip = GetButtonToolTip(InMaterialProperty);
+	const FText ButtonText = GetSelectButtonText(InEditMode, InMaterialProperty, /* Short Name */ true);
+	const FText ToolTip = GetButtonToolTip(InEditMode, InMaterialProperty);
 
 	return SNew(SCheckBox)
 		.Style(FAppStyle::Get(), "DetailsView.SectionButton")
 		.HAlign(EHorizontalAlignment::HAlign_Center)
 		.Padding(0.f)
-		.IsEnabled(this, &SDMMaterialPropertySelector_WrapBase::GetPropertySelectEnabled, InMaterialProperty)
-		.IsChecked(this, &SDMMaterialPropertySelector_WrapBase::GetPropertySelectState, InMaterialProperty)
-		.OnCheckStateChanged(this, &SDMMaterialPropertySelector_WrapBase::OnPropertySelectStateChanged, InMaterialProperty)
+		.IsEnabled(this, &SDMMaterialPropertySelector_WrapBase::GetPropertySelectEnabled, InEditMode, InMaterialProperty)
+		.IsChecked(this, &SDMMaterialPropertySelector_WrapBase::GetPropertySelectState, InEditMode, InMaterialProperty)
+		.OnCheckStateChanged(this, &SDMMaterialPropertySelector_WrapBase::OnPropertySelectStateChanged, InEditMode, InMaterialProperty)
 		.ToolTipText(ToolTip)
 		.Content()
 		[
@@ -119,7 +121,7 @@ TSharedRef<SWidget> SDMMaterialPropertySelector_WrapBase::CreateSlot_SelectButto
 			[
 				SNew(SImage)
 				.Image(FAppStyle::Get().GetBrush("FilterBar.FilterImage"))
-				.ColorAndOpacity(this, &SDMMaterialPropertySelector_WrapBase::GetPropertySelectButtonChipColor, InMaterialProperty)
+				.ColorAndOpacity(this, &SDMMaterialPropertySelector_WrapBase::GetPropertySelectButtonChipColor, InEditMode, InMaterialProperty)
 			]
 			+SHorizontalBox::Slot()
 			.Padding(4.f, 6.f)
diff --git a/Engine/Plugins/Experimental/DynamicMaterial/Source/DynamicMaterialEditor/Private/UI/Widgets/Editor/PropertySelectorLayouts/SDMMaterialPropertySelector_WrapBase.h b/Engine/Plugins/Experimental/DynamicMaterial/Source/DynamicMaterialEditor/Private/UI/Widgets/Editor/PropertySelectorLayouts/SDMMaterialPropertySelector_WrapBase.h
--- a/Engine/Plugins/Experimental/DynamicMaterial/Source/DynamicMaterialEditor/Private/UI/Widgets/Editor/PropertySelectorLayouts/SDMMaterialPropertySelector_WrapBase.h
+++ b/Engine/Plugins/Experimental/DynamicMaterial/Source/DynamicMaterialEditor/Private/UI/Widgets/Editor/PropertySelectorLayouts/SDMMaterialPropertySelector_WrapBase.h
@@ -17,6 +17,6 @@ public:
 protected:
 	//~ Begin SDMMaterialPropertySelector
 	virtual TSharedRef<SWidget> CreateSlot_PropertyList() override;
-	virtual TSharedRef<SWidget> CreateSlot_SelectButton(EDMMaterialPropertyType InMaterialProperty) override;
+	virtual TSharedRef<SWidget> CreateSlot_SelectButton(EDMMaterialEditorMode InEditMode, EDMMaterialPropertyType InMaterialProperty) override;
 	//~ End SDMMaterialPropertySelector
 };
diff --git a/Engine/Plugins/Experimental/DynamicMaterial/Source/DynamicMaterialEditor/Private/UI/Widgets/Editor/PropertySelectorLayouts/SDMMaterialPropertySelector_WrapSlim.cpp b/Engine/Plugins/Experimental/DynamicMaterial/Source/DynamicMaterialEditor/Private/UI/Widgets/Editor/PropertySelectorLayouts/SDMMaterialPropertySelector_WrapSlim.cpp
--- a/Engine/Plugins/Experimental/DynamicMaterial/Source/DynamicMaterialEditor/Private/UI/Widgets/Editor/PropertySelectorLayouts/SDMMaterialPropertySelector_WrapSlim.cpp
+++ b/Engine/Plugins/Experimental/DynamicMaterial/Source/DynamicMaterialEditor/Private/UI/Widgets/Editor/PropertySelectorLayouts/SDMMaterialPropertySelector_WrapSlim.cpp
@@ -4,6 +4,7 @@
 
 #include "Components/DMMaterialProperty.h"
 #include "Model/DynamicMaterialModelEditorOnlyData.h"
+#include "UI/Widgets/SDMMaterialEditor.h"
 #include "Widgets/Layout/SWrapBox.h"
 
 #define LOCTEXT_NAMESPACE "SDMMaterialPropertySelector_WrapSlim"
@@ -33,12 +34,12 @@ TSharedRef<SWidget> SDMMaterialPropertySelector_WrapSlim::CreateSlot_PropertyLis
 
 	NewSlotList->AddSlot()
 		[
-			CreateSlot_SelectButton(EDMMaterialPropertyType::None)
+			CreateSlot_SelectButton(EDMMaterialEditorMode::GlobalSettings, EDMMaterialPropertyType::None)
 		];
 
 	NewSlotList->AddSlot()
 		[
-			CreateSlot_SelectButton(EDMMaterialPropertyType::Any)
+			CreateSlot_SelectButton(EDMMaterialEditorMode::PropertyPreviews, EDMMaterialPropertyType::None)
 		];
 
 	for (const TPair<EDMMaterialPropertyType, UDMMaterialProperty*>& PropertyPair : EditorOnlyData->GetMaterialProperties())
@@ -55,7 +56,7 @@ TSharedRef<SWidget> SDMMaterialPropertySelector_WrapSlim::CreateSlot_PropertyLis
 
 		NewSlotList->AddSlot()
 			[
-				CreateSlot_SelectButton(PropertyPair.Key)
+				CreateSlot_SelectButton(EDMMaterialEditorMode::EditSlot, PropertyPair.Key)
 			];
 	}
 
diff --git a/Engine/Plugins/Experimental/DynamicMaterial/Source/DynamicMaterialEditor/Private/UI/Widgets/Editor/SDMMaterialPropertyPreviews.cpp b/Engine/Plugins/Experimental/DynamicMaterial/Source/DynamicMaterialEditor/Private/UI/Widgets/Editor/SDMMaterialPropertyPreviews.cpp
--- a/Engine/Plugins/Experimental/DynamicMaterial/Source/DynamicMaterialEditor/Private/UI/Widgets/Editor/SDMMaterialPropertyPreviews.cpp
+++ b/Engine/Plugins/Experimental/DynamicMaterial/Source/DynamicMaterialEditor/Private/UI/Widgets/Editor/SDMMaterialPropertyPreviews.cpp
@@ -165,6 +165,9 @@ void SDMMaterialPropertyPreviews::AddPropertyPreview(const TSharedRef<SWrapBox>&
 	{
 		PreviewWidget = SNew(SDMMaterialComponentPreview, EditorWidget.ToSharedRef(), InSlot)
 			.PreviewSize(FVector2D(60.f, 60.f));
+
+		PreviewWidget->SetCursor(EMouseCursor::Hand);
+		PreviewWidget->SetOnMouseButtonUp(FPointerEventHandler::CreateSP(this, &SDMMaterialPropertyPreviews::OnPreviewClicked, InMaterialProperty));
 	}
 	else
 	{
@@ -209,7 +212,7 @@ void SDMMaterialPropertyPreviews::AddPropertyPreview(const TSharedRef<SWrapBox>&
 TSharedRef<SWidget> SDMMaterialPropertyPreviews::CreateSlot_EnabledButton(EDMMaterialPropertyType InMaterialProperty)
 {
 	const FText Format = LOCTEXT("PropertyEnableFormat", "Toggle the {0} property.\n\nProperty must be valid for the Material Type.");
-	const FText ToolTip = FText::Format(Format, SDMMaterialPropertySelector::GetSelectButtonText(InMaterialProperty, /* Short Name */ false));
+	const FText ToolTip = FText::Format(Format, SDMMaterialPropertySelector::GetSelectButtonText(EDMMaterialEditorMode::EditSlot, InMaterialProperty, /* Short Name */ false));
 
 	return SNew(SCheckBox)
 		.IsEnabled(this, &SDMMaterialPropertyPreviews::GetPropertyEnabledEnabled, InMaterialProperty)
@@ -222,7 +225,8 @@ TSharedRef<SWidget> SDMMaterialPropertyPreviews::CreateSlot_PropertyName(EDMMate
 {
 	return SNew(STextBlock)
 		.Font(IDetailLayoutBuilder::GetDetailFont())
-		.Text(SDMMaterialPropertySelector::GetSelectButtonText(InMaterialProperty, /* Short Name */ true));
+		.Text(SDMMaterialPropertySelector::GetSelectButtonText(EDMMaterialEditorMode::EditSlot, InMaterialProperty, /* Short Name */ true))
+		.ToolTipText(SDMMaterialPropertySelector::GetSelectButtonText(EDMMaterialEditorMode::EditSlot, InMaterialProperty, /* Short Name */ false));
 }
 
 bool SDMMaterialPropertyPreviews::GetPropertyEnabledEnabled(EDMMaterialPropertyType InMaterialProperty) const
@@ -338,6 +342,45 @@ void SDMMaterialPropertyPreviews::OnPropertyEnabledStateChanged(ECheckBoxState I
 	}
 
 	Content.Invalidate();
+
+	// Make sure we go back to the property previews
+	EditorWidget->ShowPropertyPreviews();
+}
+
+FReply SDMMaterialPropertyPreviews::OnPreviewClicked(const FGeometry& InGeometry, const FPointerEvent& InPointerEvent, 
+	EDMMaterialPropertyType InMaterialProperty)
+{
+	TSharedPtr<SDMMaterialEditor> EditorWidget = EditorWidgetWeak.Pin();
+
+	if (!EditorWidget.IsValid())
+	{
+		return FReply::Handled();
+	}
+
+	UDynamicMaterialModel* MaterialModel = EditorWidget->GetMaterialModel();
+
+	if (!MaterialModel)
+	{
+		return FReply::Handled();
+	}
+
+	UDynamicMaterialModelEditorOnlyData* EditorOnlyData = UDynamicMaterialModelEditorOnlyData::Get(MaterialModel);
+
+	if (!EditorOnlyData)
+	{
+		return FReply::Handled();
+	}
+
+	UDMMaterialSlot* Slot = EditorOnlyData->GetSlotForMaterialProperty(InMaterialProperty);
+
+	if (!Slot)
+	{
+		return FReply::Handled();
+	}
+
+	EditorWidget->SelectProperty(InMaterialProperty);
+
+	return FReply::Handled();
 }
 
 #undef LOCTEXT_NAMESPACE
diff --git a/Engine/Plugins/Experimental/DynamicMaterial/Source/DynamicMaterialEditor/Private/UI/Widgets/Editor/SDMMaterialPropertyPreviews.h b/Engine/Plugins/Experimental/DynamicMaterial/Source/DynamicMaterialEditor/Private/UI/Widgets/Editor/SDMMaterialPropertyPreviews.h
--- a/Engine/Plugins/Experimental/DynamicMaterial/Source/DynamicMaterialEditor/Private/UI/Widgets/Editor/SDMMaterialPropertyPreviews.h
+++ b/Engine/Plugins/Experimental/DynamicMaterial/Source/DynamicMaterialEditor/Private/UI/Widgets/Editor/SDMMaterialPropertyPreviews.h
@@ -43,4 +43,6 @@ protected:
 	bool GetPropertyEnabledEnabled(EDMMaterialPropertyType InMaterialProperty) const;
 	ECheckBoxState GetPropertyEnabledState(EDMMaterialPropertyType InMaterialProperty) const;
 	void OnPropertyEnabledStateChanged(ECheckBoxState InState, EDMMaterialPropertyType InMaterialProperty);
+
+	FReply OnPreviewClicked(const FGeometry& InGeometry, const FPointerEvent& InPointerEvent, EDMMaterialPropertyType InMaterialProperty);
 };
diff --git a/Engine/Plugins/Experimental/DynamicMaterial/Source/DynamicMaterialEditor/Private/UI/Widgets/Editor/SDMMaterialPropertySelector.cpp b/Engine/Plugins/Experimental/DynamicMaterial/Source/DynamicMaterialEditor/Private/UI/Widgets/Editor/SDMMaterialPropertySelector.cpp
--- a/Engine/Plugins/Experimental/DynamicMaterial/Source/DynamicMaterialEditor/Private/UI/Widgets/Editor/SDMMaterialPropertySelector.cpp
+++ b/Engine/Plugins/Experimental/DynamicMaterial/Source/DynamicMaterialEditor/Private/UI/Widgets/Editor/SDMMaterialPropertySelector.cpp
@@ -3,7 +3,6 @@
 #include "UI/Widgets/Editor/SDMMaterialPropertySelector.h"
 
 #include "Components/DMMaterialProperty.h"
-#include "DetailLayoutBuilder.h"
 #include "DynamicMaterialEditorSettings.h"
 #include "Framework/Application/SlateApplication.h"
 #include "Model/DynamicMaterialModel.h"
@@ -29,7 +28,6 @@ void SDMMaterialPropertySelector::PrivateRegisterAttributes(FSlateAttributeDescr
 void SDMMaterialPropertySelector::Construct(const FArguments& InArgs, const TSharedRef<SDMMaterialEditor> InEditorWidget)
 {
 	EditorWidgetWeak = InEditorWidget;
-	SelectedProperty = EDMMaterialPropertyType::None;
 
 	SetCanTick(false);
 
@@ -49,15 +47,8 @@ TSharedPtr<SDMMaterialEditor> SDMMaterialPropertySelector::GetEditorWidget() con
 	return EditorWidgetWeak.Pin();
 }
 
-EDMMaterialPropertyType SDMMaterialPropertySelector::GetSelectedProperty() const
+void SDMMaterialPropertySelector::SetSelectedProperty(EDMMaterialEditorMode InEditMode, EDMMaterialPropertyType InMaterialProperty)
 {
-	return SelectedProperty;
-}
-
-void SDMMaterialPropertySelector::SetGlobalSettings()
-{
-	SelectedProperty = EDMMaterialPropertyType::None;
-
 	TSharedPtr<SDMMaterialEditor> EditorWidget = GetEditorWidget();
 
 	if (!EditorWidget.IsValid())
@@ -65,33 +56,23 @@ void SDMMaterialPropertySelector::SetGlobalSettings()
 		return;
 	}
 
-	EditorWidget->EditGlobalSettings();
-}
-
-void SDMMaterialPropertySelector::SetPropertyPreviews()
-{
-	SelectedProperty = EDMMaterialPropertyType::Any;
-
-	TSharedPtr<SDMMaterialEditor> EditorWidget = GetEditorWidget();
-
-	if (!EditorWidget.IsValid())
+	switch (InEditMode)
 	{
-		return;
-	}
-
-	EditorWidget->ShowPropertyPreviews();
-}
+	case EDMMaterialEditorMode::GlobalSettings:
+		EditorWidget->EditGlobalSettings();
+		break;
 
-void SDMMaterialPropertySelector::SetSelectedProperty(EDMMaterialPropertyType InMaterialProperty)
-{
-	if (SelectedProperty == InMaterialProperty)
-	{
-		return;
-	}
+	case EDMMaterialEditorMode::PropertyPreviews:
+		EditorWidget->ShowPropertyPreviews();
+		break;
 
-	SelectedProperty = InMaterialProperty;
+	case EDMMaterialEditorMode::EditSlot:
+		EditorWidget->SelectProperty(InMaterialProperty);
+		break;
 
-	OnSelectedPropertyChanged();
+	default:
+		break;
+	}
 }
 
 UDynamicMaterialModelEditorOnlyData* SDMMaterialPropertySelector::GetEditorOnlyData() const
@@ -128,7 +109,7 @@ TSharedPtr<SDMMaterialSlotEditor> SDMMaterialPropertySelector::GetSlotEditorWidg
 TSharedRef<SWidget> SDMMaterialPropertySelector::CreateSlot_EnabledButton(EDMMaterialPropertyType InMaterialProperty)
 {
 	const FText Format = LOCTEXT("PropertyEnableFormat", "Toggle the {0} property.\n\nProperty must be valid for the Material Type.");
-	const FText ToolTip = FText::Format(Format, GetSelectButtonText(InMaterialProperty, /* Short Name */ false));
+	const FText ToolTip = FText::Format(Format, GetSelectButtonText(EDMMaterialEditorMode::EditSlot, InMaterialProperty, /* Short Name */ false));
 
 	return SNew(SCheckBox)
 		.IsEnabled(this, &SDMMaterialPropertySelector::GetPropertyEnabledEnabled, InMaterialProperty)
@@ -137,42 +118,49 @@ TSharedRef<SWidget> SDMMaterialPropertySelector::CreateSlot_EnabledButton(EDMMat
 		.ToolTipText(ToolTip);
 }
 
-FText SDMMaterialPropertySelector::GetSelectButtonText(EDMMaterialPropertyType InMaterialProperty, bool bInShortName)
+FText SDMMaterialPropertySelector::GetSelectButtonText(EDMMaterialEditorMode InEditMode, EDMMaterialPropertyType InMaterialProperty, 
+	bool bInShortName)
 {
-	switch (InMaterialProperty)
+	switch (InEditMode)
 	{
-		case EDMMaterialPropertyType::None:
+		case EDMMaterialEditorMode::GlobalSettings:
 			return bInShortName
 				? LOCTEXT("GlobalSettingsShort", "Global")
 				: LOCTEXT("GlobalSettings", "Global Settings");
 
-		case EDMMaterialPropertyType::Any:
+		case EDMMaterialEditorMode::PropertyPreviews:
 			return bInShortName
 				? LOCTEXT("Properties", "Props")
 				: LOCTEXT("PropertyPreviews", "Properties");
 
-		default:
+		case EDMMaterialEditorMode::EditSlot:
 			return bInShortName
 				? UE::DynamicMaterialEditor::Private::GetMaterialPropertyShortDisplayName(InMaterialProperty)
 				: UE::DynamicMaterialEditor::Private::GetMaterialPropertyLongDisplayName(InMaterialProperty);
+
+		default:
+			return FText::GetEmpty();
 	}
 }
 
-FText SDMMaterialPropertySelector::GetButtonToolTip(EDMMaterialPropertyType InMaterialProperty)
+FText SDMMaterialPropertySelector::GetButtonToolTip(EDMMaterialEditorMode InEditMode, EDMMaterialPropertyType InMaterialProperty)
 {
-	switch (InMaterialProperty)
+	switch (InEditMode)
 	{
-		case EDMMaterialPropertyType::None:
+		case EDMMaterialEditorMode::GlobalSettings:
 			return LOCTEXT("GeneralSettingsToolTip", "Edit the Material Global Settings.");
 
-		case EDMMaterialPropertyType::Any:
+		case EDMMaterialEditorMode::PropertyPreviews:
 			return LOCTEXT("PropertyPreviewsToolTip", "Preview and toggle all the Material Properties.");
 
-		default:
+		case EDMMaterialEditorMode::EditSlot:
 		{
 			const FText Format = LOCTEXT("PropertySelectFormat", "Edit the {0} property.");
-			return FText::Format(Format, GetSelectButtonText(InMaterialProperty, /* Short Name */ false));
+			return FText::Format(Format, GetSelectButtonText(InEditMode, InMaterialProperty, /* Short Name */ false));
 		}
+
+		default:
+			return FText::GetEmpty();
 	}
 }
 
@@ -188,7 +176,7 @@ bool SDMMaterialPropertySelector::IsPropertyEnabled(EDMMaterialPropertyType InMa
 	return !!EditorOnlyData->GetMaterialProperty(InMaterialProperty);
 }
 
-bool SDMMaterialPropertySelector::DoesPropertySlotExist(EDMMaterialPropertyType InMaterialProperty) const
+bool SDMMaterialPropertySelector::SetPropertyEnabled(EDMMaterialPropertyType InMaterialProperty, bool bInEnabled)
 {
 	UDynamicMaterialModelEditorOnlyData* EditorOnlyData = GetEditorOnlyData();
 
@@ -204,15 +192,31 @@ bool SDMMaterialPropertySelector::DoesPropertySlotExist(EDMMaterialPropertyType
 		return false;
 	}
 
-	if (!MaterialProperty->IsEnabled())
+	MaterialProperty->SetEnabled(bInEnabled);
+
+	if (TSharedPtr<SDMMaterialEditor> EditorWidget = GetEditorWidget())
 	{
-		return false;
+		if (InMaterialProperty == EditorWidget->GetSelectedPropertyType())
+		{
+			SetSelectedProperty(EDMMaterialEditorMode::GlobalSettings, EDMMaterialPropertyType::None);
+		}
 	}
 
-	return !!EditorOnlyData->GetSlotForMaterialProperty(InMaterialProperty);
+	if (!bInEnabled)
+	{
+		return true;
+	}
+
+	if (EditorOnlyData->GetSlotForMaterialProperty(InMaterialProperty))
+	{
+		return true;
+	}
+
+	// Return true if the slot was successfully added
+	return !!EditorOnlyData->AddSlotForMaterialProperty(InMaterialProperty);
 }
 
-bool SDMMaterialPropertySelector::SetPropertyEnabled(EDMMaterialPropertyType InMaterialProperty, bool bInEnabled)
+bool SDMMaterialPropertySelector::DoesPropertySlotExist(EDMMaterialPropertyType InMaterialProperty) const
 {
 	UDynamicMaterialModelEditorOnlyData* EditorOnlyData = GetEditorOnlyData();
 
@@ -228,25 +232,12 @@ bool SDMMaterialPropertySelector::SetPropertyEnabled(EDMMaterialPropertyType InM
 		return false;
 	}
 
-	MaterialProperty->SetEnabled(bInEnabled);
-
-	if (InMaterialProperty == SelectedProperty)
-	{
-		SetGlobalSettings();
-	}
-
-	if (!bInEnabled)
-	{
-		return true;
-	}
-
-	if (EditorOnlyData->GetSlotForMaterialProperty(InMaterialProperty))
+	if (!MaterialProperty->IsEnabled())
 	{
-		return true;
+		return false;
 	}
 
-	// Return true if the slot was successfully added
-	return !!EditorOnlyData->AddSlotForMaterialProperty(InMaterialProperty);
+	return !!EditorOnlyData->GetSlotForMaterialProperty(InMaterialProperty);
 }
 
 bool SDMMaterialPropertySelector::GetPropertyEnabledEnabled(EDMMaterialPropertyType InMaterialProperty) const
@@ -283,108 +274,78 @@ void SDMMaterialPropertySelector::OnPropertyEnabledStateChanged(ECheckBoxState I
 	{
 		if (bSetEnabled)
 		{
-			SetSelectedProperty(InMaterialProperty);
+			SetSelectedProperty(EDMMaterialEditorMode::EditSlot, InMaterialProperty);
 		}
 	}
 }
 
-bool SDMMaterialPropertySelector::GetPropertySelectEnabled(EDMMaterialPropertyType InMaterialProperty) const
+bool SDMMaterialPropertySelector::GetPropertySelectEnabled(EDMMaterialEditorMode InEditMode, EDMMaterialPropertyType InMaterialProperty) const
 {
-	if (InMaterialProperty == EDMMaterialPropertyType::None || InMaterialProperty == EDMMaterialPropertyType::Any)
+	switch (InEditMode)
 	{
-		return true;
-	}
+		case EDMMaterialEditorMode::EditSlot:
+			return DoesPropertySlotExist(InMaterialProperty);
 
-	return DoesPropertySlotExist(InMaterialProperty);
-}
-
-ECheckBoxState SDMMaterialPropertySelector::GetPropertySelectState(EDMMaterialPropertyType InMaterialProperty) const
-{
-	return InMaterialProperty == SelectedProperty
-		? ECheckBoxState::Checked
-		: ECheckBoxState::Unchecked;
+		default:
+			return true;
+	}	
 }
 
-void SDMMaterialPropertySelector::OnPropertySelectStateChanged(ECheckBoxState InState, EDMMaterialPropertyType InMaterialProperty)
+ECheckBoxState SDMMaterialPropertySelector::GetPropertySelectState(EDMMaterialEditorMode InEditMode, EDMMaterialPropertyType 
+	InMaterialProperty) const
 {
-	if (InMaterialProperty == EDMMaterialPropertyType::None)
-	{
-		SetGlobalSettings();
-		return;
-	}
-
-	if (InMaterialProperty == EDMMaterialPropertyType::Any)
-	{
-		SetPropertyPreviews();
-		return;
-	}
-
-	UDynamicMaterialModelEditorOnlyData* EditorOnlyData = GetEditorOnlyData();
-
-	if (!EditorOnlyData)
-	{
-		return;
-	}
-
-	UDMMaterialProperty* MaterialProperty = EditorOnlyData->GetMaterialProperty(InMaterialProperty);
+	TSharedPtr<SDMMaterialEditor> EditorWidget = GetEditorWidget();
 
-	if (!MaterialProperty)
+	if (!EditorWidget.IsValid())
 	{
-		return;
+		return ECheckBoxState::Undetermined;
 	}
 
-	if (!MaterialProperty->IsEnabled())
+	switch (InEditMode)
 	{
-		return;
+		default:
+			return EditorWidget->GetEditMode() == InEditMode
+				? ECheckBoxState::Checked
+				: ECheckBoxState::Unchecked;
+
+		case EDMMaterialEditorMode::EditSlot:
+			return InMaterialProperty == EditorWidget->GetSelectedPropertyType()
+				? ECheckBoxState::Checked
+				: ECheckBoxState::Unchecked;
 	}
-
-	SetSelectedProperty(InMaterialProperty);
 }
 
-FSlateColor SDMMaterialPropertySelector::GetPropertySelectButtonChipColor(EDMMaterialPropertyType InMaterialProperty) const
+void SDMMaterialPropertySelector::OnPropertySelectStateChanged(ECheckBoxState InState, EDMMaterialEditorMode InEditMode, 
+	EDMMaterialPropertyType InMaterialProperty)
 {
-	if (InMaterialProperty == EDMMaterialPropertyType::None || InMaterialProperty == EDMMaterialPropertyType::Any)
+	if (InState != ECheckBoxState::Checked)
 	{
-		return FStyleColors::AccentGreen;
-	}
-
-	if (GetPropertySelectEnabled(InMaterialProperty))
-	{
-		return FStyleColors::Primary;
+		return;
 	}
 
-	return FStyleColors::Panel;
+	SetSelectedProperty(InEditMode, InMaterialProperty);
 }
 
-void SDMMaterialPropertySelector::OnSelectedPropertyChanged()
+FSlateColor SDMMaterialPropertySelector::GetPropertySelectButtonChipColor(EDMMaterialEditorMode InEditMode, 
+	EDMMaterialPropertyType InMaterialProperty) const
 {
-	TSharedPtr<SDMMaterialEditor> EditorWidget = GetEditorWidget();
-
-	if (!EditorWidget.IsValid())
+	switch (InEditMode)
 	{
-		return;
-	}
-
-	UDynamicMaterialModelEditorOnlyData* EditorOnlyData = GetEditorOnlyData();
-
-	if (!EditorOnlyData)
-	{
-		return;
-	}
+		case EDMMaterialEditorMode::GlobalSettings:
+		case EDMMaterialEditorMode::PropertyPreviews:
+			return FStyleColors::AccentGreen;
 
-	UDMMaterialSlot* Slot = EditorOnlyData->GetSlotForMaterialProperty(SelectedProperty);
+		case EDMMaterialEditorMode::EditSlot:
+			if (GetPropertySelectEnabled(EDMMaterialEditorMode::EditSlot, InMaterialProperty))
+			{
+				return FStyleColors::Primary;
+			}
 
-	if (!Slot)
-	{
-		Slot = EditorOnlyData->AddSlotForMaterialProperty(SelectedProperty);
+			// Fall through
 
-		if (!Slot)
-		{
-			return;
-		}
+		default:
+			return FStyleColors::Panel;
 	}
-
-	EditorWidget->EditSlot(Slot);
 }
 
 #undef LOCTEXT_NAMESPACE
diff --git a/Engine/Plugins/Experimental/DynamicMaterial/Source/DynamicMaterialEditor/Private/UI/Widgets/Editor/SDMMaterialPropertySelector.h b/Engine/Plugins/Experimental/DynamicMaterial/Source/DynamicMaterialEditor/Private/UI/Widgets/Editor/SDMMaterialPropertySelector.h
--- a/Engine/Plugins/Experimental/DynamicMaterial/Source/DynamicMaterialEditor/Private/UI/Widgets/Editor/SDMMaterialPropertySelector.h
+++ b/Engine/Plugins/Experimental/DynamicMaterial/Source/DynamicMaterialEditor/Private/UI/Widgets/Editor/SDMMaterialPropertySelector.h
@@ -14,6 +14,7 @@ class UDMMaterialProperty;
 class UDMMaterialSlot;
 class UDynamicMaterialModelEditorOnlyData;
 enum class ECheckBoxState : uint8;
+enum class EDMMaterialEditorMode : uint8;
 enum class EDMMaterialPropertyType : uint8;
 
 namespace UE::DynamicMaterialEditor::Private
@@ -33,9 +34,9 @@ class SDMMaterialPropertySelector : public SCompoundWidget
 	SLATE_END_ARGS()
 
 public:
-	static FText GetSelectButtonText(EDMMaterialPropertyType InMaterialProperty, bool bInShortName);
+	static FText GetSelectButtonText(EDMMaterialEditorMode InEditMode, EDMMaterialPropertyType InMaterialProperty, bool bInShortName);
 
-	static FText GetButtonToolTip(EDMMaterialPropertyType InMaterialProperty);
+	static FText GetButtonToolTip(EDMMaterialEditorMode InEditMode, EDMMaterialPropertyType InMaterialProperty);
 
 	virtual ~SDMMaterialPropertySelector() override = default;
 
@@ -43,17 +44,8 @@ public:
 
 	TSharedPtr<SDMMaterialEditor> GetEditorWidget() const;
 
-	EDMMaterialPropertyType GetSelectedProperty() const;
-
-	void SetGlobalSettings();
-
-	void SetPropertyPreviews();
-
-	void SetSelectedProperty(EDMMaterialPropertyType InMaterialProperty);
-
 protected:
 	TWeakPtr<SDMMaterialEditor> EditorWidgetWeak;
-	EDMMaterialPropertyType SelectedProperty;
 
 	UDynamicMaterialModelEditorOnlyData* GetEditorOnlyData() const;
 
@@ -63,7 +55,7 @@ protected:
 
 	TSharedRef<SWidget> CreateSlot_EnabledButton(EDMMaterialPropertyType InMaterialProperty);
 
-	virtual TSharedRef<SWidget> CreateSlot_SelectButton(EDMMaterialPropertyType InMaterialProperty) = 0;
+	virtual TSharedRef<SWidget> CreateSlot_SelectButton(EDMMaterialEditorMode InEditMode, EDMMaterialPropertyType InMaterialProperty) = 0;
 
 	bool IsPropertyEnabled(EDMMaterialPropertyType InMaterialProperty) const;
 
@@ -75,10 +67,10 @@ protected:
 	ECheckBoxState GetPropertyEnabledState(EDMMaterialPropertyType InMaterialProperty) const;
 	void OnPropertyEnabledStateChanged(ECheckBoxState InState, EDMMaterialPropertyType InMaterialProperty);
 
-	bool GetPropertySelectEnabled(EDMMaterialPropertyType InMaterialProperty) const;
-	ECheckBoxState GetPropertySelectState(EDMMaterialPropertyType InMaterialProperty) const;
-	void OnPropertySelectStateChanged(ECheckBoxState InState, EDMMaterialPropertyType InMaterialProperty);
-	FSlateColor GetPropertySelectButtonChipColor(EDMMaterialPropertyType InMaterialProperty) const;
+	bool GetPropertySelectEnabled(EDMMaterialEditorMode InEditMode, EDMMaterialPropertyType InMaterialProperty) const;
+	ECheckBoxState GetPropertySelectState(EDMMaterialEditorMode InEditMode, EDMMaterialPropertyType InMaterialProperty) const;
+	void OnPropertySelectStateChanged(ECheckBoxState InState, EDMMaterialEditorMode InEditMode, EDMMaterialPropertyType InMaterialProperty);
+	FSlateColor GetPropertySelectButtonChipColor(EDMMaterialEditorMode InEditMode, EDMMaterialPropertyType InMaterialProperty) const;
 
-	void OnSelectedPropertyChanged();
+	void SetSelectedProperty(EDMMaterialEditorMode InEditMode, EDMMaterialPropertyType InMaterialProperty);
 };
diff --git a/Engine/Plugins/Experimental/DynamicMaterial/Source/DynamicMaterialEditor/Private/UI/Widgets/Editor/SDMMaterialSlotEditor.cpp b/Engine/Plugins/Experimental/DynamicMaterial/Source/DynamicMaterialEditor/Private/UI/Widgets/Editor/SDMMaterialSlotEditor.cpp
--- a/Engine/Plugins/Experimental/DynamicMaterial/Source/DynamicMaterialEditor/Private/UI/Widgets/Editor/SDMMaterialSlotEditor.cpp
+++ b/Engine/Plugins/Experimental/DynamicMaterial/Source/DynamicMaterialEditor/Private/UI/Widgets/Editor/SDMMaterialSlotEditor.cpp
@@ -727,18 +727,31 @@ TSharedRef<SDMMaterialSlotLayerView> SDMMaterialSlotEditor::CreateSlot_LayerView
 
 	if (UDMMaterialSlot* Slot = GetSlot())
 	{
-		const TArray<UDMMaterialLayerObject*>& Layers = Slot->GetLayers();
-
-		if (!Layers.IsEmpty())
+		if (TSharedPtr<SDMMaterialEditor> EditorWidget = GetEditorWidget())
 		{
-			UDMMaterialLayerObject* LastLayer = Layers.Last();
-			NewLayerView->SetSelectedLayer(LastLayer);
+			bool bHasValidComponentToEdit = false;
 
-			if (UDMMaterialStage* Stage = LastLayer->GetFirstEnabledStage(EDMMaterialLayerStage::All))
+			if (UDMMaterialComponent* ComponentToEdit = EditorWidget->GetComponentToEdit())
 			{
-				if (TSharedPtr<SDMMaterialEditor> EditorWidget = GetEditorWidget())
+				if (ComponentToEdit->GetTypedParent<UDMMaterialSlot>(/* Allow Subclasses */ true) == Slot)
 				{
-					EditorWidget->EditComponent(Stage);
+					bHasValidComponentToEdit = true;
+				}
+			}
+
+			if (!bHasValidComponentToEdit)
+			{
+				const TArray<UDMMaterialLayerObject*>& Layers = Slot->GetLayers();
+
+				if (!Layers.IsEmpty())
+				{
+					UDMMaterialLayerObject* LastLayer = Layers.Last();
+					NewLayerView->SetSelectedLayer(LastLayer);
+
+					if (UDMMaterialStage* Stage = LastLayer->GetFirstEnabledStage(EDMMaterialLayerStage::All))
+					{
+						EditorWidget->EditComponent(Stage);
+					}
 				}
 			}
 		}
diff --git a/Engine/Plugins/Experimental/DynamicMaterial/Source/DynamicMaterialEditor/Private/UI/Widgets/SDMMaterialDesigner.cpp b/Engine/Plugins/Experimental/DynamicMaterial/Source/DynamicMaterialEditor/Private/UI/Widgets/SDMMaterialDesigner.cpp
--- a/Engine/Plugins/Experimental/DynamicMaterial/Source/DynamicMaterialEditor/Private/UI/Widgets/SDMMaterialDesigner.cpp
+++ b/Engine/Plugins/Experimental/DynamicMaterial/Source/DynamicMaterialEditor/Private/UI/Widgets/SDMMaterialDesigner.cpp
@@ -442,17 +442,16 @@ void SDMMaterialDesigner::OnLayoutChanged()
 
 	UDynamicMaterialModelBase* MaterialModelBase = CurrentEditor->GetMaterialModelBase();
 	const FDMObjectMaterialProperty* MaterialObjectProperty = CurrentEditor->GetMaterialObjectProperty();
-	UDMMaterialSlot* EditedSlot = nullptr;
+	EDMMaterialEditorMode EditorMode = CurrentEditor->GetEditMode();
+	EDMMaterialPropertyType SelectedProperty = CurrentEditor->GetSelectedPropertyType();
 	UDMMaterialComponent* EditedComponent = nullptr;
 
-	if (TSharedPtr<SDMMaterialSlotEditor> SlotEditorWidget = CurrentEditor->GetSlotEditorWidget())
+	if (EditorMode == EDMMaterialEditorMode::EditSlot)
 	{
-		EditedSlot = SlotEditorWidget->GetSlot();
-	}
-
-	if (TSharedPtr<SDMMaterialComponentEditor> ComponentEditorWidget = CurrentEditor->GetComponentEditorWidget())
-	{
-		EditedComponent = ComponentEditorWidget->GetComponent();
+		if (TSharedPtr<SDMMaterialComponentEditor> ComponentEditorWidget = CurrentEditor->GetComponentEditorWidget())
+		{
+			EditedComponent = ComponentEditorWidget->GetComponent();
+		}
 	}
 
 	if (MaterialObjectProperty)
@@ -472,14 +471,25 @@ void SDMMaterialDesigner::OnLayoutChanged()
 
 	TSharedRef<SDMMaterialEditor> NewEditor = StaticCastSharedRef<SDMMaterialEditor>(Content.ToSharedRef());
 
-	if (EditedSlot)
+	switch (EditorMode)
 	{
-		NewEditor->EditSlot(EditedSlot);
-	}
+		case EDMMaterialEditorMode::GlobalSettings:
+			NewEditor->EditGlobalSettings();
+			break;
 
-	if (EditedComponent)
-	{
-		NewEditor->EditComponent(EditedComponent);
+		case EDMMaterialEditorMode::PropertyPreviews:
+			NewEditor->ShowPropertyPreviews();
+			break;
+
+		case EDMMaterialEditorMode::EditSlot:
+			NewEditor->SelectProperty(SelectedProperty);
+
+			if (EditedComponent)
+			{
+				NewEditor->EditComponent(EditedComponent);
+			}
+
+			break;
 	}
 }
 
diff --git a/Engine/Plugins/Experimental/DynamicMaterial/Source/DynamicMaterialEditor/Private/UI/Widgets/SDMMaterialEditor.cpp b/Engine/Plugins/Experimental/DynamicMaterial/Source/DynamicMaterialEditor/Private/UI/Widgets/SDMMaterialEditor.cpp
--- a/Engine/Plugins/Experimental/DynamicMaterial/Source/DynamicMaterialEditor/Private/UI/Widgets/SDMMaterialEditor.cpp
+++ b/Engine/Plugins/Experimental/DynamicMaterial/Source/DynamicMaterialEditor/Private/UI/Widgets/SDMMaterialEditor.cpp
@@ -12,6 +12,7 @@
 #include "Engine/Texture.h"
 #include "Framework/Commands/GenericCommands.h"
 #include "Materials/MaterialFunctionInterface.h"
+#include "Model/DynamicMaterialModel.h"
 #include "Model/DynamicMaterialModelBase.h"
 #include "Model/DynamicMaterialModelDynamic.h"
 #include "Model/DynamicMaterialModelEditorOnlyData.h"
@@ -47,13 +48,23 @@ SDMMaterialEditor::SDMMaterialEditor()
 SDMMaterialEditor::~SDMMaterialEditor()
 {
 	FCoreDelegates::OnEnginePreExit.RemoveAll(this);
+
+	if (!FDynamicMaterialModule::AreUObjectsSafe())
+	{
+		return;
+	}
+
+	if (UDynamicMaterialModelEditorOnlyData* EditorOnlyData = EditorOnlyDataUpdateObject.Get())
+	{
+		EditorOnlyData->GetOnSlotListUpdateDelegate().RemoveAll(this);
+	}
 }
 
 void SDMMaterialEditor::Construct(const FArguments& InArgs, const TSharedRef<SDMMaterialDesigner>& InDesignerWidget)
 {
 	DesignerWidgetWeak = InDesignerWidget;
 	EditMode = EDMMaterialEditorMode::GlobalSettings;
-	PropertyToSelect.Reset();
+	SelectedMaterialProperty = EDMMaterialPropertyType::None;
 
 	SetCanTick(false);
 
@@ -97,6 +108,8 @@ void SDMMaterialEditor::SetMaterialModelBase(UDynamicMaterialModelBase* InMateri
 	EditGlobalSettings();
 
 	CreateLayout();
+
+	BindEditorOnlyDataUpdate(InMaterialModelBase);
 }
 
 UDynamicMaterialModel* SDMMaterialEditor::GetMaterialModel() const
@@ -136,6 +149,8 @@ void SDMMaterialEditor::SetObjectMaterialProperty(const FDMObjectMaterialPropert
 
 	ObjectMaterialPropertyOpt = InObjectProperty;
 	SetMaterialModelBase(MaterialModelBase);
+
+	BindEditorOnlyDataUpdate(MaterialModelBase);
 }
 
 AActor* SDMMaterialEditor::GetMaterialActor() const
@@ -175,21 +190,46 @@ TSharedPtr<SDMMaterialComponentEditor> SDMMaterialEditor::GetComponentEditorWidg
 	return &ComponentEditorSlot;
 }
 
+UDMMaterialSlot* SDMMaterialEditor::GetSlotToEdit() const
+{
+	return SlotToEdit.Get();
+}
+
+UDMMaterialComponent* SDMMaterialEditor::GetComponentToEdit() const
+{
+	return ComponentToEdit.Get();
+}
+
+EDMMaterialPropertyType SDMMaterialEditor::GetSelectedPropertyType() const
+{
+	return SelectedMaterialProperty;
+}
+
 void SDMMaterialEditor::SelectProperty(EDMMaterialPropertyType InProperty, bool bInForceRefresh)
 {
-	if (bInForceRefresh || !PropertySelectorSlot.IsValid())
+	if (EditMode == EDMMaterialEditorMode::EditSlot && SelectedMaterialProperty == InProperty && !bInForceRefresh)
 	{
-		PropertyToSelect = InProperty;
-		PropertySelectorSlot.Invalidate();
 		return;
 	}
 
-	if (PropertySelectorSlot->GetSelectedProperty() != InProperty)
+	EditMode = EDMMaterialEditorMode::EditSlot;
+	SelectedMaterialProperty = InProperty;
+
+	UDynamicMaterialModelEditorOnlyData* EditorOnlyData = UDynamicMaterialModelEditorOnlyData::Get(MaterialModelBaseWeak);
+
+	if (!EditorOnlyData)
 	{
-		PropertySelectorSlot->SetSelectedProperty(InProperty);
+		return;
 	}
-	
-	PropertyToSelect.Reset();
+
+	UDMMaterialSlot* Slot = EditorOnlyData->GetSlotForMaterialProperty(InProperty);
+
+	if (!Slot)
+	{
+		return;
+	}
+
+	EditSlot(Slot);
 }
 
 const TSharedRef<FUICommandList>& SDMMaterialEditor::GetCommandList() const
@@ -216,7 +256,7 @@ void SDMMaterialEditor::EditSlot(UDMMaterialSlot* InSlot, bool bInForceRefresh)
 	ComponentEditorSlot.Invalidate();
 	ComponentToEdit.Reset();
 
-	EditMode = InSlot ? EDMMaterialEditorMode::EditSlot : EDMMaterialEditorMode::GlobalSettings;
+	EditMode = EDMMaterialEditorMode::EditSlot;
 
 	if (InSlot)
 	{
@@ -268,6 +308,7 @@ void SDMMaterialEditor::EditGlobalSettings(bool bInForceRefresh)
 	}
 
 	EditMode = EDMMaterialEditorMode::GlobalSettings;
+	SelectedMaterialProperty = EDMMaterialPropertyType::None;
 
 	GlobalSettingsEditorSlot.Invalidate();
 }
@@ -288,6 +329,7 @@ void SDMMaterialEditor::ShowPropertyPreviews(bool bInForceRefresh)
 	}
 
 	EditMode = EDMMaterialEditorMode::PropertyPreviews;
+	SelectedMaterialProperty = EDMMaterialPropertyType::None;
 
 	MaterialPropertyPreviewsSlot.Invalidate();
 }
@@ -637,7 +679,7 @@ TSharedRef<SDMMaterialPropertySelector> SDMMaterialEditor::CreateSlot_PropertySe
 {
 	TSharedRef<SDMMaterialPropertySelector> NewPropertySelector = CreateSlot_PropertySelector_Impl();
 
-	if (!PropertyToSelect.IsSet())
+	if (EditMode == EDMMaterialEditorMode::EditSlot && SelectedMaterialProperty != EDMMaterialPropertyType::None)
 	{
 		if (UDynamicMaterialModel* MaterialModel = GetMaterialModel())
 		{
@@ -647,7 +689,7 @@ TSharedRef<SDMMaterialPropertySelector> SDMMaterialEditor::CreateSlot_PropertySe
 				{
 					if (PropertyPair.Value->IsEnabled() && PropertyPair.Value->IsValidForModel(*EditorOnlyData))
 					{
-						PropertyToSelect = PropertyPair.Key;
+						SelectedMaterialProperty = PropertyPair.Key;
 						break;
 					}
 				}
@@ -655,12 +697,6 @@ TSharedRef<SDMMaterialPropertySelector> SDMMaterialEditor::CreateSlot_PropertySe
 		}
 	}
 
-	if (PropertyToSelect.IsSet())
-	{
-		NewPropertySelector->SetSelectedProperty(PropertyToSelect.GetValue());
-		PropertyToSelect.Reset();
-	}
-
 	return NewPropertySelector;
 }
 
@@ -705,13 +741,17 @@ void SDMMaterialEditor::OnUndo()
 		return;
 	}
 
-	if (UDynamicMaterialModelEditorOnlyData* ModelEditorOnlyData = UDynamicMaterialModelEditorOnlyData::Get(MaterialModelBaseWeak))
+	if (EditMode == EDMMaterialEditorMode::EditSlot)
 	{
-		for (const TPair<EDMMaterialPropertyType, UDMMaterialProperty*>& PropertyPair : ModelEditorOnlyData->GetMaterialProperties())
+		if (UDynamicMaterialModelEditorOnlyData* EditorOnlyData = UDynamicMaterialModelEditorOnlyData::Get(MaterialModelBaseWeak))
 		{
-			if (PropertyPair.Value->IsEnabled())
+			for (const TPair<EDMMaterialPropertyType, UDMMaterialProperty*>& PropertyPair : EditorOnlyData->GetMaterialProperties())
 			{
-				PropertySelectorSlot->SetSelectedProperty(PropertyPair.Key);
+				if (PropertyPair.Value->IsEnabled() && PropertyPair.Value->IsValidForModel(*EditorOnlyData))
+				{
+					SelectProperty(PropertyPair.Key);
+					break;
+				}
 			}
 		}
 	}
@@ -734,4 +774,27 @@ void SDMMaterialEditor::OnEditorSplitterResized()
 	}
 }
 
+void SDMMaterialEditor::BindEditorOnlyDataUpdate(UDynamicMaterialModelBase* InMaterialModelBase)
+{
+	if (UDynamicMaterialModel* MaterialModel = Cast<UDynamicMaterialModel>(InMaterialModelBase))
+	{
+		if (UDynamicMaterialModelEditorOnlyData* EditorOnlyData = UDynamicMaterialModelEditorOnlyData::Get(MaterialModel))
+		{
+			EditorOnlyDataUpdateObject = EditorOnlyData;
+			EditorOnlyData->GetOnPropertyUpdateDelegate().AddSP(this, &SDMMaterialEditor::OnPropertyUpdate);
+			EditorOnlyData->GetOnSlotListUpdateDelegate().AddSP(this, &SDMMaterialEditor::OnSlotListUpdate);
+		}
+	}
+}
+
+void SDMMaterialEditor::OnPropertyUpdate(UDynamicMaterialModelBase* InMaterialModelBase)
+{
+	PropertySelectorSlot.Invalidate();
+}
+
+void SDMMaterialEditor::OnSlotListUpdate(UDynamicMaterialModelBase* InMaterialModelBase)
+{
+	PropertySelectorSlot.Invalidate();
+}
+
 #undef LOCTEXT_NAMESPACE
diff --git a/Engine/Plugins/Experimental/DynamicMaterial/Source/DynamicMaterialEditor/Private/UI/Widgets/SDMMaterialEditor.h b/Engine/Plugins/Experimental/DynamicMaterial/Source/DynamicMaterialEditor/Private/UI/Widgets/SDMMaterialEditor.h
--- a/Engine/Plugins/Experimental/DynamicMaterial/Source/DynamicMaterialEditor/Private/UI/Widgets/SDMMaterialEditor.h
+++ b/Engine/Plugins/Experimental/DynamicMaterial/Source/DynamicMaterialEditor/Private/UI/Widgets/SDMMaterialEditor.h
@@ -5,6 +5,8 @@
 #include "EditorUndoClient.h"
 #include "Widgets/SCompoundWidget.h"
 
+#include "Delegates/Delegate.h"
+#include "Delegates/DelegateCombinations.h"
 #include "DMObjectMaterialProperty.h"
 #include "Misc/Optional.h"
 #include "UI/Utils/DMWidgetSlot.h"
@@ -26,6 +28,7 @@ class UDMMaterialComponent;
 class UDMMaterialSlot;
 class UDMTextureSet;
 class UDynamicMaterialModelBase;
+class UDynamicMaterialModelEditorOnlyData;
 enum class EDMMaterialPropertyType : uint8;
 enum class EDMUpdateType : uint8;
 
@@ -77,6 +80,8 @@ public:
 	/** Widget Components */
 	EDMMaterialEditorMode GetEditMode() const;
 
+	EDMMaterialPropertyType GetSelectedPropertyType() const;
+
 	const TSharedRef<FUICommandList>& GetCommandList() const;
 
 	TSharedRef<FDMPreviewMaterialManager> GetPreviewMaterialManager() const;
@@ -85,6 +90,10 @@ public:
 
 	TSharedPtr<SDMMaterialComponentEditor> GetComponentEditorWidget() const;
 
+	UDMMaterialSlot* GetSlotToEdit() const;
+
+	UDMMaterialComponent* GetComponentToEdit() const;
+
 	/** Actions */
 	void SelectProperty(EDMMaterialPropertyType InProperty, bool bInForceRefresh = false);
 
@@ -135,14 +144,16 @@ protected:
 	TSharedRef<FUICommandList> CommandList;
 	TSharedRef<FDMPreviewMaterialManager> PreviewMaterialManager;
 
-	TOptional<EDMMaterialPropertyType> PropertyToSelect;
+	EDMMaterialEditorMode EditMode;
+	EDMMaterialPropertyType SelectedMaterialProperty;
 	TWeakObjectPtr<UDMMaterialSlot> SlotToEdit;
 	TWeakObjectPtr<UDMMaterialComponent> ComponentToEdit;
-	EDMMaterialEditorMode EditMode;
 
 	FOnEditedSlotChanged OnEditedSlotChanged;
 	FOnEditedComponentChanged OnEditedComponentChanged;
 
+	TWeakObjectPtr<UDynamicMaterialModelEditorOnlyData> EditorOnlyDataUpdateObject;
+
 	/** Operations */
 	void SetMaterialModelBase(UDynamicMaterialModelBase* InMaterialModelBase);
 
@@ -196,4 +207,10 @@ protected:
 	void OnEnginePreExit();
 
 	void OnEditorSplitterResized();
+
+	void BindEditorOnlyDataUpdate(UDynamicMaterialModelBase* InMaterialModelBase);
+
+	void OnPropertyUpdate(UDynamicMaterialModelBase* InMaterialModelBase);
+
+	void OnSlotListUpdate(UDynamicMaterialModelBase* InMaterialModelBase);
 };
diff --git a/Engine/Plugins/Experimental/DynamicMaterial/Source/DynamicMaterialEditor/Public/Model/DynamicMaterialModelEditorOnlyData.h b/Engine/Plugins/Experimental/DynamicMaterial/Source/DynamicMaterialEditor/Public/Model/DynamicMaterialModelEditorOnlyData.h
--- a/Engine/Plugins/Experimental/DynamicMaterial/Source/DynamicMaterialEditor/Public/Model/DynamicMaterialModelEditorOnlyData.h
+++ b/Engine/Plugins/Experimental/DynamicMaterial/Source/DynamicMaterialEditor/Public/Model/DynamicMaterialModelEditorOnlyData.h
@@ -34,6 +34,7 @@ struct FDMMaterialChannelListPreset;
 DECLARE_MULTICAST_DELEGATE_OneParam(FDMOnMaterialBuilt, UDynamicMaterialModelBase*);
 DECLARE_MULTICAST_DELEGATE_OneParam(FDMOnValueListUpdated, UDynamicMaterialModelBase*);
 DECLARE_MULTICAST_DELEGATE_OneParam(FDMOnSlotListUpdated, UDynamicMaterialModelBase*);
+DECLARE_MULTICAST_DELEGATE_OneParam(FDMOnPropertyUpdated, UDynamicMaterialModelBase*);
 
 UENUM(BlueprintType)
 enum class EDMState : uint8
@@ -206,6 +207,9 @@ public:
 	FDMOnMaterialBuilt::RegistrationType& GetOnMaterialBuiltDelegate() { return OnMaterialBuiltDelegate; }
 	FDMOnValueListUpdated::RegistrationType& GetOnValueListUpdateDelegate() { return OnValueListUpdateDelegate; }
 	FDMOnSlotListUpdated::RegistrationType& GetOnSlotListUpdateDelegate() { return OnSlotListUpdateDelegate; }
+	FDMOnPropertyUpdated::RegistrationType& GetOnPropertyUpdateDelegate() { return OnPropertyUpdateDelegate; }
+
+	void OnPropertyUpdate(UDMMaterialProperty* InProperty);
 
 	DYNAMICMATERIALEDITOR_API TSharedRef<FDMMaterialBuildState> CreateBuildState(UMaterial* InMaterialToBuild, bool bInDirtyAssets = true) const;
 
@@ -320,6 +324,7 @@ protected:
 	FDMOnMaterialBuilt OnMaterialBuiltDelegate;
 	FDMOnValueListUpdated OnValueListUpdateDelegate;
 	FDMOnSlotListUpdated OnSlotListUpdateDelegate;
+	FDMOnPropertyUpdated OnPropertyUpdateDelegate;
 
 	UPROPERTY(BlueprintReadOnly, Instanced, Category = "Material Designer")
 	TObjectPtr<UDMMaterialProperty> BaseColor;
