Material Designer
- Material statistics fix

diff --git a/Engine/Plugins/Experimental/DynamicMaterial/Source/DynamicMaterialEditor/Private/Model/DynamicMaterialModelEditorOnlyData.cpp b/Engine/Plugins/Experimental/DynamicMaterial/Source/DynamicMaterialEditor/Private/Model/DynamicMaterialModelEditorOnlyData.cpp
--- a/Engine/Plugins/Experimental/DynamicMaterial/Source/DynamicMaterialEditor/Private/Model/DynamicMaterialModelEditorOnlyData.cpp
+++ b/Engine/Plugins/Experimental/DynamicMaterial/Source/DynamicMaterialEditor/Private/Model/DynamicMaterialModelEditorOnlyData.cpp
@@ -492,7 +492,7 @@ void UDynamicMaterialModelEditorOnlyData::BuildMaterial(bool bInDirtyAssets)
 			break;
 	}
 
-	TSharedRef<FDMMaterialBuildState> BuildState = CreateBuildState(MaterialModel->DynamicMaterial, bInDirtyAssets);
+	TSharedPtr<FDMMaterialBuildState> BuildState = CreateBuildState(MaterialModel->DynamicMaterial, bInDirtyAssets);
 
 	/**
 	 * Process slots to build base material inputs.
@@ -509,12 +509,12 @@ void UDynamicMaterialModelEditorOnlyData::BuildMaterial(bool bInDirtyAssets)
 			{
 				if (Property->IsEnabled() && Property->IsMaterialPin() && GetSlotForMaterialProperty(InType))
 				{
-					Property->GenerateExpressions(BuildState);
+					Property->GenerateExpressions(BuildState.ToSharedRef());
 
 					// Global opacity is handled at later
 					if (InType != EDMMaterialPropertyType::Opacity && InType != EDMMaterialPropertyType::OpacityMask)
 					{
-						Property->AddAlphaMultiplier(BuildState);
+						Property->AddAlphaMultiplier(BuildState.ToSharedRef());
 					}
 				}
 			}
@@ -568,7 +568,7 @@ void UDynamicMaterialModelEditorOnlyData::BuildMaterial(bool bInDirtyAssets)
 				UMaterialExpression* OpacityOutputNode;
 				int32 OutputIndex;
 				int32 OutputChannel;
-				UDMMaterialProperty::GenerateOpacityExpressions(BuildState, OpacitySlot, OpacityProperty, OpacityOutputNode, OutputIndex, OutputChannel);
+				UDMMaterialProperty::GenerateOpacityExpressions(BuildState.ToSharedRef(), OpacitySlot, OpacityProperty, OpacityOutputNode, OutputIndex, OutputChannel);
 
 				if (OpacityOutputNode)
 				{
@@ -607,7 +607,7 @@ void UDynamicMaterialModelEditorOnlyData::BuildMaterial(bool bInDirtyAssets)
 
 		if (OpacityProperty != nullptr && OpacityProperty->IsEnabled())
 		{
-			OpacityProperty->AddAlphaMultiplier(BuildState);
+			OpacityProperty->AddAlphaMultiplier(BuildState.ToSharedRef());
 		}
 	}
 
@@ -621,7 +621,7 @@ void UDynamicMaterialModelEditorOnlyData::BuildMaterial(bool bInDirtyAssets)
 			{
 				if (Property->IsMaterialPin() && Property->IsEnabled() && GetSlotForMaterialProperty(InType))
 				{
-					Property->AddOutputProcessor(BuildState);
+					Property->AddOutputProcessor(BuildState.ToSharedRef());
 				}
 			}
 
@@ -629,7 +629,13 @@ void UDynamicMaterialModelEditorOnlyData::BuildMaterial(bool bInDirtyAssets)
 		}
 	);
 
-	//MaterialStats = UMaterialEditingLibrary::GetStatistics(MaterialModel->DynamicMaterial);
+	/**
+	 * To generate the statistics, you need to force a material recompile. The build state object does this in its destructor.
+	 * Resetting the build state SharedPtr destroys the object and thus generates the material shaders.
+	 */
+	BuildState.Reset();
+
+	MaterialStats = UMaterialEditingLibrary::GetStatistics(MaterialModel->DynamicMaterial);
 
 	State = EDMState::Idle;
 
