Material Designer
- Rejigged the new UI layouts a little.
- Resized the property select buttons.
- Added material preview tooltip/popout.

diff --git a/Engine/Plugins/Experimental/DynamicMaterial/Source/DynamicMaterial/Public/DMDefs.h b/Engine/Plugins/Experimental/DynamicMaterial/Source/DynamicMaterial/Public/DMDefs.h
--- a/Engine/Plugins/Experimental/DynamicMaterial/Source/DynamicMaterial/Public/DMDefs.h
+++ b/Engine/Plugins/Experimental/DynamicMaterial/Source/DynamicMaterial/Public/DMDefs.h
@@ -18,7 +18,7 @@ enum class EDMMaterialPropertyType : uint8
 	None = 0,
 	BaseColor           UMETA(ShortName = "Base"),
 	EmissiveColor       UMETA(ShortName = "Emis"),
-	Opacity,
+	Opacity             UMETA(ShortName = "Opac"),
 	OpacityMask         UMETA(ShortName = "Mask"),
 	Roughness           UMETA(ShortName = "Rough"),
 	Specular            UMETA(ShortName = "Spec"),
diff --git a/Engine/Plugins/Experimental/DynamicMaterial/Source/DynamicMaterialEditor/Private/UI/Widgets/Editor/PropertySelectorLayouts/SDMMaterialPropertySelector_Vertical.cpp b/Engine/Plugins/Experimental/DynamicMaterial/Source/DynamicMaterialEditor/Private/UI/Widgets/Editor/PropertySelectorLayouts/SDMMaterialPropertySelector_Vertical.cpp
--- a/Engine/Plugins/Experimental/DynamicMaterial/Source/DynamicMaterialEditor/Private/UI/Widgets/Editor/PropertySelectorLayouts/SDMMaterialPropertySelector_Vertical.cpp
+++ b/Engine/Plugins/Experimental/DynamicMaterial/Source/DynamicMaterialEditor/Private/UI/Widgets/Editor/PropertySelectorLayouts/SDMMaterialPropertySelector_Vertical.cpp
@@ -49,9 +49,10 @@ TSharedRef<SWidget> SDMMaterialPropertySelector_Vertical::CreateSlot_SelectButto
 					SNew(SImage)
 					.Image(FAppStyle::Get().GetBrush("FilterBar.FilterImage"))
 					.ColorAndOpacity(this, &SDMMaterialPropertySelector_Vertical::GetPropertySelectButtonChipColor, InEditMode, InMaterialProperty)
+						.DesiredSizeOverride(FVector2D(8, 17))
 				]
 				+SHorizontalBox::Slot()
-				.Padding(10.f, 6.f)
+				.Padding(5.f, 4.f)
 				.VAlign(VAlign_Center)
 				.FillWidth(1.f)
 				[
diff --git a/Engine/Plugins/Experimental/DynamicMaterial/Source/DynamicMaterialEditor/Private/UI/Widgets/Editor/PropertySelectorLayouts/SDMMaterialPropertySelector_VerticalBase.cpp b/Engine/Plugins/Experimental/DynamicMaterial/Source/DynamicMaterialEditor/Private/UI/Widgets/Editor/PropertySelectorLayouts/SDMMaterialPropertySelector_VerticalBase.cpp
--- a/Engine/Plugins/Experimental/DynamicMaterial/Source/DynamicMaterialEditor/Private/UI/Widgets/Editor/PropertySelectorLayouts/SDMMaterialPropertySelector_VerticalBase.cpp
+++ b/Engine/Plugins/Experimental/DynamicMaterial/Source/DynamicMaterialEditor/Private/UI/Widgets/Editor/PropertySelectorLayouts/SDMMaterialPropertySelector_VerticalBase.cpp
@@ -2,6 +2,7 @@
 
 #include "UI/Widgets/Editor/PropertySelectorLayouts/SDMMaterialPropertySelector_VerticalBase.h"
 
+#include "Components/DMMaterialProperty.h"
 #include "DMDefs.h"
 #include "Model/DynamicMaterialModelEditorOnlyData.h"
 #include "UI/Widgets/SDMMaterialEditor.h"
@@ -34,9 +35,12 @@ TSharedRef<SWidget> SDMMaterialPropertySelector_VerticalBase::CreateSlot_Propert
 		return NewSlotList;
 	}
 
+	const FMargin Padding = FMargin(0.f, 1.f);
+
 	int32 Row = 0;
 
 	NewSlotList->AddSlot(PropertySelectorColumns::Select, Row)
+		.Padding(Padding)
 		[
 			CreateSlot_SelectButton(EDMMaterialEditorMode::GlobalSettings, EDMMaterialPropertyType::None)
 		];
@@ -44,15 +48,91 @@ TSharedRef<SWidget> SDMMaterialPropertySelector_VerticalBase::CreateSlot_Propert
 	++Row;
 
 	NewSlotList->AddSlot(PropertySelectorColumns::Select, Row)
+		.Padding(Padding)
 		[
 			CreateSlot_SelectButton(EDMMaterialEditorMode::PropertyPreviews, EDMMaterialPropertyType::None)
 		];
 
 	++Row;
 
+	// Valid model properties
+	for (const TPair<EDMMaterialPropertyType, UDMMaterialProperty*>& PropertyPair : EditorOnlyData->GetMaterialProperties())
+	{
+		if (!PropertyPair.Value || IsCustomMaterialProperty(PropertyPair.Key))
+		{
+			continue;
+		}
+
+		// Skip invalid
+		if (!PropertyPair.Value->IsValidForModel(*EditorOnlyData))
+		{
+			continue;
+		}
+
+		// Skip inactive
+		if (!PropertyPair.Value->IsEnabled() || !EditorOnlyData->GetSlotForMaterialProperty(PropertyPair.Value->GetMaterialProperty()))
+		{
+			continue;
+		}
+
+		NewSlotList->AddSlot(PropertySelectorColumns::Enable, Row)
+			[
+				CreateSlot_EnabledButton(PropertyPair.Key)
+			];
+
+		NewSlotList->AddSlot(PropertySelectorColumns::Select, Row)
+			.Padding(Padding)
+			[
+				CreateSlot_SelectButton(EDMMaterialEditorMode::EditSlot, PropertyPair.Key)
+			];
+
+		++Row;
+	}
+
+	// Disabled model properties
 	for (const TPair<EDMMaterialPropertyType, UDMMaterialProperty*>& PropertyPair : EditorOnlyData->GetMaterialProperties())
 	{
-		if (IsCustomMaterialProperty(PropertyPair.Key))
+		if (!PropertyPair.Value || IsCustomMaterialProperty(PropertyPair.Key))
+		{
+			continue;
+		}
+
+		// Skip invalid
+		if (!PropertyPair.Value->IsValidForModel(*EditorOnlyData))
+		{
+			continue;
+		}
+
+		// Skip active
+		if (PropertyPair.Value->IsEnabled() && EditorOnlyData->GetSlotForMaterialProperty(PropertyPair.Value->GetMaterialProperty()))
+		{
+			continue;
+		}
+
+		NewSlotList->AddSlot(PropertySelectorColumns::Enable, Row)
+			[
+				CreateSlot_EnabledButton(PropertyPair.Key)
+			];
+
+		NewSlotList->AddSlot(PropertySelectorColumns::Select, Row)
+			.Padding(Padding)
+			[
+				CreateSlot_SelectButton(EDMMaterialEditorMode::EditSlot, PropertyPair.Key)
+			];
+
+		++Row;
+	}
+
+	// Invalid model properties
+	for (const TPair<EDMMaterialPropertyType, UDMMaterialProperty*>& PropertyPair : EditorOnlyData->GetMaterialProperties())
+	{
+		if (!PropertyPair.Value || IsCustomMaterialProperty(PropertyPair.Key))
+		{
+			continue;
+		}
+
+		// Skip valid
+		if (PropertyPair.Value->IsValidForModel(*EditorOnlyData))
 		{
 			continue;
 		}
@@ -63,6 +143,7 @@ TSharedRef<SWidget> SDMMaterialPropertySelector_VerticalBase::CreateSlot_Propert
 			];
 
 		NewSlotList->AddSlot(PropertySelectorColumns::Select, Row)
+			.Padding(Padding)
 			[
 				CreateSlot_SelectButton(EDMMaterialEditorMode::EditSlot, PropertyPair.Key)
 			];
diff --git a/Engine/Plugins/Experimental/DynamicMaterial/Source/DynamicMaterialEditor/Private/UI/Widgets/Editor/PropertySelectorLayouts/SDMMaterialPropertySelector_VerticalSlim.cpp b/Engine/Plugins/Experimental/DynamicMaterial/Source/DynamicMaterialEditor/Private/UI/Widgets/Editor/PropertySelectorLayouts/SDMMaterialPropertySelector_VerticalSlim.cpp
--- a/Engine/Plugins/Experimental/DynamicMaterial/Source/DynamicMaterialEditor/Private/UI/Widgets/Editor/PropertySelectorLayouts/SDMMaterialPropertySelector_VerticalSlim.cpp
+++ b/Engine/Plugins/Experimental/DynamicMaterial/Source/DynamicMaterialEditor/Private/UI/Widgets/Editor/PropertySelectorLayouts/SDMMaterialPropertySelector_VerticalSlim.cpp
@@ -2,13 +2,15 @@
 
 #include "UI/Widgets/Editor/PropertySelectorLayouts/SDMMaterialPropertySelector_VerticalSlim.h"
 
+#include "Components/DMMaterialProperty.h"
 #include "DetailLayoutBuilder.h"
 #include "DMDefs.h"
+#include "Model/DynamicMaterialModelEditorOnlyData.h"
+#include "UI/Widgets/SDMMaterialEditor.h"
 #include "Utils/DMPrivate.h"
 #include "Widgets/Images/SImage.h"
 #include "Widgets/Input/SCheckBox.h"
 #include "Widgets/Layout/SBox.h"
-#include "Widgets/Layout/SGridPanel.h"
 #include "Widgets/SBoxPanel.h"
 #include "Widgets/Text/STextBlock.h"
 
@@ -22,6 +24,126 @@ void SDMMaterialPropertySelector_VerticalSlim::Construct(const FArguments& InArg
 	);
 }
 
+TSharedRef<SWidget> SDMMaterialPropertySelector_VerticalSlim::CreateSlot_PropertyList()
+{
+	using namespace UE::DynamicMaterialEditor::Private;
+
+	TSharedRef<SVerticalBox> NewSlotList = SNew(SVerticalBox);
+
+	UDynamicMaterialModelEditorOnlyData* EditorOnlyData = GetEditorOnlyData();
+
+	if (!EditorOnlyData)
+	{
+		return NewSlotList;
+	}
+
+	const FMargin Padding = FMargin(0.f, 1.f);
+
+	TSharedRef<SWidget> PreviewButton = CreateSlot_SelectButton(EDMMaterialEditorMode::MaterialPreview, EDMMaterialPropertyType::None);
+	SetupMaterialPreviewButton(PreviewButton);
+
+	NewSlotList->AddSlot()
+		.AutoHeight()
+		.Padding(Padding)
+		[
+			PreviewButton
+		];
+
+	NewSlotList->AddSlot()
+		.AutoHeight()
+		.Padding(Padding)
+		[
+			CreateSlot_SelectButton(EDMMaterialEditorMode::GlobalSettings, EDMMaterialPropertyType::None)
+		];
+
+	NewSlotList->AddSlot()
+		.AutoHeight()
+		.Padding(Padding)
+		[
+			CreateSlot_SelectButton(EDMMaterialEditorMode::PropertyPreviews, EDMMaterialPropertyType::None)
+		];
+
+	// Valid model properties
+	for (const TPair<EDMMaterialPropertyType, UDMMaterialProperty*>& PropertyPair : EditorOnlyData->GetMaterialProperties())
+	{
+		if (!PropertyPair.Value || IsCustomMaterialProperty(PropertyPair.Key))
+		{
+			continue;
+		}
+
+		// Skip invalid
+		if (!PropertyPair.Value->IsValidForModel(*EditorOnlyData))
+		{
+			continue;
+		}
+
+		// Skip inactive
+		if (!PropertyPair.Value->IsEnabled() || !EditorOnlyData->GetSlotForMaterialProperty(PropertyPair.Value->GetMaterialProperty()))
+		{
+			continue;
+		}
+
+		NewSlotList->AddSlot()
+			.AutoHeight()
+			.Padding(Padding)
+			[
+				CreateSlot_SelectButton(EDMMaterialEditorMode::EditSlot, PropertyPair.Key)
+			];
+	}
+
+	// Disabled model properties
+	for (const TPair<EDMMaterialPropertyType, UDMMaterialProperty*>& PropertyPair : EditorOnlyData->GetMaterialProperties())
+	{
+		if (!PropertyPair.Value || IsCustomMaterialProperty(PropertyPair.Key))
+		{
+			continue;
+		}
+
+		// Skip invalid
+		if (!PropertyPair.Value->IsValidForModel(*EditorOnlyData))
+		{
+			continue;
+		}
+
+		// Skip active
+		if (PropertyPair.Value->IsEnabled() && EditorOnlyData->GetSlotForMaterialProperty(PropertyPair.Value->GetMaterialProperty()))
+		{
+			continue;
+		}
+
+		NewSlotList->AddSlot()
+			.AutoHeight()
+			.Padding(Padding)
+			[
+				CreateSlot_SelectButton(EDMMaterialEditorMode::EditSlot, PropertyPair.Key)
+			];
+	}
+
+	// Invalid model properties
+	for (const TPair<EDMMaterialPropertyType, UDMMaterialProperty*>& PropertyPair : EditorOnlyData->GetMaterialProperties())
+	{
+		if (!PropertyPair.Value || IsCustomMaterialProperty(PropertyPair.Key))
+		{
+			continue;
+		}
+
+		// Skip valid
+		if (PropertyPair.Value->IsValidForModel(*EditorOnlyData))
+		{
+			continue;
+		}
+
+		NewSlotList->AddSlot()
+			.AutoHeight()
+			.Padding(Padding)
+			[
+				CreateSlot_SelectButton(EDMMaterialEditorMode::EditSlot, PropertyPair.Key)
+			];
+	}
+
+	return NewSlotList;
+}
+
 TSharedRef<SWidget> SDMMaterialPropertySelector_VerticalSlim::CreateSlot_SelectButton(EDMMaterialEditorMode InEditMode, 
 	EDMMaterialPropertyType InMaterialProperty)
 {
@@ -39,7 +161,7 @@ TSharedRef<SWidget> SDMMaterialPropertySelector_VerticalSlim::CreateSlot_SelectB
 		.Content()
 		[
 			SNew(SBox)
-			.WidthOverride(42.f)
+			.WidthOverride(45.f)
 			.Clipping(EWidgetClipping::ClipToBounds)
 			[
 				SNew(SHorizontalBox)
@@ -50,9 +172,10 @@ TSharedRef<SWidget> SDMMaterialPropertySelector_VerticalSlim::CreateSlot_SelectB
 					SNew(SImage)
 					.Image(FAppStyle::Get().GetBrush("FilterBar.FilterImage"))
 					.ColorAndOpacity(this, &SDMMaterialPropertySelector_VerticalSlim::GetPropertySelectButtonChipColor, InEditMode, InMaterialProperty)
+					.DesiredSizeOverride(FVector2D(8, 17))
 				]
 				+SHorizontalBox::Slot()
-				.Padding(2.f, 6.f)
+				.Padding(2.f, 4.f)
 				.VAlign(VAlign_Center)
 				.FillWidth(1.f)
 				[
diff --git a/Engine/Plugins/Experimental/DynamicMaterial/Source/DynamicMaterialEditor/Private/UI/Widgets/Editor/PropertySelectorLayouts/SDMMaterialPropertySelector_VerticalSlim.h b/Engine/Plugins/Experimental/DynamicMaterial/Source/DynamicMaterialEditor/Private/UI/Widgets/Editor/PropertySelectorLayouts/SDMMaterialPropertySelector_VerticalSlim.h
--- a/Engine/Plugins/Experimental/DynamicMaterial/Source/DynamicMaterialEditor/Private/UI/Widgets/Editor/PropertySelectorLayouts/SDMMaterialPropertySelector_VerticalSlim.h
+++ b/Engine/Plugins/Experimental/DynamicMaterial/Source/DynamicMaterialEditor/Private/UI/Widgets/Editor/PropertySelectorLayouts/SDMMaterialPropertySelector_VerticalSlim.h
@@ -16,6 +16,7 @@ public:
 
 protected:
 	//~ Begin SDMMaterialPropertySelector
+	virtual TSharedRef<SWidget> CreateSlot_PropertyList() override;
 	virtual TSharedRef<SWidget> CreateSlot_SelectButton(EDMMaterialEditorMode InEditMode, EDMMaterialPropertyType InMaterialProperty) override;
 	//~ End SDMMaterialPropertySelector
 };
diff --git a/Engine/Plugins/Experimental/DynamicMaterial/Source/DynamicMaterialEditor/Private/UI/Widgets/Editor/PropertySelectorLayouts/SDMMaterialPropertySelector_WrapBase.cpp b/Engine/Plugins/Experimental/DynamicMaterial/Source/DynamicMaterialEditor/Private/UI/Widgets/Editor/PropertySelectorLayouts/SDMMaterialPropertySelector_WrapBase.cpp
--- a/Engine/Plugins/Experimental/DynamicMaterial/Source/DynamicMaterialEditor/Private/UI/Widgets/Editor/PropertySelectorLayouts/SDMMaterialPropertySelector_WrapBase.cpp
+++ b/Engine/Plugins/Experimental/DynamicMaterial/Source/DynamicMaterialEditor/Private/UI/Widgets/Editor/PropertySelectorLayouts/SDMMaterialPropertySelector_WrapBase.cpp
@@ -40,58 +40,33 @@ TSharedRef<SWidget> SDMMaterialPropertySelector_WrapBase::CreateSlot_PropertyLis
 		return NewSlotList;
 	}
 
+	const FMargin Padding = FMargin(0.f, 1.f);
+
 	NewSlotList->AddSlot()
+		.Padding(Padding)
 		[
-			SNew(SHorizontalBox)
-			+ SHorizontalBox::Slot()
-			.AutoWidth()
-			[
-				SNew(SBox)
-				.WidthOverride(20.f)
-			]
-			+ SHorizontalBox::Slot()
-			.AutoWidth()
-			[
-				CreateSlot_SelectButton(EDMMaterialEditorMode::GlobalSettings, EDMMaterialPropertyType::None)
-			]
+			CreateSlot_SelectButton(EDMMaterialEditorMode::GlobalSettings, EDMMaterialPropertyType::None)
 		];
 
 	NewSlotList->AddSlot()
+		.Padding(Padding)
 		[
-			SNew(SHorizontalBox)
-			+ SHorizontalBox::Slot()
-			.AutoWidth()
-			[
-				SNew(SBox)
-				.WidthOverride(20.f)
-			]
-			+ SHorizontalBox::Slot()
-			.AutoWidth()
-			[
-				CreateSlot_SelectButton(EDMMaterialEditorMode::PropertyPreviews, EDMMaterialPropertyType::None)
-			]
+			CreateSlot_SelectButton(EDMMaterialEditorMode::PropertyPreviews, EDMMaterialPropertyType::None)
 		];
 
 	for (const TPair<EDMMaterialPropertyType, UDMMaterialProperty*>& PropertyPair : EditorOnlyData->GetMaterialProperties())
 	{
-		if (IsCustomMaterialProperty(PropertyPair.Key))
+		if (!PropertyPair.Value || !PropertyPair.Value->IsEnabled() || IsCustomMaterialProperty(PropertyPair.Key)
+			|| !PropertyPair.Value->IsValidForModel(*EditorOnlyData)
+			|| !EditorOnlyData->GetSlotForMaterialProperty(PropertyPair.Value->GetMaterialProperty()))
 		{
 			continue;
 		}
 
 		NewSlotList->AddSlot()
+			.Padding(Padding)
 			[
-				SNew(SHorizontalBox)
-				+ SHorizontalBox::Slot()
-				.AutoWidth()
-				[
-					CreateSlot_EnabledButton(PropertyPair.Key)
-				]
-				+ SHorizontalBox::Slot()
-				.AutoWidth()
-				[
-					CreateSlot_SelectButton(EDMMaterialEditorMode::EditSlot, PropertyPair.Key)
-				]
+				CreateSlot_SelectButton(EDMMaterialEditorMode::EditSlot, PropertyPair.Key)
 			];
 	}
 
@@ -101,7 +76,7 @@ TSharedRef<SWidget> SDMMaterialPropertySelector_WrapBase::CreateSlot_PropertyLis
 TSharedRef<SWidget> SDMMaterialPropertySelector_WrapBase::CreateSlot_SelectButton(EDMMaterialEditorMode InEditMode, 
 	EDMMaterialPropertyType InMaterialProperty)
 {
-	const FText ButtonText = GetSelectButtonText(InEditMode, InMaterialProperty, /* Short Name */ true);
+	const FText ButtonText = GetSelectButtonText(InEditMode, InMaterialProperty, /* Short Name */ false);
 	const FText ToolTip = GetButtonToolTip(InEditMode, InMaterialProperty);
 
 	return SNew(SCheckBox)
@@ -122,20 +97,17 @@ TSharedRef<SWidget> SDMMaterialPropertySelector_WrapBase::CreateSlot_SelectButto
 				SNew(SImage)
 				.Image(FAppStyle::Get().GetBrush("FilterBar.FilterImage"))
 				.ColorAndOpacity(this, &SDMMaterialPropertySelector_WrapBase::GetPropertySelectButtonChipColor, InEditMode, InMaterialProperty)
+				.DesiredSizeOverride(FVector2D(8, 17))
 			]
 			+SHorizontalBox::Slot()
-			.Padding(4.f, 6.f)
+			.Padding(4.f, 4.f)
 			.VAlign(VAlign_Center)
 			.FillWidth(1.f)
 			[
-				SNew(SBox)
-				.WidthOverride(32.f)
-				[
-					SNew(STextBlock)
-					.Font(IDetailLayoutBuilder::GetDetailFont())
-					.Text(ButtonText)
-					.Justification(ETextJustify::Center)
-				]
+				SNew(STextBlock)
+				.Font(IDetailLayoutBuilder::GetDetailFont())
+				.Text(ButtonText)
+				.Justification(ETextJustify::Center)
 			]
 		];
 }
diff --git a/Engine/Plugins/Experimental/DynamicMaterial/Source/DynamicMaterialEditor/Private/UI/Widgets/Editor/PropertySelectorLayouts/SDMMaterialPropertySelector_WrapSlim.cpp b/Engine/Plugins/Experimental/DynamicMaterial/Source/DynamicMaterialEditor/Private/UI/Widgets/Editor/PropertySelectorLayouts/SDMMaterialPropertySelector_WrapSlim.cpp
--- a/Engine/Plugins/Experimental/DynamicMaterial/Source/DynamicMaterialEditor/Private/UI/Widgets/Editor/PropertySelectorLayouts/SDMMaterialPropertySelector_WrapSlim.cpp
+++ b/Engine/Plugins/Experimental/DynamicMaterial/Source/DynamicMaterialEditor/Private/UI/Widgets/Editor/PropertySelectorLayouts/SDMMaterialPropertySelector_WrapSlim.cpp
@@ -3,9 +3,15 @@
 #include "UI/Widgets/Editor/PropertySelectorLayouts/SDMMaterialPropertySelector_WrapSlim.h"
 
 #include "Components/DMMaterialProperty.h"
+#include "DetailLayoutBuilder.h"
 #include "Model/DynamicMaterialModelEditorOnlyData.h"
 #include "UI/Widgets/SDMMaterialEditor.h"
+#include "Utils/DMPrivate.h"
+#include "Widgets/Images/SImage.h"
+#include "Widgets/Input/SCheckBox.h"
+#include "Widgets/Layout/SBox.h"
 #include "Widgets/Layout/SWrapBox.h"
+#include "Widgets/Text/STextBlock.h"
 
 #define LOCTEXT_NAMESPACE "SDMMaterialPropertySelector_WrapSlim"
 
@@ -32,29 +38,40 @@ TSharedRef<SWidget> SDMMaterialPropertySelector_WrapSlim::CreateSlot_PropertyLis
 		return NewSlotList;
 	}
 
+	const FMargin Padding = FMargin(0.f, 1.f);
+
+	TSharedRef<SWidget> PreviewButton = CreateSlot_SelectButton(EDMMaterialEditorMode::MaterialPreview, EDMMaterialPropertyType::None);
+	SetupMaterialPreviewButton(PreviewButton);
+
 	NewSlotList->AddSlot()
+		.Padding(Padding)
+		[
+			PreviewButton
+		];
+
+	NewSlotList->AddSlot()
+		.Padding(Padding)
 		[
 			CreateSlot_SelectButton(EDMMaterialEditorMode::GlobalSettings, EDMMaterialPropertyType::None)
 		];
 
 	NewSlotList->AddSlot()
+		.Padding(Padding)
 		[
 			CreateSlot_SelectButton(EDMMaterialEditorMode::PropertyPreviews, EDMMaterialPropertyType::None)
 		];
 
 	for (const TPair<EDMMaterialPropertyType, UDMMaterialProperty*>& PropertyPair : EditorOnlyData->GetMaterialProperties())
 	{
-		if (!PropertyPair.Value || !PropertyPair.Value->IsEnabled())
-		{
-			continue;
-		}
-
-		if (!EditorOnlyData->GetSlotForMaterialProperty(PropertyPair.Value->GetMaterialProperty()))
+		if (!PropertyPair.Value || !PropertyPair.Value->IsEnabled() || IsCustomMaterialProperty(PropertyPair.Key)
+			|| !PropertyPair.Value->IsValidForModel(*EditorOnlyData)
+			|| !EditorOnlyData->GetSlotForMaterialProperty(PropertyPair.Value->GetMaterialProperty()))
 		{
 			continue;
 		}
 
 		NewSlotList->AddSlot()
+			.Padding(Padding)
 			[
 				CreateSlot_SelectButton(EDMMaterialEditorMode::EditSlot, PropertyPair.Key)
 			];
@@ -63,4 +80,46 @@ TSharedRef<SWidget> SDMMaterialPropertySelector_WrapSlim::CreateSlot_PropertyLis
 	return NewSlotList;
 }
 
+TSharedRef<SWidget> SDMMaterialPropertySelector_WrapSlim::CreateSlot_SelectButton(EDMMaterialEditorMode InEditMode, EDMMaterialPropertyType InMaterialProperty)
+{
+	const FText ButtonText = GetSelectButtonText(InEditMode, InMaterialProperty, /* Short Name */ true);
+	const FText ToolTip = GetButtonToolTip(InEditMode, InMaterialProperty);
+
+	return SNew(SCheckBox)
+		.Style(FAppStyle::Get(), "DetailsView.SectionButton")
+		.HAlign(EHorizontalAlignment::HAlign_Center)
+		.Padding(0.f)
+		.IsEnabled(this, &SDMMaterialPropertySelector_WrapSlim::GetPropertySelectEnabled, InEditMode, InMaterialProperty)
+		.IsChecked(this, &SDMMaterialPropertySelector_WrapSlim::GetPropertySelectState, InEditMode, InMaterialProperty)
+		.OnCheckStateChanged(this, &SDMMaterialPropertySelector_WrapSlim::OnPropertySelectStateChanged, InEditMode, InMaterialProperty)
+		.ToolTipText(ToolTip)
+		.Content()
+		[
+			SNew(SHorizontalBox)
+			+SHorizontalBox::Slot()
+			.VAlign(VAlign_Center)
+			.AutoWidth()
+			[
+				SNew(SImage)
+				.Image(FAppStyle::Get().GetBrush("FilterBar.FilterImage"))
+				.ColorAndOpacity(this, &SDMMaterialPropertySelector_WrapSlim::GetPropertySelectButtonChipColor, InEditMode, InMaterialProperty)
+				.DesiredSizeOverride(FVector2D(8, 17))
+			]
+			+SHorizontalBox::Slot()
+			.Padding(4.f, 4.f)
+			.VAlign(VAlign_Center)
+			.FillWidth(1.f)
+			[
+				SNew(SBox)
+				.WidthOverride(32.f)
+				[
+					SNew(STextBlock)
+					.Font(IDetailLayoutBuilder::GetDetailFont())
+					.Text(ButtonText)
+					.Justification(ETextJustify::Center)
+				]
+			]
+		];
+}
+
 #undef LOCTEXT_NAMESPACE
diff --git a/Engine/Plugins/Experimental/DynamicMaterial/Source/DynamicMaterialEditor/Private/UI/Widgets/Editor/PropertySelectorLayouts/SDMMaterialPropertySelector_WrapSlim.h b/Engine/Plugins/Experimental/DynamicMaterial/Source/DynamicMaterialEditor/Private/UI/Widgets/Editor/PropertySelectorLayouts/SDMMaterialPropertySelector_WrapSlim.h
--- a/Engine/Plugins/Experimental/DynamicMaterial/Source/DynamicMaterialEditor/Private/UI/Widgets/Editor/PropertySelectorLayouts/SDMMaterialPropertySelector_WrapSlim.h
+++ b/Engine/Plugins/Experimental/DynamicMaterial/Source/DynamicMaterialEditor/Private/UI/Widgets/Editor/PropertySelectorLayouts/SDMMaterialPropertySelector_WrapSlim.h
@@ -15,7 +15,8 @@ public:
 	void Construct(const FArguments& InArgs, const TSharedRef<SDMMaterialEditor> InEditorWidget);
 
 protected:
-	//~ Begin SDMMaterialPropertySelector_WrapBase
+	//~ Begin SDMMaterialPropertySelector
 	virtual TSharedRef<SWidget> CreateSlot_PropertyList() override;
-	//~ End SDMMaterialPropertySelector_WrapBase
+	virtual TSharedRef<SWidget> CreateSlot_SelectButton(EDMMaterialEditorMode InEditMode, EDMMaterialPropertyType InMaterialProperty) override;
+	//~ End SDMMaterialPropertySelector
 };
diff --git a/Engine/Plugins/Experimental/DynamicMaterial/Source/DynamicMaterialEditor/Private/UI/Widgets/Editor/SDMMaterialPreview.cpp b/Engine/Plugins/Experimental/DynamicMaterial/Source/DynamicMaterialEditor/Private/UI/Widgets/Editor/SDMMaterialPreview.cpp
--- a/Engine/Plugins/Experimental/DynamicMaterial/Source/DynamicMaterialEditor/Private/UI/Widgets/Editor/SDMMaterialPreview.cpp
+++ b/Engine/Plugins/Experimental/DynamicMaterial/Source/DynamicMaterialEditor/Private/UI/Widgets/Editor/SDMMaterialPreview.cpp
@@ -26,6 +26,7 @@
 #include "UI/Widgets/Editor/DMMaterialPreviewViewportClient.h"
 #include "UI/Widgets/SDMMaterialEditor.h"
 #include "UnrealEdGlobals.h"
+#include "Widgets/SNullWidget.h"
 
 #define LOCTEXT_NAMESPACE "SDMMaterialPreview"
 
@@ -81,6 +82,7 @@ void SDMMaterialPreview::Construct(const FArguments& InArgs, const TSharedRef<SD
 	UDynamicMaterialModelBase* InMaterialModelBase)
 {
 	EditorWidgetWeak = InEditorWidget;
+	bShowMenu = InArgs._ShowMenu;
 	PreviewScene = MakeShareable(new FAdvancedPreviewScene(FPreviewScene::ConstructionValues()));
 
 	PreviewMaterial = nullptr;
@@ -139,6 +141,11 @@ void SDMMaterialPreview::Construct(const FArguments& InArgs, const TSharedRef<SD
 
 TSharedPtr<SWidget> SDMMaterialPreview::MakeViewportToolbar()
 {
+	if (!bShowMenu)
+	{
+		return SNullWidget::NullWidget;
+	}
+
 	FSlimHorizontalToolBarBuilder ToolbarBuilder(CommandList, FMultiBoxCustomization::None, {});
 
 	const FName ToolBarStyle = TEXT("EditorViewportToolBar");
@@ -469,6 +476,19 @@ TSharedRef<SWidget> SDMMaterialPreview::GenerateToolbarMenu()
 
 		FToolMenuSection& SettingsSection = Menu->AddSection(TEXT("Settings"), LOCTEXT("Settings", "Settings"));
 		SettingsSection.AddEntry(FToolMenuEntry::InitMenuEntry(MaterialEditorCommands.TogglePreviewBackground));
+
+		FToolMenuSection& ActionsSection = Menu->AddSection(TEXT("Actions"), LOCTEXT("Actions", "Actions"));
+
+		FUIAction OpenPreviewTabAction;
+		OpenPreviewTabAction.ExecuteAction.BindSP(this, &SDMMaterialPreview::OpenMaterialPreviewTab);
+
+		ActionsSection.AddEntry(FToolMenuEntry::InitMenuEntry(
+			TEXT("PopoutMaterialPreviewTab"),
+			LOCTEXT("OpenPreview", "Open Preview"),
+			LOCTEXT("OpenPreviewToolTip", "Open a tab with a preview of the material."),
+			TAttribute<FSlateIcon>(),
+			OpenPreviewTabAction
+		));		
 	}
 
 	FToolMenuContext Context;
@@ -498,6 +518,18 @@ void SDMMaterialPreview::OnEditorSettingsChanged(const FPropertyChangedEvent& In
 	}
 }
 
+void SDMMaterialPreview::OpenMaterialPreviewTab()
+{
+	TSharedPtr<SDMMaterialEditor> EditorWidget = EditorWidgetWeak.Pin();
+
+	if (!EditorWidget.IsValid())
+	{
+		return;
+	}
+
+	EditorWidget->OpenMaterialPreviewTab();
+}
+
 TSharedRef<FEditorViewportClient> SDMMaterialPreview::MakeEditorViewportClient()
 {
 	EditorViewportClient = MakeShared<FDMMaterialPreviewViewportClient>(SharedThis(this), *PreviewScene.Get());
diff --git a/Engine/Plugins/Experimental/DynamicMaterial/Source/DynamicMaterialEditor/Private/UI/Widgets/Editor/SDMMaterialPreview.h b/Engine/Plugins/Experimental/DynamicMaterial/Source/DynamicMaterialEditor/Private/UI/Widgets/Editor/SDMMaterialPreview.h
--- a/Engine/Plugins/Experimental/DynamicMaterial/Source/DynamicMaterialEditor/Private/UI/Widgets/Editor/SDMMaterialPreview.h
+++ b/Engine/Plugins/Experimental/DynamicMaterial/Source/DynamicMaterialEditor/Private/UI/Widgets/Editor/SDMMaterialPreview.h
@@ -25,7 +25,10 @@ class SDMMaterialPreview : public SEditorViewport, public FGCObject
 {
 	SLATE_DECLARE_WIDGET(SDMMaterialPreview, SCompoundWidget)
 
-	SLATE_BEGIN_ARGS(SDMMaterialPreview) {}
+	SLATE_BEGIN_ARGS(SDMMaterialPreview)
+		: _ShowMenu(true)
+		{}
+		SLATE_ARGUMENT(bool, ShowMenu)
 	SLATE_END_ARGS()
 
 public:
@@ -44,6 +47,7 @@ public:
 
 protected:
 	TWeakPtr<SDMMaterialEditor> EditorWidgetWeak;
+	bool bShowMenu;
 
 	TSharedPtr<FDMMaterialPreviewViewportClient> EditorViewportClient;
 	TSharedPtr<FAdvancedPreviewScene> PreviewScene;
@@ -83,6 +87,8 @@ protected:
 
 	void OnEditorSettingsChanged(const FPropertyChangedEvent& InPropertyChangedEvent);
 
+	void OpenMaterialPreviewTab();
+
 	//~ Begin SEditorViewport
 	virtual TSharedRef<FEditorViewportClient> MakeEditorViewportClient() override;
 	virtual EVisibility OnGetViewportContentVisibility() const override;
diff --git a/Engine/Plugins/Experimental/DynamicMaterial/Source/DynamicMaterialEditor/Private/UI/Widgets/Editor/SDMMaterialPropertySelector.cpp b/Engine/Plugins/Experimental/DynamicMaterial/Source/DynamicMaterialEditor/Private/UI/Widgets/Editor/SDMMaterialPropertySelector.cpp
--- a/Engine/Plugins/Experimental/DynamicMaterial/Source/DynamicMaterialEditor/Private/UI/Widgets/Editor/SDMMaterialPropertySelector.cpp
+++ b/Engine/Plugins/Experimental/DynamicMaterial/Source/DynamicMaterialEditor/Private/UI/Widgets/Editor/SDMMaterialPropertySelector.cpp
@@ -58,21 +58,56 @@ void SDMMaterialPropertySelector::SetSelectedProperty(EDMMaterialEditorMode InEd
 
 	switch (InEditMode)
 	{
-	case EDMMaterialEditorMode::GlobalSettings:
-		EditorWidget->EditGlobalSettings();
-		break;
+		case EDMMaterialEditorMode::MaterialPreview:
+			OpenMaterialPreviewTab();
+			break;
 
-	case EDMMaterialEditorMode::PropertyPreviews:
-		EditorWidget->ShowPropertyPreviews();
-		break;
+		case EDMMaterialEditorMode::GlobalSettings:
+			EditorWidget->EditGlobalSettings();
+			break;
+
+		case EDMMaterialEditorMode::PropertyPreviews:
+			EditorWidget->ShowPropertyPreviews();
+			break;
+
+		case EDMMaterialEditorMode::EditSlot:
+			EditorWidget->SelectProperty(InMaterialProperty);
+			break;
+
+		default:
+			break;
+	}
+}
+
+void SDMMaterialPropertySelector::SetupMaterialPreviewButton(const TSharedRef<SWidget>& InSelectButton)
+{
+	TSharedPtr<SDMMaterialEditor> EditorWidget = GetEditorWidget();
+
+	if (!EditorWidget.IsValid())
+	{
+		return;
+	}
 
-	case EDMMaterialEditorMode::EditSlot:
-		EditorWidget->SelectProperty(InMaterialProperty);
-		break;
+	TSharedPtr<IToolTip> PreviewToolTip = EditorWidget->GetMaterialPreviewToolTip();
 
-	default:
-		break;
+	if (!PreviewToolTip.IsValid())
+	{
+		return;
 	}
+
+	InSelectButton->SetToolTip(PreviewToolTip);
+}
+
+void SDMMaterialPropertySelector::OpenMaterialPreviewTab()
+{
+	TSharedPtr<SDMMaterialEditor> EditorWidget = GetEditorWidget();
+
+	if (!EditorWidget.IsValid())
+	{
+		return;
+	}
+
+	EditorWidget->OpenMaterialPreviewTab();
 }
 
 UDynamicMaterialModelEditorOnlyData* SDMMaterialPropertySelector::GetEditorOnlyData() const
@@ -123,6 +158,11 @@ FText SDMMaterialPropertySelector::GetSelectButtonText(EDMMaterialEditorMode InE
 {
 	switch (InEditMode)
 	{
+		case EDMMaterialEditorMode::MaterialPreview:
+			return bInShortName
+				? LOCTEXT("MaterialPreviewShort", "Prev")
+				: LOCTEXT("MaterialPreview", "Material Preview");
+
 		case EDMMaterialEditorMode::GlobalSettings:
 			return bInShortName
 				? LOCTEXT("GlobalSettingsShort", "Global")
@@ -147,6 +187,9 @@ FText SDMMaterialPropertySelector::GetButtonToolTip(EDMMaterialEditorMode InEdit
 {
 	switch (InEditMode)
 	{
+		case EDMMaterialEditorMode::MaterialPreview:
+			return LOCTEXT("MaterialPreviewToolTip", "Show a preview of the Material.");
+
 		case EDMMaterialEditorMode::GlobalSettings:
 			return LOCTEXT("GeneralSettingsToolTip", "Edit the Material Global Settings.");
 
@@ -261,7 +304,7 @@ bool SDMMaterialPropertySelector::GetPropertyEnabledEnabled(EDMMaterialPropertyT
 
 ECheckBoxState SDMMaterialPropertySelector::GetPropertyEnabledState(EDMMaterialPropertyType InMaterialProperty) const
 {
-	return DoesPropertySlotExist(InMaterialProperty)
+	return (GetPropertyEnabledEnabled(InMaterialProperty) && DoesPropertySlotExist(InMaterialProperty))
 		? ECheckBoxState::Checked
 		: ECheckBoxState::Unchecked;
 }
@@ -284,7 +327,7 @@ bool SDMMaterialPropertySelector::GetPropertySelectEnabled(EDMMaterialEditorMode
 	switch (InEditMode)
 	{
 		case EDMMaterialEditorMode::EditSlot:
-			return DoesPropertySlotExist(InMaterialProperty);
+			return GetPropertyEnabledEnabled(InMaterialProperty) && DoesPropertySlotExist(InMaterialProperty);
 
 		default:
 			return true;
@@ -308,6 +351,9 @@ ECheckBoxState SDMMaterialPropertySelector::GetPropertySelectState(EDMMaterialEd
 				? ECheckBoxState::Checked
 				: ECheckBoxState::Unchecked;
 
+		case EDMMaterialEditorMode::MaterialPreview:
+			return ECheckBoxState::Unchecked;
+
 		case EDMMaterialEditorMode::EditSlot:
 			return InMaterialProperty == EditorWidget->GetSelectedPropertyType()
 				? ECheckBoxState::Checked
@@ -331,6 +377,7 @@ FSlateColor SDMMaterialPropertySelector::GetPropertySelectButtonChipColor(EDMMat
 {
 	switch (InEditMode)
 	{
+		case EDMMaterialEditorMode::MaterialPreview:
 		case EDMMaterialEditorMode::GlobalSettings:
 		case EDMMaterialEditorMode::PropertyPreviews:
 			return FStyleColors::AccentGreen;
diff --git a/Engine/Plugins/Experimental/DynamicMaterial/Source/DynamicMaterialEditor/Private/UI/Widgets/Editor/SDMMaterialPropertySelector.h b/Engine/Plugins/Experimental/DynamicMaterial/Source/DynamicMaterialEditor/Private/UI/Widgets/Editor/SDMMaterialPropertySelector.h
--- a/Engine/Plugins/Experimental/DynamicMaterial/Source/DynamicMaterialEditor/Private/UI/Widgets/Editor/SDMMaterialPropertySelector.h
+++ b/Engine/Plugins/Experimental/DynamicMaterial/Source/DynamicMaterialEditor/Private/UI/Widgets/Editor/SDMMaterialPropertySelector.h
@@ -73,4 +73,8 @@ protected:
 	FSlateColor GetPropertySelectButtonChipColor(EDMMaterialEditorMode InEditMode, EDMMaterialPropertyType InMaterialProperty) const;
 
 	void SetSelectedProperty(EDMMaterialEditorMode InEditMode, EDMMaterialPropertyType InMaterialProperty);
+
+	void SetupMaterialPreviewButton(const TSharedRef<SWidget>& InSelectButton);
+
+	void OpenMaterialPreviewTab();
 };
diff --git a/Engine/Plugins/Experimental/DynamicMaterial/Source/DynamicMaterialEditor/Private/UI/Widgets/SDMMaterialDesigner.cpp b/Engine/Plugins/Experimental/DynamicMaterial/Source/DynamicMaterialEditor/Private/UI/Widgets/SDMMaterialDesigner.cpp
--- a/Engine/Plugins/Experimental/DynamicMaterial/Source/DynamicMaterialEditor/Private/UI/Widgets/SDMMaterialDesigner.cpp
+++ b/Engine/Plugins/Experimental/DynamicMaterial/Source/DynamicMaterialEditor/Private/UI/Widgets/SDMMaterialDesigner.cpp
@@ -473,6 +473,7 @@ void SDMMaterialDesigner::OnLayoutChanged()
 
 	switch (EditorMode)
 	{
+		default:
 		case EDMMaterialEditorMode::GlobalSettings:
 			NewEditor->EditGlobalSettings();
 			break;
diff --git a/Engine/Plugins/Experimental/DynamicMaterial/Source/DynamicMaterialEditor/Private/UI/Widgets/SDMMaterialEditor.cpp b/Engine/Plugins/Experimental/DynamicMaterial/Source/DynamicMaterialEditor/Private/UI/Widgets/SDMMaterialEditor.cpp
--- a/Engine/Plugins/Experimental/DynamicMaterial/Source/DynamicMaterialEditor/Private/UI/Widgets/SDMMaterialEditor.cpp
+++ b/Engine/Plugins/Experimental/DynamicMaterial/Source/DynamicMaterialEditor/Private/UI/Widgets/SDMMaterialEditor.cpp
@@ -10,12 +10,15 @@
 #include "DynamicMaterialEditorSettings.h"
 #include "DynamicMaterialModule.h"
 #include "Engine/Texture.h"
+#include "Framework/Application/SlateApplication.h"
 #include "Framework/Commands/GenericCommands.h"
+#include "Materials/Material.h"
 #include "Materials/MaterialFunctionInterface.h"
 #include "Model/DynamicMaterialModel.h"
 #include "Model/DynamicMaterialModelBase.h"
 #include "Model/DynamicMaterialModelDynamic.h"
 #include "Model/DynamicMaterialModelEditorOnlyData.h"
+#include "Styling/SlateIconFinder.h"
 #include "UI/Utils/DMPreviewMaterialManager.h"
 #include "UI/Widgets/Editor/SDMMaterialComponentEditor.h"
 #include "UI/Widgets/Editor/SDMMaterialGlobalSettingsEditor.h"
@@ -27,10 +30,13 @@
 #include "UI/Widgets/Editor/SDMToolBar.h"
 #include "UI/Widgets/SDMMaterialDesigner.h"
 #include "Utils/DMMaterialModelFunctionLibrary.h"
+#include "Widgets/Docking/SDockTab.h"
+#include "Widgets/Layout/SBox.h"
 #include "Widgets/Layout/SScrollBox.h"
 #include "Widgets/Layout/SSplitter.h"
 #include "Widgets/SBoxPanel.h"
 #include "Widgets/SNullWidget.h"
+#include "Widgets/SToolTip.h"
 
 #define LOCTEXT_NAMESPACE "SDMMaterialEditor"
 
@@ -48,6 +54,8 @@ SDMMaterialEditor::SDMMaterialEditor()
 SDMMaterialEditor::~SDMMaterialEditor()
 {
 	FCoreDelegates::OnEnginePreExit.RemoveAll(this);
+	CloseMaterialPreviewTab();
+	DestroyMaterialPreviewToolTip();
 
 	if (!FDynamicMaterialModule::AreUObjectsSafe())
 	{
@@ -56,6 +64,8 @@ SDMMaterialEditor::~SDMMaterialEditor()
 
 	if (UDynamicMaterialModelEditorOnlyData* EditorOnlyData = EditorOnlyDataUpdateObject.Get())
 	{
+		EditorOnlyData->GetOnMaterialBuiltDelegate().RemoveAll(this);
+		EditorOnlyData->GetOnPropertyUpdateDelegate().RemoveAll(this);
 		EditorOnlyData->GetOnSlotListUpdateDelegate().RemoveAll(this);
 	}
 }
@@ -334,6 +344,131 @@ void SDMMaterialEditor::ShowPropertyPreviews(bool bInForceRefresh)
 	MaterialPropertyPreviewsSlot.Invalidate();
 }
 
+void SDMMaterialEditor::OpenMaterialPreviewTab()
+{
+	UDynamicMaterialModelBase* MaterialModelBase = GetMaterialModelBase();
+
+	if (!MaterialModelBase)
+	{
+		return;
+	}
+
+	CloseMaterialPreviewTab();
+
+	FSlateApplication::Get().CloseToolTip();
+
+	const FName TabId = TEXT("MaterialPreviewTab");
+
+	if (!FGlobalTabmanager::Get()->HasTabSpawner(TabId))
+	{
+		FGlobalTabmanager::Get()->RegisterNomadTabSpawner(
+			TabId,
+			FOnSpawnTab::CreateLambda(
+				[TabId](const FSpawnTabArgs& InArgs)
+				{
+					TSharedRef<SDockTab> DockTab = SNew(SDockTab)
+						.Label(FText::FromName(TabId))
+						.LabelSuffix(LOCTEXT("TabSuffix", "Material Preview"));
+
+					DockTab->SetTabIcon(FSlateIconFinder::FindIconForClass(UMaterial::StaticClass()).GetIcon());
+
+					return DockTab;
+				}
+			)
+		);
+	}
+
+	MaterialPreviewTab = FGlobalTabmanager::Get()->TryInvokeTab(TabId);
+	MaterialPreviewTab->ActivateInParent(ETabActivationCause::SetDirectly);
+	MaterialPreviewTab->SetLabel(FText::FromString(MaterialModelBase->GetPathName()));
+	MaterialPreviewTab->SetOnTabClosed(SDockTab::FOnTabClosedCallback::CreateSPLambda(
+		this,
+		[this](TSharedRef<SDockTab> InDockTab)
+		{
+			MaterialPreviewTabSlot.ClearWidget();
+		}));
+
+	TSharedRef<SBox> Wrapper = SNew(SBox);
+
+	MaterialPreviewTabSlot = TDMWidgetSlot<SDMMaterialPreview>(
+		Wrapper, 
+		0, 
+		SNew(SDMMaterialPreview, SharedThis(this), MaterialModelBase)
+	);
+
+	MaterialPreviewTab->SetContent(Wrapper);
+}
+
+void SDMMaterialEditor::CloseMaterialPreviewTab()
+{
+	if (MaterialPreviewTab.IsValid())
+	{
+		MaterialPreviewTabSlot.ClearWidget();
+		MaterialPreviewTab->RequestCloseTab();
+		MaterialPreviewTab.Reset();
+	}
+}
+
+TSharedPtr<IToolTip> SDMMaterialEditor::GetMaterialPreviewToolTip()
+{
+	UDynamicMaterialModelBase* MaterialModelBase = GetMaterialModelBase();
+
+	if (!MaterialModelBase)
+	{
+		return nullptr;
+	}
+
+	UDynamicMaterialEditorSettings* Settings = UDynamicMaterialEditorSettings::Get();
+
+	if (!Settings)
+	{
+		return nullptr;
+	}
+
+	DestroyMaterialPreviewToolTip();
+
+	TSharedRef<SBox> Wrapper = SNew(SBox)
+		.WidthOverride(TAttribute<FOptionalSize>::CreateWeakLambda(
+			Settings,
+			[Settings]
+			{
+				return Settings->ThumbnailSize;
+			}
+		))
+		.HeightOverride(TAttribute<FOptionalSize>::CreateWeakLambda(
+			Settings,
+			[Settings]
+			{
+				return Settings->ThumbnailSize;
+			}
+		));
+
+	MaterialPreviewToolTipSlot = TDMWidgetSlot<SDMMaterialPreview>(
+		Wrapper,
+		0,
+		SNew(SDMMaterialPreview, SharedThis(this), MaterialModelBase)
+		.ShowMenu(false)
+	);
+
+	MaterialPreviewToolTip = SNew(SToolTip)
+		.IsInteractive(false)
+		.BorderImage(FCoreStyle::Get().GetBrush("ToolTip.Background"))
+		[
+			Wrapper
+		];
+
+	return MaterialPreviewToolTip.ToSharedRef();
+}
+
+void SDMMaterialEditor::DestroyMaterialPreviewToolTip()
+{
+	if (MaterialPreviewToolTip.IsValid())
+	{
+		MaterialPreviewToolTipSlot.ClearWidget();
+		MaterialPreviewToolTip.Reset();
+	}
+}
+
 void SDMMaterialEditor::Validate()
 {
 	if (!FDynamicMaterialModule::AreUObjectsSafe())
@@ -760,6 +895,8 @@ void SDMMaterialEditor::OnUndo()
 void SDMMaterialEditor::OnEnginePreExit()
 {
 	MaterialPreviewSlot.ClearWidget();
+	CloseMaterialPreviewTab();
+	DestroyMaterialPreviewToolTip();
 }
 
 void SDMMaterialEditor::OnEditorSplitterResized()
@@ -781,12 +918,18 @@ void SDMMaterialEditor::BindEditorOnlyDataUpdate(UDynamicMaterialModelBase* InMa
 		if (UDynamicMaterialModelEditorOnlyData* EditorOnlyData = UDynamicMaterialModelEditorOnlyData::Get(MaterialModel))
 		{
 			EditorOnlyDataUpdateObject = EditorOnlyData;
+			EditorOnlyData->GetOnMaterialBuiltDelegate().AddSP(this, &SDMMaterialEditor::OnMaterialBuilt);
 			EditorOnlyData->GetOnPropertyUpdateDelegate().AddSP(this, &SDMMaterialEditor::OnPropertyUpdate);
 			EditorOnlyData->GetOnSlotListUpdateDelegate().AddSP(this, &SDMMaterialEditor::OnSlotListUpdate);
 		}
 	}
 }
 
+void SDMMaterialEditor::OnMaterialBuilt(UDynamicMaterialModelBase* InMaterialModelBase)
+{
+	PropertySelectorSlot.Invalidate();
+}
+
 void SDMMaterialEditor::OnPropertyUpdate(UDynamicMaterialModelBase* InMaterialModelBase)
 {
 	PropertySelectorSlot.Invalidate();
diff --git a/Engine/Plugins/Experimental/DynamicMaterial/Source/DynamicMaterialEditor/Private/UI/Widgets/SDMMaterialEditor.h b/Engine/Plugins/Experimental/DynamicMaterial/Source/DynamicMaterialEditor/Private/UI/Widgets/SDMMaterialEditor.h
--- a/Engine/Plugins/Experimental/DynamicMaterial/Source/DynamicMaterialEditor/Private/UI/Widgets/SDMMaterialEditor.h
+++ b/Engine/Plugins/Experimental/DynamicMaterial/Source/DynamicMaterialEditor/Private/UI/Widgets/SDMMaterialEditor.h
@@ -14,6 +14,7 @@
 class FDMPreviewMaterialManager;
 class FSlotBase;
 class FUICommandList;
+class IToolTip;
 class SDMMaterialComponentEditor;
 class SDMMaterialDesigner;
 class SDMMaterialGlobalSettingsEditor;
@@ -23,6 +24,7 @@ class SDMMaterialPropertySelector;
 class SDMMaterialSlotEditor;
 class SDMStatusBar;
 class SDMToolBar;
+class SDockTab;
 class SSplitter;
 class UDMMaterialComponent;
 class UDMMaterialSlot;
@@ -42,7 +44,8 @@ enum class EDMMaterialEditorMode : uint8
 {
 	GlobalSettings,
 	PropertyPreviews,
-	EditSlot
+	EditSlot,
+	MaterialPreview
 };
 
 class SDMMaterialEditor : public SCompoundWidget, public FSelfRegisteringEditorUndoClient
@@ -105,6 +108,14 @@ public:
 
 	virtual void ShowPropertyPreviews(bool bInForceRefresh = false);
 
+	void OpenMaterialPreviewTab();
+
+	void CloseMaterialPreviewTab();
+
+	TSharedPtr<IToolTip> GetMaterialPreviewToolTip();
+
+	void DestroyMaterialPreviewToolTip();
+
 	void Validate();
 
 	DECLARE_MULTICAST_DELEGATE_TwoParams(FOnEditedSlotChanged, const TSharedRef<SDMMaterialSlotEditor>&, UDMMaterialSlot*);
@@ -143,6 +154,10 @@ protected:
 
 	TSharedRef<FUICommandList> CommandList;
 	TSharedRef<FDMPreviewMaterialManager> PreviewMaterialManager;
+	TSharedPtr<SDockTab> MaterialPreviewTab;
+	TDMWidgetSlot<SDMMaterialPreview> MaterialPreviewTabSlot;
+	TSharedPtr<IToolTip> MaterialPreviewToolTip;
+	TDMWidgetSlot<SDMMaterialPreview> MaterialPreviewToolTipSlot;
 
 	EDMMaterialEditorMode EditMode;
 	EDMMaterialPropertyType SelectedMaterialProperty;
@@ -210,6 +225,8 @@ protected:
 
 	void BindEditorOnlyDataUpdate(UDynamicMaterialModelBase* InMaterialModelBase);
 
+	void OnMaterialBuilt(UDynamicMaterialModelBase* InMaterialModelBase);
+
 	void OnPropertyUpdate(UDynamicMaterialModelBase* InMaterialModelBase);
 
 	void OnSlotListUpdate(UDynamicMaterialModelBase* InMaterialModelBase);
