Material Designer
- Blend and Mask stages now push their channel masks when their appropriate inputs are changed.

diff --git a/Engine/Plugins/Experimental/DynamicMaterial/Source/DynamicMaterialEditor/Private/Components/DMMaterialStage.cpp b/Engine/Plugins/Experimental/DynamicMaterial/Source/DynamicMaterialEditor/Private/Components/DMMaterialStage.cpp
--- a/Engine/Plugins/Experimental/DynamicMaterial/Source/DynamicMaterialEditor/Private/Components/DMMaterialStage.cpp
+++ b/Engine/Plugins/Experimental/DynamicMaterial/Source/DynamicMaterialEditor/Private/Components/DMMaterialStage.cpp
@@ -417,7 +417,7 @@ void UDMMaterialStage::InputUpdated(UDMMaterialStageInput* InInput, EDMUpdateTyp
 
 		if (bIsUsedInInput)
 		{
-			Throughput->InputUpdated(InputMapIdx, InUpdateType);
+			Throughput->OnInputUpdated(InputMapIdx, InUpdateType);
 		}
 	}
 }
@@ -693,6 +693,8 @@ UDMMaterialStageInput* UDMMaterialStage::ChangeInput(TSubclassOf<UDMMaterialStag
 		EDMMaterialPropertyType::None
 	);
 
+	Throughput->OnPostInputAdded(InInputIdx);
+
 	return NewInput;
 }
 
diff --git a/Engine/Plugins/Experimental/DynamicMaterial/Source/DynamicMaterialEditor/Private/Components/DMMaterialStageBlend.cpp b/Engine/Plugins/Experimental/DynamicMaterial/Source/DynamicMaterialEditor/Private/Components/DMMaterialStageBlend.cpp
--- a/Engine/Plugins/Experimental/DynamicMaterial/Source/DynamicMaterialEditor/Private/Components/DMMaterialStageBlend.cpp
+++ b/Engine/Plugins/Experimental/DynamicMaterial/Source/DynamicMaterialEditor/Private/Components/DMMaterialStageBlend.cpp
@@ -395,6 +395,13 @@ bool UDMMaterialStageBlend::GenerateStagePreviewMaterial(UDMMaterialStage* InSta
 	return true;
 }
 
+void UDMMaterialStageBlend::Update(UDMMaterialComponent* InSource, EDMUpdateType InUpdateType)
+{
+	Super::Update(InSource, InUpdateType);
+
+	PullBaseChannelOverride();
+}
+
 void UDMMaterialStageBlend::NotifyPostChange(const FPropertyChangedEvent& InPropertyChangedEvent, FEditPropertyChain* InPropertyThatChanged)
 {
 	if (!IsComponentValid())
@@ -565,6 +572,16 @@ int32 UDMMaterialStageBlend::ResolveInput(const TSharedRef<FDMMaterialBuildState
 	return NodeOutputIndex;
 }
 
+void UDMMaterialStageBlend::OnPostInputAdded(int32 InInputIdx)
+{
+	Super::OnPostInputAdded(InInputIdx);
+
+	if (BaseChannelOverride != EAvaColorChannel::None)
+	{
+		PushBaseChannelOverride();
+	}
+}
+
 FText UDMMaterialStageBlend::GetStageDescription() const
 {
 	if (UDMMaterialStageInput* StageInputB = GetInputB())
diff --git a/Engine/Plugins/Experimental/DynamicMaterial/Source/DynamicMaterialEditor/Private/Components/DMMaterialStageThroughputLayerBlend.cpp b/Engine/Plugins/Experimental/DynamicMaterial/Source/DynamicMaterialEditor/Private/Components/DMMaterialStageThroughputLayerBlend.cpp
--- a/Engine/Plugins/Experimental/DynamicMaterial/Source/DynamicMaterialEditor/Private/Components/DMMaterialStageThroughputLayerBlend.cpp
+++ b/Engine/Plugins/Experimental/DynamicMaterial/Source/DynamicMaterialEditor/Private/Components/DMMaterialStageThroughputLayerBlend.cpp
@@ -97,6 +97,8 @@ void UDMMaterialStageThroughputLayerBlend::OnComponentAdded()
 	{
 		UpdateAlphaOnlyMasks(EDMUpdateType::Structure);
 	}
+
+	PushMaskChannelOverride();
 }
  
 void UDMMaterialStageThroughputLayerBlend::AddDefaultInput(int32 InInputIndex) const
@@ -214,7 +216,7 @@ void UDMMaterialStageThroughputLayerBlend::Update(UDMMaterialComponent* InSource
 			}
 		}
 	}
- 
+
 	UpdateLinkedInputStage(InUpdateType);
  
 	UpdateAlphaOnlyMasks(InUpdateType);
@@ -764,7 +766,17 @@ void UDMMaterialStageThroughputLayerBlend::ConnectOutputToInput(const TSharedRef
 			break;
 	}
 }
- 
+
+void UDMMaterialStageThroughputLayerBlend::OnPostInputAdded(int32 InInputIdx)
+{
+	Super::OnPostInputAdded(InInputIdx);
+
+	if (MaskChannelOverride != EAvaColorChannel::None)
+	{
+		PushMaskChannelOverride();
+	}
+}
+
 void UDMMaterialStageThroughputLayerBlend::GetMaskOutput(const TSharedRef<FDMMaterialBuildState>& InBuildState, UMaterialExpression*& OutExpression,
 	int32& OutOutputIndex, int32& OutOutputChannel) const
 {
diff --git a/Engine/Plugins/Experimental/DynamicMaterial/Source/DynamicMaterialEditor/Private/Components/MaterialStageExpressions/DMMSETextureSampleBase.cpp b/Engine/Plugins/Experimental/DynamicMaterial/Source/DynamicMaterialEditor/Private/Components/MaterialStageExpressions/DMMSETextureSampleBase.cpp
--- a/Engine/Plugins/Experimental/DynamicMaterial/Source/DynamicMaterialEditor/Private/Components/MaterialStageExpressions/DMMSETextureSampleBase.cpp
+++ b/Engine/Plugins/Experimental/DynamicMaterial/Source/DynamicMaterialEditor/Private/Components/MaterialStageExpressions/DMMSETextureSampleBase.cpp
@@ -229,7 +229,7 @@ bool UDMMaterialStageExpressionTextureSampleBase::CanChangeInputType(int32 InInp
 	return Super::CanChangeInputType(InInputIndex);
 }
 
-void UDMMaterialStageExpressionTextureSampleBase::InputUpdated(int32 InInputIndex, EDMUpdateType InUpdateType)
+void UDMMaterialStageExpressionTextureSampleBase::OnInputUpdated(int32 InInputIndex, EDMUpdateType InUpdateType)
 {
 	// If the texture changes, update the mask!
 	if (InInputIndex == 0)
diff --git a/Engine/Plugins/Experimental/DynamicMaterial/Source/DynamicMaterialEditor/Public/Components/DMMaterialStageBlend.h b/Engine/Plugins/Experimental/DynamicMaterial/Source/DynamicMaterialEditor/Public/Components/DMMaterialStageBlend.h
--- a/Engine/Plugins/Experimental/DynamicMaterial/Source/DynamicMaterialEditor/Public/Components/DMMaterialStageBlend.h
+++ b/Engine/Plugins/Experimental/DynamicMaterial/Source/DynamicMaterialEditor/Public/Components/DMMaterialStageBlend.h
@@ -50,6 +50,7 @@ public:
 	DYNAMICMATERIALEDITOR_API virtual bool IsInputVisible(int32 InputIndex) const override;
 	DYNAMICMATERIALEDITOR_API virtual int32 ResolveInput(const TSharedRef<FDMMaterialBuildState>& InBuildState, int32 InputIndex, 
 		FDMMaterialStageConnectorChannel& OutChannel, TArray<UMaterialExpression*>& OutExpressions) const override;
+	virtual void OnPostInputAdded(int32 InInputIdx) override;
 	//~ End UDMMaterialStageThroughput
 
 	//~ Begin UDMMaterialStageSource
@@ -62,6 +63,10 @@ public:
 		UMaterialExpression*& OutMaterialExpression, int32& OutputIndex) override;
 	//~ End UDMMaterialStageSource
 
+	//~ Begin UDMMaterialComponent
+	DYNAMICMATERIALEDITOR_API virtual void Update(UDMMaterialComponent* InSource, EDMUpdateType InUpdateType) override;
+	//~ End UDMMaterialComponent
+
 	//~ Begin FNotifyHook
 	DYNAMICMATERIALEDITOR_API virtual void NotifyPostChange(const FPropertyChangedEvent& InPropertyChangedEvent, class FEditPropertyChain* InPropertyThatChanged) override;
 	//~ End FNotifyHook
@@ -69,13 +74,13 @@ public:
 protected:
 	static TArray<TStrongObjectPtr<UClass>> Blends;
 
+	static void GenerateBlendList();
+
 	UPROPERTY(EditInstanceOnly, BlueprintReadOnly, Getter, Setter, BlueprintGetter = GetBaseChannelOverride,
 		Category = "Material Designer", DisplayName = "Channel Mask",
 		meta = (NotKeyframeable, ToolTip = "Changes the output channel of the base input."))
 	mutable EAvaColorChannel BaseChannelOverride;
 
-	static void GenerateBlendList();
-
 	DYNAMICMATERIALEDITOR_API UDMMaterialStageBlend(const FText& InName);
 
 	/* Returns true if there are any outputs on the base input that have more than 1 channel. */
diff --git a/Engine/Plugins/Experimental/DynamicMaterial/Source/DynamicMaterialEditor/Public/Components/DMMaterialStageThroughput.h b/Engine/Plugins/Experimental/DynamicMaterial/Source/DynamicMaterialEditor/Public/Components/DMMaterialStageThroughput.h
--- a/Engine/Plugins/Experimental/DynamicMaterial/Source/DynamicMaterialEditor/Public/Components/DMMaterialStageThroughput.h
+++ b/Engine/Plugins/Experimental/DynamicMaterial/Source/DynamicMaterialEditor/Public/Components/DMMaterialStageThroughput.h
@@ -116,7 +116,9 @@ public:
 		int32 InThroughputInputIndex, FDMMaterialStageConnectorChannel& OutChannel, TArray<UMaterialExpression*>& OutExpressions) const;
 
 	/** This is called when the input value of this node's stage is updated. */
-	virtual void InputUpdated(int32 InThroughputInputIndex, EDMUpdateType InUpdateType) {}
+	virtual void OnInputUpdated(int32 InThroughputInputIndex, EDMUpdateType InUpdateType) {}
+
+	virtual void OnPostInputAdded(int32 InInputIdx) {}
 
 	//~ Begin UDMMaterialComponent
 	virtual FText GetComponentDescription() const override { return GetDescription(); }
diff --git a/Engine/Plugins/Experimental/DynamicMaterial/Source/DynamicMaterialEditor/Public/Components/DMMaterialStageThroughputLayerBlend.h b/Engine/Plugins/Experimental/DynamicMaterial/Source/DynamicMaterialEditor/Public/Components/DMMaterialStageThroughputLayerBlend.h
--- a/Engine/Plugins/Experimental/DynamicMaterial/Source/DynamicMaterialEditor/Public/Components/DMMaterialStageThroughputLayerBlend.h
+++ b/Engine/Plugins/Experimental/DynamicMaterial/Source/DynamicMaterialEditor/Public/Components/DMMaterialStageThroughputLayerBlend.h
@@ -39,7 +39,6 @@ public:
 	DYNAMICMATERIALEDITOR_API void SetMaskChannelOverride(EAvaColorChannel InMaskChannel);
 
 	//~ Begin UDMMaterialComponent
-	DYNAMICMATERIALEDITOR_API virtual void OnComponentAdded() override;
 	DYNAMICMATERIALEDITOR_API virtual void Update(UDMMaterialComponent* InSource, EDMUpdateType InUpdateType) override;
 	DYNAMICMATERIALEDITOR_API virtual void GenerateExpressions(const TSharedRef<FDMMaterialBuildState>& InBuildState) const override;
 	DYNAMICMATERIALEDITOR_API virtual void PostEditorDuplicate(UDynamicMaterialModel* InMaterialModel, UDMMaterialComponent* InParent) override;
@@ -64,6 +63,7 @@ public:
 		FDMMaterialStageConnectorChannel& OutChannel, TArray<UMaterialExpression*>& OutExpressions) const override;
 	DYNAMICMATERIALEDITOR_API virtual void ConnectOutputToInput(const TSharedRef<FDMMaterialBuildState>& InBuildState, int32 InInputIdx, 
 		int32 InExpressionInputIndex, UMaterialExpression* InSourceExpression, int32 InSourceOutputIndex, int32 InSourceOutputChannel) override;
+	virtual void OnPostInputAdded(int32 InInputIdx) override;
 	//~ End UDMMaterialStageThroughput
 
 	/** Resolves the Mask input. */
@@ -137,4 +137,8 @@ protected:
 	//~ Begin UDMMaterialStageThroughput
 	DYNAMICMATERIALEDITOR_API virtual void GeneratePreviewMaterial(UMaterial* InPreviewMaterial) override;
 	//~ End UDMMaterialStageThroughput
+
+	//~ Begin UDMMaterialComponent
+	virtual void OnComponentAdded() override;
+	//~ End UDMMaterialComponent
 };
diff --git a/Engine/Plugins/Experimental/DynamicMaterial/Source/DynamicMaterialEditor/Public/Components/MaterialStageExpressions/DMMSETextureSampleBase.h b/Engine/Plugins/Experimental/DynamicMaterial/Source/DynamicMaterialEditor/Public/Components/MaterialStageExpressions/DMMSETextureSampleBase.h
--- a/Engine/Plugins/Experimental/DynamicMaterial/Source/DynamicMaterialEditor/Public/Components/MaterialStageExpressions/DMMSETextureSampleBase.h
+++ b/Engine/Plugins/Experimental/DynamicMaterial/Source/DynamicMaterialEditor/Public/Components/MaterialStageExpressions/DMMSETextureSampleBase.h
@@ -35,7 +35,7 @@ public:
 	DYNAMICMATERIALEDITOR_API virtual bool CanChangeInputType(int32 InInputIndex) const override;
 	virtual bool SupportsLayerMaskTextureUVLink() const override { return true; }
 	virtual int32 GetLayerMaskTextureUVLinkInputIndex() const override { return 1; }
-	DYNAMICMATERIALEDITOR_API virtual void InputUpdated(int32 InInputIndex, EDMUpdateType InUpdateType) override;
+	DYNAMICMATERIALEDITOR_API virtual void OnInputUpdated(int32 InInputIndex, EDMUpdateType InUpdateType) override;
 	//~ End UDMMaterialStageThroughput
 
 	//~ Begin UObject
