Material Designer
- Fixed update code path for substages.
- Component preview now updates more efficiently.

diff --git a/Engine/Plugins/Experimental/DynamicMaterial/Source/DynamicMaterialEditor/Private/Components/DMMaterialSubStage.cpp b/Engine/Plugins/Experimental/DynamicMaterial/Source/DynamicMaterialEditor/Private/Components/DMMaterialSubStage.cpp
--- a/Engine/Plugins/Experimental/DynamicMaterial/Source/DynamicMaterialEditor/Private/Components/DMMaterialSubStage.cpp
+++ b/Engine/Plugins/Experimental/DynamicMaterial/Source/DynamicMaterialEditor/Private/Components/DMMaterialSubStage.cpp
@@ -78,6 +78,42 @@ bool UDMMaterialSubStage::IsCompatibleWithNextStage(const UDMMaterialStage* Next
 	return false;
 }
 
+void UDMMaterialSubStage::Update(UDMMaterialComponent* InSource, EDMUpdateType InUpdateType)
+{
+	if (!FDMUpdateGuard::CanUpdate())
+	{
+		return;
+	}
+
+	if (!IsComponentValid())
+	{
+		return;
+	}
+
+	if (HasComponentBeenRemoved())
+	{
+		return;
+	}
+
+	if (EnumHasAnyFlags(InUpdateType, EDMUpdateType::Structure))
+	{
+		MarkComponentDirty();
+		VerifyAllInputMaps();
+	}
+
+	// Skip UDMMaterialStage* update because we don't want to update other stages or layers.
+	Super::Super::Update(InSource, InUpdateType);
+
+	if (IsValid(ParentComponent))
+	{
+		ParentComponent->Update(InSource, InUpdateType);
+	}
+	else if (IsValid(ParentStage))
+	{
+		ParentStage->Update(InSource, InUpdateType);
+	}
+}
+
 FString UDMMaterialSubStage::GetComponentPathComponent() const
 {
 	// Skip stage renaming
diff --git a/Engine/Plugins/Experimental/DynamicMaterial/Source/DynamicMaterialEditor/Private/UI/Widgets/Visualizers/SDMMaterialComponentPreview.cpp b/Engine/Plugins/Experimental/DynamicMaterial/Source/DynamicMaterialEditor/Private/UI/Widgets/Visualizers/SDMMaterialComponentPreview.cpp
--- a/Engine/Plugins/Experimental/DynamicMaterial/Source/DynamicMaterialEditor/Private/UI/Widgets/Visualizers/SDMMaterialComponentPreview.cpp
+++ b/Engine/Plugins/Experimental/DynamicMaterial/Source/DynamicMaterialEditor/Private/UI/Widgets/Visualizers/SDMMaterialComponentPreview.cpp
@@ -37,20 +37,6 @@ SDMMaterialComponentPreview::~SDMMaterialComponentPreview()
 			EditorWidget->GetPreviewMaterialManager()->CreatePreviewMaterial(Component);
 		}
 	}
-
-	if (UDynamicMaterialModelBase* MaterialModelBase = MaterialModelBaseWeak.Get())
-	{
-		if (UDynamicMaterialModel* MaterialModel = Cast<UDynamicMaterialModel>(MaterialModelBase))
-		{
-			MaterialModel->GetOnValueUpdateDelegate().RemoveAll(this);
-			MaterialModel->GetOnTextureUVUpdateDelegate().RemoveAll(this);
-		}
-		else if (UDynamicMaterialModelDynamic* MaterialModelDynamic = Cast<UDynamicMaterialModelDynamic>(MaterialModelBase))
-		{
-			MaterialModelDynamic->GetOnValueDynamicUpdateDelegate().RemoveAll(this);
-			MaterialModelDynamic->GetOnTextureUVDynamicUpdateDelegate().RemoveAll(this);
-		}
-	}
 }
 
 void SDMMaterialComponentPreview::Construct(const FArguments& InArgs, const TSharedRef<SDMMaterialEditor>& InEditorWidget, UDMMaterialComponent* InComponent)
@@ -64,14 +50,10 @@ void SDMMaterialComponentPreview::Construct(const FArguments& InArgs, const TSha
 
 	if (UDynamicMaterialModel* MaterialModel = Cast<UDynamicMaterialModel>(MaterialModelBaseWeak.Get()))
 	{
-		MaterialModel->GetOnValueUpdateDelegate().AddSP(this, &SDMMaterialComponentPreview::OnValueUpdated);
-		MaterialModel->GetOnTextureUVUpdateDelegate().AddSP(this, &SDMMaterialComponentPreview::OnTextureUVUpdated);
 		MaterialModel->ApplyComponents(PreviewMaterialDynamicWeak.Get());
 	}
 	else if (UDynamicMaterialModelDynamic* MaterialModelDynamic = Cast<UDynamicMaterialModelDynamic>(MaterialModelBaseWeak.Get()))
 	{
-		MaterialModelDynamic->GetOnValueDynamicUpdateDelegate().AddSP(this, &SDMMaterialComponentPreview::OnValueDynamicUpdated);
-		MaterialModelDynamic->GetOnTextureUVDynamicUpdateDelegate().AddSP(this, &SDMMaterialComponentPreview::OnTextureUVDynamicUpdated);
 		MaterialModelDynamic->ApplyComponents(PreviewMaterialDynamicWeak.Get());
 	}
 	else
@@ -133,58 +115,45 @@ void SDMMaterialComponentPreview::OnComponentUpdated(UDMMaterialComponent* InCom
 		PreviewMaterialBase = EditorWidget->GetPreviewMaterialManager()->CreatePreviewMaterial(Stage);
 	}
 
-	if (!EnumHasAnyFlags(InUpdateType, EDMUpdateType::Structure))
-	{
-		return;
-	}
-
-	Stage->GeneratePreviewMaterial(PreviewMaterialBase);
+	UMaterialInstanceDynamic* MID = PreviewMaterialDynamicWeak.Get();
 
-	EditorWidget->GetPreviewMaterialManager()->FreePreviewMaterialDynamic(PreviewMaterialBase);
-	PreviewMaterialDynamicWeak = EditorWidget->GetPreviewMaterialManager()->CreatePreviewMaterialDynamic(PreviewMaterialBase);
-
-	UDynamicMaterialModelBase* MaterialModelBase = EditorWidget->GetMaterialModelBase();
-
-	if (UDynamicMaterialModel* MaterialModel = Cast<UDynamicMaterialModel>(MaterialModelBase))
+	if (!MID || !EnumHasAnyFlags(InUpdateType, EDMUpdateType::Structure))
 	{
-		MaterialModel->ApplyComponents(PreviewMaterialDynamicWeak.Get());
+		if (UDMMaterialValue* Value = Cast<UDMMaterialValue>(InSource))
+		{
+			Value->SetMIDParameter(MID);
+		}
+		else if (UDMMaterialValueDynamic* ValueDynamic = Cast<UDMMaterialValueDynamic>(InSource))
+		{
+			ValueDynamic->SetMIDParameter(MID);
+		}
+		else if (UDMTextureUV* TextureUV = Cast<UDMTextureUV>(InSource))
+		{
+			TextureUV->SetMIDParameters(MID);
+		}
+		else if (UDMTextureUVDynamic* TextureUVDynamic = Cast<UDMTextureUVDynamic>(InSource))
+		{
+			TextureUVDynamic->SetMIDParameters(MID);
+		}
 	}
-	else if (UDynamicMaterialModelDynamic* MaterialModelDynamic = Cast<UDynamicMaterialModelDynamic>(MaterialModelBase))
+	else
 	{
-		MaterialModelDynamic->ApplyComponents(PreviewMaterialDynamicWeak.Get());
-	}
-
-	Brush.SetMaterial(PreviewMaterialDynamicWeak.Get());
-}
+		Stage->GeneratePreviewMaterial(PreviewMaterialBase);
 
-void SDMMaterialComponentPreview::OnValueUpdated(UDynamicMaterialModel* InMaterialModel, UDMMaterialValue* InValue)
-{
-	if (UMaterialInstanceDynamic* PreviewMaterialDynamic = PreviewMaterialDynamicWeak.Get())
-	{
-		InValue->SetMIDParameter(PreviewMaterialDynamic);
-	}
-}
+		EditorWidget->GetPreviewMaterialManager()->FreePreviewMaterialDynamic(PreviewMaterialBase);
+		PreviewMaterialDynamicWeak = EditorWidget->GetPreviewMaterialManager()->CreatePreviewMaterialDynamic(PreviewMaterialBase);
 
-void SDMMaterialComponentPreview::OnTextureUVUpdated(UDynamicMaterialModel* InMaterialModel, UDMTextureUV* InTextureUV)
-{
-	if (UMaterialInstanceDynamic* PreviewMaterialDynamic = PreviewMaterialDynamicWeak.Get())
-	{
-		InTextureUV->SetMIDParameters(PreviewMaterialDynamic);
-	}
-}
+		UDynamicMaterialModelBase* MaterialModelBase = EditorWidget->GetMaterialModelBase();
 
-void SDMMaterialComponentPreview::OnValueDynamicUpdated(UDynamicMaterialModelDynamic* InMaterialModel, UDMMaterialValueDynamic* InValueDynamic)
-{
-	if (UMaterialInstanceDynamic* PreviewMaterialDynamic = PreviewMaterialDynamicWeak.Get())
-	{
-		InValueDynamic->SetMIDParameter(PreviewMaterialDynamic);
-	}
-}
+		if (UDynamicMaterialModel* MaterialModel = Cast<UDynamicMaterialModel>(MaterialModelBase))
+		{
+			MaterialModel->ApplyComponents(PreviewMaterialDynamicWeak.Get());
+		}
+		else if (UDynamicMaterialModelDynamic* MaterialModelDynamic = Cast<UDynamicMaterialModelDynamic>(MaterialModelBase))
+		{
+			MaterialModelDynamic->ApplyComponents(PreviewMaterialDynamicWeak.Get());
+		}
 
-void SDMMaterialComponentPreview::OnTextureUVDynamicUpdated(UDynamicMaterialModelDynamic* InMaterialModel, UDMTextureUVDynamic* InTextureUVDynamic)
-{
-	if (UMaterialInstanceDynamic* PreviewMaterialDynamic = PreviewMaterialDynamicWeak.Get())
-	{
-		InTextureUVDynamic->SetMIDParameters(PreviewMaterialDynamic);
+		Brush.SetMaterial(PreviewMaterialDynamicWeak.Get());
 	}
 }
diff --git a/Engine/Plugins/Experimental/DynamicMaterial/Source/DynamicMaterialEditor/Private/UI/Widgets/Visualizers/SDMMaterialComponentPreview.h b/Engine/Plugins/Experimental/DynamicMaterial/Source/DynamicMaterialEditor/Private/UI/Widgets/Visualizers/SDMMaterialComponentPreview.h
--- a/Engine/Plugins/Experimental/DynamicMaterial/Source/DynamicMaterialEditor/Private/UI/Widgets/Visualizers/SDMMaterialComponentPreview.h
+++ b/Engine/Plugins/Experimental/DynamicMaterial/Source/DynamicMaterialEditor/Private/UI/Widgets/Visualizers/SDMMaterialComponentPreview.h
@@ -47,12 +47,4 @@ protected:
 	FSlateMaterialBrush Brush;
 
 	void OnComponentUpdated(UDMMaterialComponent* InComponent, UDMMaterialComponent* InSource, EDMUpdateType InUpdateType);
-
-	void OnValueUpdated(UDynamicMaterialModel* InMaterialModel, UDMMaterialValue* InValue);
-
-	void OnTextureUVUpdated(UDynamicMaterialModel* InMaterialModel, UDMTextureUV* InTextureUV);
-
-	void OnValueDynamicUpdated(UDynamicMaterialModelDynamic* InMaterialModel, UDMMaterialValueDynamic* InValueDynamic);
-
-	void OnTextureUVDynamicUpdated(UDynamicMaterialModelDynamic* InMaterialModel, UDMTextureUVDynamic* InTextureUVDynamic);
 };
diff --git a/Engine/Plugins/Experimental/DynamicMaterial/Source/DynamicMaterialEditor/Public/Components/DMMaterialSubStage.h b/Engine/Plugins/Experimental/DynamicMaterial/Source/DynamicMaterialEditor/Public/Components/DMMaterialSubStage.h
--- a/Engine/Plugins/Experimental/DynamicMaterial/Source/DynamicMaterialEditor/Public/Components/DMMaterialSubStage.h
+++ b/Engine/Plugins/Experimental/DynamicMaterial/Source/DynamicMaterialEditor/Public/Components/DMMaterialSubStage.h
@@ -37,6 +37,7 @@ public:
 	//~ End UDMMaterialStage
 
 	//~ Begin UDMMaterialComponent
+	DYNAMICMATERIALEDITOR_API virtual void Update(UDMMaterialComponent* InSource, EDMUpdateType InUpdateType) override;
 	DYNAMICMATERIALEDITOR_API virtual FString GetComponentPathComponent() const override;
 	DYNAMICMATERIALEDITOR_API virtual UDMMaterialComponent* GetParentComponent() const override;
 	DYNAMICMATERIALEDITOR_API virtual void PostEditorDuplicate(UDynamicMaterialModel* InMaterialModel, UDMMaterialComponent* InParent) override;
