Material Designer
- You can now change the mask source again.

diff --git a/Engine/Plugins/Experimental/DynamicMaterial/Source/DynamicMaterialEditor/Private/UI/PropertyGenerators/DMStagePropertyRowGenerator.cpp b/Engine/Plugins/Experimental/DynamicMaterial/Source/DynamicMaterialEditor/Private/UI/PropertyGenerators/DMStagePropertyRowGenerator.cpp
--- a/Engine/Plugins/Experimental/DynamicMaterial/Source/DynamicMaterialEditor/Private/UI/PropertyGenerators/DMStagePropertyRowGenerator.cpp
+++ b/Engine/Plugins/Experimental/DynamicMaterial/Source/DynamicMaterialEditor/Private/UI/PropertyGenerators/DMStagePropertyRowGenerator.cpp
@@ -57,27 +57,24 @@ void FDMStagePropertyRowGenerator::AddComponentProperties(const TSharedRef<SDMMa
 		return;
 	}
 
-	if (Stage->CanChangeSource())
-	{
-		bool bIsDynamic = false;
+	bool bIsDynamic = false;
 
-		if (TSharedPtr<SDMMaterialEditor> EditorWidget = InComponentEditorWidget->GetEditorWidget())
+	if (TSharedPtr<SDMMaterialEditor> EditorWidget = InComponentEditorWidget->GetEditorWidget())
+	{
+		if (UDynamicMaterialModelBase* MaterialModelBase = EditorWidget->GetMaterialModelBase())
 		{
-			if (UDynamicMaterialModelBase* MaterialModelBase = EditorWidget->GetMaterialModelBase())
-			{
-				bIsDynamic = !MaterialModelBase->IsA<UDynamicMaterialModel>();
-			}
+			bIsDynamic = !MaterialModelBase->IsA<UDynamicMaterialModel>();
 		}
-
-		FDMPropertyHandle& SourceHandle = InOutPropertyRows.AddDefaulted_GetRef();
-		SourceHandle.ValueName = TEXT("SourceType");
-		SourceHandle.NameOverride = LOCTEXT("SourceType", "Source Type");
-		SourceHandle.bEnabled = !bIsDynamic;
-		SourceHandle.ValueWidget = CreateSourceTypeEditWidget(InComponentEditorWidget, Stage);
-		SourceHandle.ResetToDefaultOverride = FResetToDefaultOverride::Hide(true);
-		SourceHandle.Priority = EDMPropertyHandlePriority::High;
 	}
 
+	FDMPropertyHandle& SourceHandle = InOutPropertyRows.AddDefaulted_GetRef();
+	SourceHandle.ValueName = TEXT("SourceType");
+	SourceHandle.NameOverride = LOCTEXT("SourceType", "Source Type");
+	SourceHandle.bEnabled = !bIsDynamic;
+	SourceHandle.ValueWidget = CreateSourceTypeEditWidget(InComponentEditorWidget, Stage);
+	SourceHandle.ResetToDefaultOverride = FResetToDefaultOverride::Hide(true);
+	SourceHandle.Priority = EDMPropertyHandlePriority::High;
+
 	FDynamicMaterialEditorModule::GeneratorComponentPropertyRows(InComponentEditorWidget, Source, InOutPropertyRows, InOutProcessedObjects);
 	FDMComponentPropertyRowGenerator::AddComponentProperties(InComponentEditorWidget, Stage, InOutPropertyRows, InOutProcessedObjects);
 }
