Material Designer
- Fixed texture values causing a structure update when render targets produced non-structure updates.

diff --git a/Engine/Plugins/Experimental/DynamicMaterial/Source/DynamicMaterial/Private/Components/MaterialValues/DMMaterialValueTexture.cpp b/Engine/Plugins/Experimental/DynamicMaterial/Source/DynamicMaterial/Private/Components/MaterialValues/DMMaterialValueTexture.cpp
--- a/Engine/Plugins/Experimental/DynamicMaterial/Source/DynamicMaterial/Private/Components/MaterialValues/DMMaterialValueTexture.cpp
+++ b/Engine/Plugins/Experimental/DynamicMaterial/Source/DynamicMaterial/Private/Components/MaterialValues/DMMaterialValueTexture.cpp
@@ -140,26 +140,9 @@ void UDMMaterialValueTexture::PostEditChangeProperty(FPropertyChangedEvent& InPr
 		return;
 	}
 
-	for (const FName& EditableProperty : EditableProperties)
+	if (MemberPropertyName == ValueName)
 	{
-		if (EditableProperty == MemberPropertyName)
-		{
-			if (EditableProperty == ValueName)
-			{
-				const TextureCompressionSettings CurrentCompression = OldValue ? OldValue->CompressionSettings.GetValue() : TextureCompressionSettings::TC_MAX;
-				const TextureCompressionSettings NewCompression = Value ? Value->CompressionSettings.GetValue() : TextureCompressionSettings::TC_MAX;
-
-				const EDMUpdateType UpdateType = CurrentCompression == NewCompression
-					? EDMUpdateType::Structure
-					: EDMUpdateType::Value;
-
-				OnValueChanged(EDMUpdateType::Value | EDMUpdateType::AllowParentUpdate);
-				return;
-			}
-
-			OnValueChanged(EDMUpdateType::Structure | EDMUpdateType::AllowParentUpdate);
-			return;
-		}
+		OnValueChanged(EDMUpdateType::Value | EDMUpdateType::AllowParentUpdate);
 	}
 }
 #endif
