Material Designer
- Property previews now render correctly.

diff --git a/Engine/Plugins/Experimental/DynamicMaterial/Source/DynamicMaterialEditor/Private/Components/DMMaterialProperty.cpp b/Engine/Plugins/Experimental/DynamicMaterial/Source/DynamicMaterialEditor/Private/Components/DMMaterialProperty.cpp
--- a/Engine/Plugins/Experimental/DynamicMaterial/Source/DynamicMaterialEditor/Private/Components/DMMaterialProperty.cpp
+++ b/Engine/Plugins/Experimental/DynamicMaterial/Source/DynamicMaterialEditor/Private/Components/DMMaterialProperty.cpp
@@ -16,7 +16,9 @@
 #include "Components/MaterialValues/DMMaterialValueTexture.h"
 #include "DMComponentPath.h"
 #include "DMValueDefinition.h"
+#include "DynamicMaterialEditorModule.h"
 #include "DynamicMaterialEditorSettings.h"
+#include "Materials/Material.h"
 #include "Materials/MaterialExpressionConstant.h"
 #include "Materials/MaterialExpressionConstant2Vector.h"
 #include "Materials/MaterialExpressionConstant3Vector.h"
@@ -593,6 +595,116 @@ void UDMMaterialProperty::AddOutputProcessor(const TSharedRef<FDMMaterialBuildSt
 	MaterialPropertyPtr->OutputIndex = 0;
 }
 
+void UDMMaterialProperty::GeneratePreviewMaterial(UMaterial* InPreviewMaterial)
+{
+	if (!IsComponentValid())
+	{
+		return;
+	}
+
+	UDynamicMaterialModelEditorOnlyData* EditorOnlyData = GetMaterialModelEditorOnlyData();
+
+	if (!EditorOnlyData)
+	{
+		return;
+	}
+
+	InPreviewMaterial->BlendMode = EBlendMode::BLEND_Translucent;
+
+	TSharedRef<FDMMaterialBuildState> BuildState = EditorOnlyData->CreateBuildState(InPreviewMaterial);
+	BuildState->SetPreviewMaterial();
+
+	UE_LOG(LogDynamicMaterialEditor, Display, TEXT("Building Material Designer Property Preview (%s)..."), *GetName());
+
+	UE::DynamicMaterial::ForEachMaterialPropertyType(
+		[this, &BuildState](EDMMaterialPropertyType InType)
+		{
+			if (FExpressionInput* Input = BuildState->GetMaterialProperty(InType))
+			{
+				Input->Expression = nullptr;
+			}
+
+			return EDMIterationResult::Continue;
+		});
+
+
+	if (!IsEnabled() || !EditorOnlyData->GetSlotForMaterialProperty(MaterialProperty))
+	{
+		return;
+	}
+
+	GenerateExpressions(BuildState);
+
+	AddAlphaMultiplier(BuildState);
+
+	// Move output to the emissive channel.
+	if (MaterialProperty != EDMMaterialPropertyType::EmissiveColor)
+	{
+		if (FExpressionInput* EmissiveExpressionInput = BuildState->GetMaterialProperty(EDMMaterialPropertyType::EmissiveColor))
+		{
+			FExpressionInput* MyFirstExpressionInput = nullptr;
+
+			UE::DynamicMaterial::ForEachMaterialPropertyType(
+				[this, &BuildState, &MyFirstExpressionInput](EDMMaterialPropertyType InType)
+				{
+					using namespace UE::DynamicMaterialEditor::Private;
+
+					MyFirstExpressionInput = BuildState->GetMaterialProperty(InType);
+
+					if (!MyFirstExpressionInput || MyFirstExpressionInput->Expression == nullptr)
+					{
+						MyFirstExpressionInput = nullptr;
+						return EDMIterationResult::Continue;
+					}
+
+					return EDMIterationResult::Break;
+				});
+
+			if (MyFirstExpressionInput)
+			{
+				// Swap inputs to emissive channel
+				EmissiveExpressionInput->Expression = MyFirstExpressionInput->Expression;
+				EmissiveExpressionInput->Mask = MyFirstExpressionInput->Mask;
+				EmissiveExpressionInput->MaskR = MyFirstExpressionInput->MaskR;
+				EmissiveExpressionInput->MaskG = MyFirstExpressionInput->MaskG;
+				EmissiveExpressionInput->MaskB = MyFirstExpressionInput->MaskB;
+				EmissiveExpressionInput->MaskA = MyFirstExpressionInput->MaskA;
+				EmissiveExpressionInput->OutputIndex = MyFirstExpressionInput->OutputIndex;
+
+				MyFirstExpressionInput->Expression = nullptr;
+			}
+		}
+	}
+
+	// Attempt to create an opacity channel.
+	FExpressionInput* OpacityExpression = BuildState->GetMaterialProperty(EDMMaterialPropertyType::Opacity);
+
+	if (OpacityExpression && !OpacityExpression->Expression)
+	{
+		UDMMaterialSlot* MySlot = EditorOnlyData->GetSlotForMaterialProperty(MaterialProperty);
+		UMaterialExpression* OpacityOutputNode;
+		int32 OutputIndex;
+		int32 OutputChannel;
+		GenerateOpacityExpressions(BuildState, MySlot, MaterialProperty, OpacityOutputNode, OutputIndex, OutputChannel);
+
+		if (OpacityOutputNode)
+		{
+			OpacityExpression->Expression = OpacityOutputNode;
+			OpacityExpression->OutputIndex = 0;
+			OpacityExpression->Mask = 0;
+
+			if (OutputChannel != FDMMaterialStageConnectorChannel::WHOLE_CHANNEL)
+			{
+				OpacityExpression->Mask = 1;
+				OpacityExpression->MaskR = !!(OutputChannel & FDMMaterialStageConnectorChannel::FIRST_CHANNEL);
+				OpacityExpression->MaskG = !!(OutputChannel & FDMMaterialStageConnectorChannel::SECOND_CHANNEL);
+				OpacityExpression->MaskB = !!(OutputChannel & FDMMaterialStageConnectorChannel::THIRD_CHANNEL);
+				OpacityExpression->MaskA = !!(OutputChannel & FDMMaterialStageConnectorChannel::FOURTH_CHANNEL);
+			}
+		}
+	}
+}
+
 void UDMMaterialProperty::Update(UDMMaterialComponent* InSource, EDMUpdateType InUpdateType)
 {
 	if (!FDMUpdateGuard::CanUpdate())
diff --git a/Engine/Plugins/Experimental/DynamicMaterial/Source/DynamicMaterialEditor/Private/Components/DMMaterialSlot.cpp b/Engine/Plugins/Experimental/DynamicMaterial/Source/DynamicMaterialEditor/Private/Components/DMMaterialSlot.cpp
--- a/Engine/Plugins/Experimental/DynamicMaterial/Source/DynamicMaterialEditor/Private/Components/DMMaterialSlot.cpp
+++ b/Engine/Plugins/Experimental/DynamicMaterial/Source/DynamicMaterialEditor/Private/Components/DMMaterialSlot.cpp
@@ -6,25 +6,19 @@
 #include "Components/DMMaterialLayer_Deprecated.h"
 #include "Components/DMMaterialProperty.h"
 #include "Components/DMMaterialStage.h"
-#include "Components/DMMaterialStageThroughputLayerBlend.h"
 #include "Components/DMMaterialSubStage.h"
 #include "CoreGlobals.h"
 #include "DMComponentPath.h"
 #include "DMDefs.h"
-#include "DynamicMaterialEditorModule.h"
 #include "Materials/Material.h"
 #include "Materials/MaterialExpression.h"
-#include "Materials/MaterialExpressionMax.h"
-#include "Materials/MaterialExpressionMultiply.h"
 #include "MaterialValueType.h"
 #include "Model/DMMaterialBuildState.h"
-#include "Model/DMMaterialBuildUtils.h"
 #include "Model/DynamicMaterialModel.h"
 #include "Model/DynamicMaterialModelEditorOnlyData.h"
 #include "UObject/Package.h"
 #include "UObject/UObjectGlobals.h"
 #include "Utils/DMPrivate.h"
-#include "Utils/DMUtils.h"
 
 #define LOCTEXT_NAMESPACE "DMMaterialStage"
 
@@ -826,415 +820,6 @@ bool UDMMaterialSlot::UnreferencedBySlot(UDMMaterialSlot* InOtherSlot)
 	}
 }
 
-void UDMMaterialSlot::GeneralPreviewMaterial(UMaterial* InPreviewMaterial, EDMMaterialLayerStage InLayerStage)
-{
-	if (!IsComponentValid())
-	{
-		return;
-	}
-
-	bool bHasActivateStage = false;
-
-	for (const TObjectPtr<UDMMaterialLayerObject>& Layer : LayerObjects)
-	{
-		if (!Layer->IsEnabled())
-		{
-			continue;
-		}
-
-		if (Layer->IsStageEnabled(EDMMaterialLayerStage::Base))
-		{
-			bHasActivateStage = true;
-			break;
-		}
-	}
-
-	if (!bHasActivateStage)
-	{
-		return;
-	}
-
-	check(InLayerStage == EDMMaterialLayerStage::Base || InLayerStage == EDMMaterialLayerStage::Mask);
-
-	UDynamicMaterialModelEditorOnlyData* ModelEditorOnlyData = GetMaterialModelEditorOnlyData();
-	check(ModelEditorOnlyData);
-
-	if (InLayerStage == EDMMaterialLayerStage::Base)
-	{
-		TSharedRef<FDMMaterialBuildState> BuildState = ModelEditorOnlyData->CreateBuildState(InPreviewMaterial);
-		BuildState->SetPreviewMaterial();
-
-		GenerateExpressions(BuildState);
-		UpdateBasePreviewMaterial(BuildState);
-	}
-	else
-	{
-		TSharedRef<FDMMaterialBuildState> BuildState = ModelEditorOnlyData->CreateBuildState(InPreviewMaterial);
-		BuildState->SetPreviewMaterial();
-
-		GenerateExpressions(BuildState);
-		UpdateMaskPreviewMaterial(BuildState);
-	}
-}
-
-void UDMMaterialSlot::UpdateBasePreviewMaterial(const TSharedRef<FDMMaterialBuildState>& InBuildState)
-{
-	if (!IsComponentValid())
-	{
-		return;
-	}
-
-	UE_LOG(LogDynamicMaterialEditor, Display, TEXT("Building Material Designer Slot Base Preview (%s)..."), *GetName());
-
-	UMaterial* PreviewMaterial = InBuildState->GetDynamicMaterial();
-
-	PreviewMaterial->GetEditorOnlyData()->EmissiveColor.Expression = nullptr;
-	PreviewMaterial->GetEditorOnlyData()->EmissiveColor.OutputIndex = 0;
-
-	const TMap<EDMMaterialPropertyType, TArray<UMaterialExpression*>>& SlotPropertyExpressions = InBuildState->GetSlotPropertyExpressions(this);
-
-	// Single material property, connect it up to emissive and output it.
-	if (SlotPropertyExpressions.Num() == 1)
-	{
-		for (const TPair<EDMMaterialPropertyType, TArray<UMaterialExpression*>>& Pair : SlotPropertyExpressions)
-		{
-			UpdateBasePreviewMaterialProperty(InBuildState, Pair.Key);
-		}
-	}
-	else
-	{
-		UpdateBasePreviewMaterialFull(InBuildState);
-	}
-}
-
-void UDMMaterialSlot::UpdateBasePreviewMaterialProperty(const TSharedRef<FDMMaterialBuildState>& InBuildState, EDMMaterialPropertyType InBaseProperty)
-{
-	UpdatePreviewMaterialProperty(InBuildState, InBaseProperty);
-}
-
-void UDMMaterialSlot::UpdateBasePreviewMaterialFull(const TSharedRef<FDMMaterialBuildState>& InBuildState)
-{
-	if (!IsComponentValid())
-	{
-		return;
-	}
-
-	UMaterial* PreviewMaterial = InBuildState->GetDynamicMaterial();
-
-	UDynamicMaterialModelEditorOnlyData* ModelEditorOnlyData = GetMaterialModelEditorOnlyData();
-	check(ModelEditorOnlyData);
-
-	const TMap<EDMMaterialPropertyType, UDMMaterialProperty*> MaterialProperties = ModelEditorOnlyData->GetMaterialProperties();
-	const TMap<EDMMaterialPropertyType, TArray<UMaterialExpression*>>& SlotPropertyExpressions = InBuildState->GetSlotPropertyExpressions(this);
-
-	for (const TPair<EDMMaterialPropertyType, UDMMaterialProperty*>& Pair : MaterialProperties)
-	{
-		if (!Pair.Value->IsMaterialPin())
-		{
-			continue;
-		}
-
-		// For now we don't have channel remapping!
-		FExpressionInput* MaterialPropertyPtr = InBuildState->GetMaterialProperty(Pair.Key);
-		check(MaterialPropertyPtr);
-
-		MaterialPropertyPtr->Expression = nullptr;
-		MaterialPropertyPtr->OutputIndex = 0;
-
-		if (!SlotPropertyExpressions.Contains(Pair.Key))
-		{
-			continue;
-		}
-
-		const TArray<UMaterialExpression*>& PropertyExpressions = SlotPropertyExpressions[Pair.Key];
-
-		if (PropertyExpressions.IsEmpty())
-		{
-			continue;
-		}
-
-		MaterialPropertyPtr->Expression = PropertyExpressions.Last();
-
-		if (Pair.Value->GetInputConnectionMap().Channels.IsEmpty() == false)
-		{
-			MaterialPropertyPtr->OutputIndex = Pair.Value->GetInputConnectionMap().Channels[0].OutputIndex;
-		}
-		else
-		{
-			MaterialPropertyPtr->OutputIndex = 0;
-		}
-	}
-}
-
-void UDMMaterialSlot::UpdateMaskPreviewMaterial(const TSharedRef<FDMMaterialBuildState>& InBuildState)
-{
-	if (!IsComponentValid())
-	{
-		return;
-	}
-
-	UE_LOG(LogDynamicMaterialEditor, Display, TEXT("Building Material Designer Slot Mask Preview (%s)..."), *GetName());
-
-	UMaterial* PreviewMaterial = InBuildState->GetDynamicMaterial();
-
-	PreviewMaterial->GetEditorOnlyData()->EmissiveColor.Expression = nullptr;
-	PreviewMaterial->GetEditorOnlyData()->EmissiveColor.OutputIndex = 0;
-
-	UDynamicMaterialModelEditorOnlyData* ModelEditorOnlyData = GetMaterialModelEditorOnlyData();
-	check(ModelEditorOnlyData);
-
-	const TMap<EDMMaterialPropertyType, TArray<UMaterialExpression*>>& SlotPropertyExpressions = InBuildState->GetSlotPropertyExpressions(this);
-
-	if (SlotPropertyExpressions.Num() == 1)
-	{
-		for (const TPair<EDMMaterialPropertyType, TArray<UMaterialExpression*>>& Pair : SlotPropertyExpressions)
-		{
-			UpdateMaskPreviewMaterialMaskCombination(InBuildState, Pair.Key);
-		}
-	}
-	else if (SlotPropertyExpressions.Contains(EDMMaterialPropertyType::Opacity))
-	{
-		UpdateMaskPreviewMaterialProperty(InBuildState, EDMMaterialPropertyType::Opacity);
-	}
-	else if (SlotPropertyExpressions.Contains(EDMMaterialPropertyType::OpacityMask))
-	{
-		UpdateMaskPreviewMaterialProperty(InBuildState, EDMMaterialPropertyType::OpacityMask);
-	}
-	else if (SlotPropertyExpressions.Contains(EDMMaterialPropertyType::BaseColor))
-	{
-		UpdateMaskPreviewMaterialMaskCombination(InBuildState, EDMMaterialPropertyType::BaseColor);
-	}
-	else if (SlotPropertyExpressions.Contains(EDMMaterialPropertyType::EmissiveColor))
-	{
-		UpdateMaskPreviewMaterialMaskCombination(InBuildState, EDMMaterialPropertyType::EmissiveColor);
-	}
-}
-
-void UDMMaterialSlot::UpdateMaskPreviewMaterialProperty(const TSharedRef<FDMMaterialBuildState>& InBuildState, EDMMaterialPropertyType InMaskProperty)
-{
-	UpdatePreviewMaterialProperty(InBuildState, InMaskProperty);
-}
-
-void UDMMaterialSlot::UpdateMaskPreviewMaterialMaskCombination(const TSharedRef<FDMMaterialBuildState>& InBuildState, EDMMaterialPropertyType InMaskProperty)
-{
-	if (!IsComponentValid())
-	{
-		return;
-	}
-
-	UMaterial* PreviewMaterial = InBuildState->GetDynamicMaterial();
-
-	int32 OutOutputChannel = FDMMaterialStageConnectorChannel::WHOLE_CHANNEL;
-
-	for (const TObjectPtr<UDMMaterialLayerObject>& Layer : LayerObjects)
-	{
-		// Although we are working with masks, if the base is disabled, this is handled by the GenerateExpressions
-		// of the LayerBlend code (to multiply alpha together, instead of maxing it).
-		if (Layer->GetMaterialProperty() != InMaskProperty || !Layer->AreAllStagesEnabled(EDMMaterialLayerStage::All))
-		{
-			continue;
-		}
-
-		UDMMaterialStage* BaseStage = Layer->GetStage(EDMMaterialLayerStage::Base);
-		UDMMaterialStage* MaskStage = Layer->GetStage(EDMMaterialLayerStage::Mask);
-
-		MaskStage->GenerateExpressions(InBuildState);
-		UDMMaterialStageThroughputLayerBlend* LayerBlend = Cast<UDMMaterialStageThroughputLayerBlend>(MaskStage->GetSource());
-
-		if (!LayerBlend)
-		{
-			continue;
-		}
-
-		UMaterialExpression* MaskOutputExpression;
-		int32 MaskOutputIndex;
-		int32 MaskOutputChannel;
-		LayerBlend->GetMaskOutput(InBuildState, MaskOutputExpression, MaskOutputIndex, MaskOutputChannel);
-
-		if (!MaskOutputExpression)
-		{
-			continue;
-		}
-
-		if (LayerBlend->UsePremultiplyAlpha())
-		{
-			if (UDMMaterialStageSource* Source = BaseStage->GetSource())
-			{
-				UMaterialExpression* LayerAlphaOutputExpression;
-				int32 LayerAlphaOutputIndex;
-				int32 LayerAlphaOutputChannel;
-
-				Source->GetMaskAlphaBlendNode(InBuildState, LayerAlphaOutputExpression, LayerAlphaOutputIndex, LayerAlphaOutputChannel);
-
-				if (LayerAlphaOutputExpression)
-				{
-					UMaterialExpressionMultiply* AlphaMultiply = InBuildState->GetBuildUtils().CreateExpression<UMaterialExpressionMultiply>(UE_DM_NodeComment_Default);
-
-					AlphaMultiply->A.Expression = MaskOutputExpression;
-					AlphaMultiply->A.OutputIndex = MaskOutputIndex;
-					AlphaMultiply->A.Mask = 0;
-
-					if (MaskOutputChannel != FDMMaterialStageConnectorChannel::WHOLE_CHANNEL)
-					{
-						AlphaMultiply->A.Mask = 1;
-						AlphaMultiply->A.MaskR = !!(MaskOutputChannel & FDMMaterialStageConnectorChannel::FIRST_CHANNEL);
-						AlphaMultiply->A.MaskG = !!(MaskOutputChannel & FDMMaterialStageConnectorChannel::SECOND_CHANNEL);
-						AlphaMultiply->A.MaskB = !!(MaskOutputChannel & FDMMaterialStageConnectorChannel::THIRD_CHANNEL);
-						AlphaMultiply->A.MaskA = !!(MaskOutputChannel & FDMMaterialStageConnectorChannel::FOURTH_CHANNEL);
-					}
-
-					AlphaMultiply->B.Expression = LayerAlphaOutputExpression;
-					AlphaMultiply->B.OutputIndex = LayerAlphaOutputIndex;
-					AlphaMultiply->B.Mask = 0;
-
-					if (LayerAlphaOutputChannel != FDMMaterialStageConnectorChannel::WHOLE_CHANNEL)
-					{
-						AlphaMultiply->B.Mask = 1;
-						AlphaMultiply->B.MaskR = !!(LayerAlphaOutputChannel & FDMMaterialStageConnectorChannel::FIRST_CHANNEL);
-						AlphaMultiply->B.MaskG = !!(LayerAlphaOutputChannel & FDMMaterialStageConnectorChannel::SECOND_CHANNEL);
-						AlphaMultiply->B.MaskB = !!(LayerAlphaOutputChannel & FDMMaterialStageConnectorChannel::THIRD_CHANNEL);
-						AlphaMultiply->B.MaskA = !!(LayerAlphaOutputChannel & FDMMaterialStageConnectorChannel::FOURTH_CHANNEL);
-					}
-
-					MaskOutputExpression = AlphaMultiply;
-					MaskOutputIndex = 0;
-					MaskOutputChannel = FDMMaterialStageConnectorChannel::WHOLE_CHANNEL;
-				}
-			}
-		}
-
-		if (PreviewMaterial->GetEditorOnlyData()->EmissiveColor.Expression == nullptr)
-		{
-			PreviewMaterial->GetEditorOnlyData()->EmissiveColor.Expression = MaskOutputExpression;
-
-			// The first output will use the node's output info.
-			PreviewMaterial->GetEditorOnlyData()->EmissiveColor.OutputIndex = MaskOutputIndex;
-			OutOutputChannel = MaskOutputChannel;
-			continue;
-		}
-
-		UMaterialExpressionMax* Max = InBuildState->GetBuildUtils().CreateExpression<UMaterialExpressionMax>(UE_DM_NodeComment_Default);
-		check(Max);
-
-		Max->A.Expression = PreviewMaterial->GetEditorOnlyData()->EmissiveColor.Expression;
-		Max->A.OutputIndex = PreviewMaterial->GetEditorOnlyData()->EmissiveColor.OutputIndex;
-		Max->A.Mask = 0;
-
-		if (OutOutputChannel != FDMMaterialStageConnectorChannel::WHOLE_CHANNEL)
-		{
-			Max->A.Mask = 1;
-			Max->A.MaskR = !!(OutOutputChannel & FDMMaterialStageConnectorChannel::FIRST_CHANNEL);
-			Max->A.MaskG = !!(OutOutputChannel & FDMMaterialStageConnectorChannel::SECOND_CHANNEL);
-			Max->A.MaskB = !!(OutOutputChannel & FDMMaterialStageConnectorChannel::THIRD_CHANNEL);
-			Max->A.MaskA = !!(OutOutputChannel & FDMMaterialStageConnectorChannel::FOURTH_CHANNEL);
-		}
-
-		Max->B.Expression = MaskOutputExpression;
-		Max->B.OutputIndex = MaskOutputIndex;
-		Max->B.Mask = 0;
-
-		if (MaskOutputChannel != FDMMaterialStageConnectorChannel::WHOLE_CHANNEL)
-		{
-			Max->B.Mask = 1;
-			Max->B.MaskR = !!(MaskOutputChannel & FDMMaterialStageConnectorChannel::FIRST_CHANNEL);
-			Max->B.MaskG = !!(MaskOutputChannel & FDMMaterialStageConnectorChannel::SECOND_CHANNEL);
-			Max->B.MaskB = !!(MaskOutputChannel & FDMMaterialStageConnectorChannel::THIRD_CHANNEL);
-			Max->B.MaskA = !!(MaskOutputChannel & FDMMaterialStageConnectorChannel::FOURTH_CHANNEL);
-		}
-
-		PreviewMaterial->GetEditorOnlyData()->EmissiveColor.Expression = Max;
-
-		// If we have to combine, it will use the Max node's output info
-		PreviewMaterial->GetEditorOnlyData()->EmissiveColor.OutputIndex = 0;
-		OutOutputChannel = FDMMaterialStageConnectorChannel::WHOLE_CHANNEL;
-	}
-
-	PreviewMaterial->GetEditorOnlyData()->EmissiveColor.Mask = 0;
-
-	if (OutOutputChannel != FDMMaterialStageConnectorChannel::WHOLE_CHANNEL)
-	{
-		PreviewMaterial->GetEditorOnlyData()->EmissiveColor.Mask = 1;
-		PreviewMaterial->GetEditorOnlyData()->EmissiveColor.MaskR = !!(OutOutputChannel & FDMMaterialStageConnectorChannel::FIRST_CHANNEL);
-		PreviewMaterial->GetEditorOnlyData()->EmissiveColor.MaskG = !!(OutOutputChannel & FDMMaterialStageConnectorChannel::SECOND_CHANNEL);
-		PreviewMaterial->GetEditorOnlyData()->EmissiveColor.MaskB = !!(OutOutputChannel & FDMMaterialStageConnectorChannel::THIRD_CHANNEL);
-		PreviewMaterial->GetEditorOnlyData()->EmissiveColor.MaskA = !!(OutOutputChannel & FDMMaterialStageConnectorChannel::FOURTH_CHANNEL);
-	}
-}
-
-void UDMMaterialSlot::UpdatePreviewMaterialProperty(const TSharedRef<FDMMaterialBuildState>& InBuildState, EDMMaterialPropertyType InProperty)
-{
-	if (!IsComponentValid())
-	{
-		return;
-	}
-
-	UMaterial* PreviewMaterial = InBuildState->GetDynamicMaterial();
-
-	UDynamicMaterialModelEditorOnlyData* ModelEditorOnlyData = GetMaterialModelEditorOnlyData();
-	check(ModelEditorOnlyData);
-	
-	UDMMaterialProperty* PropertyObj = ModelEditorOnlyData->GetMaterialProperty(InProperty);
-	check(PropertyObj);
-
-	const TArray<UMaterialExpression*>& SlotExpressions = InBuildState->GetSlotExpressions(this);
-
-	if (SlotExpressions.IsEmpty())
-	{
-		return;
-	}
-
-	UMaterialExpression* LastExpression = SlotExpressions.Last();
-	int32 BestMatch = INDEX_NONE;
-	int32 OutputCount = 0;
-	const int32 FloatsForPropertyType = UDMValueDefinitionLibrary::GetValueDefinition(PropertyObj->GetInputConnectorType()).GetFloatCount();
-
-	for (int32 OutputIdx = 0; OutputIdx < LastExpression->GetOutputs().Num(); ++OutputIdx)
-	{
-		const EMaterialValueType CurrentOutputType = static_cast<EMaterialValueType>(LastExpression->GetOutputType(OutputIdx));
-		int32 CurrentOutputCount = 0;
-
-		switch (CurrentOutputType)
-		{
-			case MCT_Float:
-			case MCT_Float1:
-				CurrentOutputCount = 1;
-				break;
-
-			case MCT_Float2:
-				CurrentOutputCount = 2;
-				break;
-
-			case MCT_Float3:
-				CurrentOutputCount = 3;
-				break;
-
-			case MCT_Float4:
-				CurrentOutputCount = 4;
-				break;
-
-			default:
-				continue; // For loop
-		}
-
-		if (CurrentOutputCount > OutputCount)
-		{
-			BestMatch = OutputIdx;
-			OutputCount = CurrentOutputCount;
-
-			if (CurrentOutputCount >= FloatsForPropertyType)
-			{
-				break;
-			}
-		}
-	}
-
-	if (BestMatch != INDEX_NONE)
-	{
-		PreviewMaterial->GetEditorOnlyData()->EmissiveColor.Expression = LastExpression;
-		PreviewMaterial->GetEditorOnlyData()->EmissiveColor.OutputIndex = BestMatch;
-	}
-}
-
 bool UDMMaterialSlot::SetLayerMaterialPropertyAndReplaceOthers(UDMMaterialLayerObject* InLayer, EDMMaterialPropertyType InPropertyFrom, 
 	EDMMaterialPropertyType InPropertyTo)
 {
diff --git a/Engine/Plugins/Experimental/DynamicMaterial/Source/DynamicMaterialEditor/Private/UI/Widgets/Editor/SDMMaterialPropertyPreviews.cpp b/Engine/Plugins/Experimental/DynamicMaterial/Source/DynamicMaterialEditor/Private/UI/Widgets/Editor/SDMMaterialPropertyPreviews.cpp
--- a/Engine/Plugins/Experimental/DynamicMaterial/Source/DynamicMaterialEditor/Private/UI/Widgets/Editor/SDMMaterialPropertyPreviews.cpp
+++ b/Engine/Plugins/Experimental/DynamicMaterial/Source/DynamicMaterialEditor/Private/UI/Widgets/Editor/SDMMaterialPropertyPreviews.cpp
@@ -82,7 +82,7 @@ TSharedRef<SWidget> SDMMaterialPropertyPreviews::CreateSlot_Content()
 
 				if (bIsActive)
 				{
-					AddPropertyPreview(WrapBox, InMaterialProperty, Slot);
+					AddPropertyPreview(WrapBox, MaterialProperty);
 				}
 			}
 
@@ -108,7 +108,7 @@ TSharedRef<SWidget> SDMMaterialPropertyPreviews::CreateSlot_Content()
 
 				if (!bIsActive && bIsValid)
 				{
-					AddPropertyPreview(WrapBox, InMaterialProperty, Slot);
+					AddPropertyPreview(WrapBox, MaterialProperty);
 				}
 			}
 
@@ -134,7 +134,7 @@ TSharedRef<SWidget> SDMMaterialPropertyPreviews::CreateSlot_Content()
 
 				if (!bIsActive && !bIsValid)
 				{
-					AddPropertyPreview(WrapBox, InMaterialProperty, Slot);
+					AddPropertyPreview(WrapBox, MaterialProperty);
 				}
 			}
 
@@ -149,8 +149,7 @@ TSharedRef<SWidget> SDMMaterialPropertyPreviews::CreateSlot_Content()
 		];
 }
 
-void SDMMaterialPropertyPreviews::AddPropertyPreview(const TSharedRef<SWrapBox>& InContainer, EDMMaterialPropertyType InMaterialProperty,
-	UDMMaterialSlot* InSlot)
+void SDMMaterialPropertyPreviews::AddPropertyPreview(const TSharedRef<SWrapBox>& InContainer, UDMMaterialProperty* InProperty)
 {
 	TSharedPtr<SDMMaterialEditor> EditorWidget = EditorWidgetWeak.Pin();
 
@@ -159,15 +158,33 @@ void SDMMaterialPropertyPreviews::AddPropertyPreview(const TSharedRef<SWrapBox>&
 		return;
 	}
 
+	UDynamicMaterialModel* MaterialModel = EditorWidget->GetMaterialModel();
+
+	if (!MaterialModel)
+	{
+		return;
+	}
+
+	UDynamicMaterialModelEditorOnlyData* EditorOnlyData = UDynamicMaterialModelEditorOnlyData::Get(MaterialModel);
+
+	if (!EditorOnlyData)
+	{
+		return;
+	}
+
 	TSharedPtr<SWidget> PreviewWidget;
+	EDMMaterialPropertyType MaterialProperty = InProperty ? InProperty->GetMaterialProperty() : EDMMaterialPropertyType::None;
+	UDMMaterialSlot* Slot = EditorOnlyData->GetSlotForMaterialProperty(MaterialProperty);
 
-	if (InSlot)
+	if (InProperty && InProperty->IsEnabled() && Slot)
 	{
-		PreviewWidget = SNew(SDMMaterialComponentPreview, EditorWidget.ToSharedRef(), InSlot)
+		MaterialProperty = InProperty->GetMaterialProperty();
+
+		PreviewWidget = SNew(SDMMaterialComponentPreview, EditorWidget.ToSharedRef(), InProperty)
 			.PreviewSize(FVector2D(60.f, 60.f));
 
 		PreviewWidget->SetCursor(EMouseCursor::Hand);
-		PreviewWidget->SetOnMouseButtonUp(FPointerEventHandler::CreateSP(this, &SDMMaterialPropertyPreviews::OnPreviewClicked, InMaterialProperty));
+		PreviewWidget->SetOnMouseButtonUp(FPointerEventHandler::CreateSP(this, &SDMMaterialPropertyPreviews::OnPreviewClicked, MaterialProperty));
 	}
 	else
 	{
@@ -195,7 +212,7 @@ void SDMMaterialPropertyPreviews::AddPropertyPreview(const TSharedRef<SWrapBox>&
 				.AutoWidth()
 				.VAlign(EVerticalAlignment::VAlign_Center)
 				[
-					CreateSlot_EnabledButton(InMaterialProperty)
+					CreateSlot_EnabledButton(MaterialProperty)
 				]
 				+ SHorizontalBox::Slot()
 				.FillWidth(1.f)
@@ -203,7 +220,7 @@ void SDMMaterialPropertyPreviews::AddPropertyPreview(const TSharedRef<SWrapBox>&
 				.VAlign(EVerticalAlignment::VAlign_Center)
 				.Padding(3.f, 0.f, 0.f, 0.f)
 				[
-					CreateSlot_PropertyName(InMaterialProperty)
+					CreateSlot_PropertyName(MaterialProperty)
 				]
 			]
 		];
diff --git a/Engine/Plugins/Experimental/DynamicMaterial/Source/DynamicMaterialEditor/Private/UI/Widgets/Editor/SDMMaterialPropertyPreviews.h b/Engine/Plugins/Experimental/DynamicMaterial/Source/DynamicMaterialEditor/Private/UI/Widgets/Editor/SDMMaterialPropertyPreviews.h
--- a/Engine/Plugins/Experimental/DynamicMaterial/Source/DynamicMaterialEditor/Private/UI/Widgets/Editor/SDMMaterialPropertyPreviews.h
+++ b/Engine/Plugins/Experimental/DynamicMaterial/Source/DynamicMaterialEditor/Private/UI/Widgets/Editor/SDMMaterialPropertyPreviews.h
@@ -34,7 +34,7 @@ protected:
 
 	TSharedRef<SWidget> CreateSlot_Content();
 
-	void AddPropertyPreview(const TSharedRef<SWrapBox>& InContainer, EDMMaterialPropertyType InMaterialProperty, UDMMaterialSlot* InSlot);
+	void AddPropertyPreview(const TSharedRef<SWrapBox>& InContainer, UDMMaterialProperty* InProperty);
 
 	TSharedRef<SWidget> CreateSlot_EnabledButton(EDMMaterialPropertyType InMaterialProperty);
 
diff --git a/Engine/Plugins/Experimental/DynamicMaterial/Source/DynamicMaterialEditor/Private/UI/Widgets/Visualizers/SDMMaterialComponentPreview.cpp b/Engine/Plugins/Experimental/DynamicMaterial/Source/DynamicMaterialEditor/Private/UI/Widgets/Visualizers/SDMMaterialComponentPreview.cpp
--- a/Engine/Plugins/Experimental/DynamicMaterial/Source/DynamicMaterialEditor/Private/UI/Widgets/Visualizers/SDMMaterialComponentPreview.cpp
+++ b/Engine/Plugins/Experimental/DynamicMaterial/Source/DynamicMaterialEditor/Private/UI/Widgets/Visualizers/SDMMaterialComponentPreview.cpp
@@ -2,6 +2,7 @@
 
 #include "UI/Widgets/Visualizers/SDMMaterialComponentPreview.h"
 
+#include "Components/DMMaterialProperty.h"
 #include "Components/DMMaterialStage.h"
 #include "Components/DMMaterialValue.h"
 #include "Components/DMTextureUV.h"
@@ -89,14 +90,7 @@ void SDMMaterialComponentPreview::Tick(const FGeometry& AllottedGeometry, const
 
 void SDMMaterialComponentPreview::OnComponentUpdated(UDMMaterialComponent* InComponent, UDMMaterialComponent* InSource, EDMUpdateType InUpdateType)
 {
-	UDMMaterialStage* Stage = Cast<UDMMaterialStage>(InComponent);
-
-	if (!Stage)
-	{
-		return;
-	}
-
-	if (Stage != ComponentWeak.Get() || !IsValid(Stage) || !Stage->IsComponentValid())
+	if (InComponent != ComponentWeak.Get() || !IsValid(InComponent) || !InComponent->IsComponentValid())
 	{
 		return;
 	}
@@ -112,7 +106,7 @@ void SDMMaterialComponentPreview::OnComponentUpdated(UDMMaterialComponent* InCom
 
 	if (!PreviewMaterialBase)
 	{
-		PreviewMaterialBase = EditorWidget->GetPreviewMaterialManager()->CreatePreviewMaterial(Stage);
+		PreviewMaterialBase = EditorWidget->GetPreviewMaterialManager()->CreatePreviewMaterial(InComponent);
 	}
 
 	UMaterialInstanceDynamic* MID = PreviewMaterialDynamicWeak.Get();
@@ -136,7 +130,7 @@ void SDMMaterialComponentPreview::OnComponentUpdated(UDMMaterialComponent* InCom
 			TextureUVDynamic->SetMIDParameters(MID);
 		}
 	}
-	else
+	else if (UDMMaterialStage* Stage = Cast<UDMMaterialStage>(InComponent))
 	{
 		Stage->GeneratePreviewMaterial(PreviewMaterialBase);
 
@@ -154,6 +148,26 @@ void SDMMaterialComponentPreview::OnComponentUpdated(UDMMaterialComponent* InCom
 			MaterialModelDynamic->ApplyComponents(PreviewMaterialDynamicWeak.Get());
 		}
 
+		Brush.SetMaterial(PreviewMaterialDynamicWeak.Get());
+	}
+	else if (UDMMaterialProperty* Property = Cast<UDMMaterialProperty>(InComponent))
+	{
+		Property->GeneratePreviewMaterial(PreviewMaterialBase);
+
+		EditorWidget->GetPreviewMaterialManager()->FreePreviewMaterialDynamic(PreviewMaterialBase);
+		PreviewMaterialDynamicWeak = EditorWidget->GetPreviewMaterialManager()->CreatePreviewMaterialDynamic(PreviewMaterialBase);
+
+		UDynamicMaterialModelBase* MaterialModelBase = EditorWidget->GetMaterialModelBase();
+
+		if (UDynamicMaterialModel* MaterialModel = Cast<UDynamicMaterialModel>(MaterialModelBase))
+		{
+			MaterialModel->ApplyComponents(PreviewMaterialDynamicWeak.Get());
+		}
+		else if (UDynamicMaterialModelDynamic* MaterialModelDynamic = Cast<UDynamicMaterialModelDynamic>(MaterialModelBase))
+		{
+			MaterialModelDynamic->ApplyComponents(PreviewMaterialDynamicWeak.Get());
+		}
+
 		Brush.SetMaterial(PreviewMaterialDynamicWeak.Get());
 	}
 }
diff --git a/Engine/Plugins/Experimental/DynamicMaterial/Source/DynamicMaterialEditor/Public/Components/DMMaterialProperty.h b/Engine/Plugins/Experimental/DynamicMaterial/Source/DynamicMaterialEditor/Public/Components/DMMaterialProperty.h
--- a/Engine/Plugins/Experimental/DynamicMaterial/Source/DynamicMaterialEditor/Public/Components/DMMaterialProperty.h
+++ b/Engine/Plugins/Experimental/DynamicMaterial/Source/DynamicMaterialEditor/Public/Components/DMMaterialProperty.h
@@ -16,6 +16,7 @@ class UDMMaterialSlot;
 class UDMMaterialStageInputExpression;
 class UDynamicMaterialModel;
 class UDynamicMaterialModelEditorOnlyData;
+class UMaterial;
 class UMaterialExpression;
 class UMaterialFunctionInterface;
 struct FDMMaterialBuildState;
@@ -196,6 +197,11 @@ public:
 	 */
 	virtual void AddOutputProcessor(const TSharedRef<FDMMaterialBuildState>& InBuildState) const;
 
+	/**
+	 * Adds nodes to the given material based on this property's generated expressions.
+	 */
+	virtual void GeneratePreviewMaterial(UMaterial* InPreviewMaterial);
+
 	//~ Begin UDMMaterialComponent
 	DYNAMICMATERIALEDITOR_API virtual void Update(UDMMaterialComponent* InSource, EDMUpdateType InUpdateType) override;
 	DYNAMICMATERIALEDITOR_API virtual FString GetComponentPathComponent() const override;
diff --git a/Engine/Plugins/Experimental/DynamicMaterial/Source/DynamicMaterialEditor/Public/Components/DMMaterialSlot.h b/Engine/Plugins/Experimental/DynamicMaterial/Source/DynamicMaterialEditor/Public/Components/DMMaterialSlot.h
--- a/Engine/Plugins/Experimental/DynamicMaterial/Source/DynamicMaterialEditor/Public/Components/DMMaterialSlot.h
+++ b/Engine/Plugins/Experimental/DynamicMaterial/Source/DynamicMaterialEditor/Public/Components/DMMaterialSlot.h
@@ -136,9 +136,6 @@ public:
 	/** Returns true if all associations have been removed */
 	bool UnreferencedBySlot(UDMMaterialSlot* InOtherSlot);
 
-	UFUNCTION(BlueprintCallable, Category = "Material Designer")
-	DYNAMICMATERIALEDITOR_API void GeneralPreviewMaterial(UMaterial* InPreviewMaterial, EDMMaterialLayerStage InLayerStage);
-
 	/** Sets the material property of the given layer and changes all other layers matching that property to a different one. */
 	UFUNCTION(BlueprintCallable, Category = "Material Designer")
 	DYNAMICMATERIALEDITOR_API bool SetLayerMaterialPropertyAndReplaceOthers(UDMMaterialLayerObject* InLayer, EDMMaterialPropertyType InPropertyFrom,
@@ -178,14 +175,6 @@ protected:
 	FDMOnMaterialSlotPropertiesUpdated OnPropertiesUpdateDelegate;
 	FDMOnMaterialSlotLayersUpdated OnLayersUpdateDelegate;
 
-	void UpdateBasePreviewMaterial(const TSharedRef<FDMMaterialBuildState>& InBuildState);
-	void UpdateBasePreviewMaterialProperty(const TSharedRef<FDMMaterialBuildState>& InBuildState, EDMMaterialPropertyType InBaseProperty);
-	void UpdateBasePreviewMaterialFull(const TSharedRef<FDMMaterialBuildState>& InBuildState);
-	void UpdateMaskPreviewMaterial(const TSharedRef<FDMMaterialBuildState>& InBuildState);
-	void UpdateMaskPreviewMaterialProperty(const TSharedRef<FDMMaterialBuildState>& InBuildState, EDMMaterialPropertyType InMaskProperty);
-	void UpdateMaskPreviewMaterialMaskCombination(const TSharedRef<FDMMaterialBuildState>& InBuildState, EDMMaterialPropertyType InMaskProperty);
-	void UpdatePreviewMaterialProperty(const TSharedRef<FDMMaterialBuildState>& InBuildState, EDMMaterialPropertyType InProperty);
-
 	//~ Begin UDMMaterialComponent
 	DYNAMICMATERIALEDITOR_API virtual void OnComponentAdded() override;
 	DYNAMICMATERIALEDITOR_API virtual void OnComponentRemoved() override;
