Material Designer
- Changed the way that properties are named in the details view. They should be more descriptive instead of repeating "Texture" over and over, for instance.

diff --git a/Engine/Plugins/Experimental/DynamicMaterial/Source/DynamicMaterialEditor/Private/UI/PropertyGenerators/DMThroughputPropertyRowGenerator.cpp b/Engine/Plugins/Experimental/DynamicMaterial/Source/DynamicMaterialEditor/Private/UI/PropertyGenerators/DMThroughputPropertyRowGenerator.cpp
--- a/Engine/Plugins/Experimental/DynamicMaterial/Source/DynamicMaterialEditor/Private/UI/PropertyGenerators/DMThroughputPropertyRowGenerator.cpp
+++ b/Engine/Plugins/Experimental/DynamicMaterial/Source/DynamicMaterialEditor/Private/UI/PropertyGenerators/DMThroughputPropertyRowGenerator.cpp
@@ -105,7 +105,22 @@ void FDMThroughputPropertyRowGenerator::AddComponentProperties(const TSharedRef<
 					continue;
 				}
 
-				PropertyRow.NameOverride = InputConnectors[InputIdx].Name;
+				if (!PropertyRow.NameOverride.IsSet() || PropertyRow.NameOverride.GetValue().IsEmpty())
+
+				{
+					if (!PropertyRow.ValueName.IsNone())
+					{
+						PropertyRow.NameOverride = FText::FromName(PropertyRow.ValueName);
+					}
+					else if (PropertyRow.PropertyHandle.IsValid())
+					{
+						PropertyRow.NameOverride = PropertyRow.PropertyHandle->GetPropertyDisplayName();
+					}
+					else
+					{
+						PropertyRow.NameOverride = InputConnectors[InputIdx].Name;
+					}
+				}
 			}
 		}
 	}
