Material Designer
- The source selector text is now more indicative of what it actually is when a render target renderer is selected.

diff --git a/Engine/Plugins/Experimental/DynamicMaterial/Source/DynamicMaterial/Private/Components/DMRenderTargetRenderer.cpp b/Engine/Plugins/Experimental/DynamicMaterial/Source/DynamicMaterial/Private/Components/DMRenderTargetRenderer.cpp
--- a/Engine/Plugins/Experimental/DynamicMaterial/Source/DynamicMaterial/Private/Components/DMRenderTargetRenderer.cpp
+++ b/Engine/Plugins/Experimental/DynamicMaterial/Source/DynamicMaterial/Private/Components/DMRenderTargetRenderer.cpp
@@ -6,6 +6,8 @@
 #include "Misc/CoreDelegates.h"
 #include "Templates/SubclassOf.h"
 
+#define LOCTEXT_NAMESPACE "DMRenderTargetRenderer"
+
 UDMRenderTargetRenderer* UDMRenderTargetRenderer::CreateRenderTargetRenderer(TSubclassOf<UDMRenderTargetRenderer> InRendererClass, 
 	UDMMaterialValueRenderTarget* InRenderTargetValue)
 {
@@ -76,6 +78,13 @@ bool UDMRenderTargetRenderer::JsonDeserialize(const TSharedPtr<FJsonValue>& InJs
 	return false;
 }
 
+#if WITH_EDITOR
+FText UDMRenderTargetRenderer::GetComponentDescription() const
+{
+	return LOCTEXT("Renderer", "Renderer");
+}
+#endif
+
 void UDMRenderTargetRenderer::Update(UDMMaterialComponent* InSource, EDMUpdateType InUpdateType)
 {
 	if (!FDMUpdateGuard::CanUpdate())
@@ -114,3 +123,5 @@ void UDMRenderTargetRenderer::PostLoad()
 		RenderTargetValue->EnsureRenderTarget();
 	}
 }
+
+#undef LOCTEXT_NAMESPACE
diff --git a/Engine/Plugins/Experimental/DynamicMaterial/Source/DynamicMaterial/Private/Components/MaterialValues/DMMaterialValueBool.cpp b/Engine/Plugins/Experimental/DynamicMaterial/Source/DynamicMaterial/Private/Components/MaterialValues/DMMaterialValueBool.cpp
--- a/Engine/Plugins/Experimental/DynamicMaterial/Source/DynamicMaterial/Private/Components/MaterialValues/DMMaterialValueBool.cpp
+++ b/Engine/Plugins/Experimental/DynamicMaterial/Source/DynamicMaterial/Private/Components/MaterialValues/DMMaterialValueBool.cpp
@@ -94,6 +94,11 @@ FString UDMMaterialValueBool::GetComponentPathComponent() const
 	return TEXT("Bool");
 }
 
+FText UDMMaterialValueBool::GetComponentDescription() const
+{
+	return LOCTEXT("Bool", "Bool");
+}
+
 TSharedPtr<FJsonValue> UDMMaterialValueBool::JsonSerialize() const
 {
 	return FDMJsonUtils::Serialize(Value);
diff --git a/Engine/Plugins/Experimental/DynamicMaterial/Source/DynamicMaterial/Private/Components/MaterialValues/DMMaterialValueColorAtlas.cpp b/Engine/Plugins/Experimental/DynamicMaterial/Source/DynamicMaterial/Private/Components/MaterialValues/DMMaterialValueColorAtlas.cpp
--- a/Engine/Plugins/Experimental/DynamicMaterial/Source/DynamicMaterial/Private/Components/MaterialValues/DMMaterialValueColorAtlas.cpp
+++ b/Engine/Plugins/Experimental/DynamicMaterial/Source/DynamicMaterial/Private/Components/MaterialValues/DMMaterialValueColorAtlas.cpp
@@ -17,6 +17,8 @@
 #include "Utils/DMUtils.h"
 #endif
 
+#define LOCTEXT_NAMESPACE "DMMaterialValueColorAtlas"
+
 UDMMaterialValueColorAtlas::UDMMaterialValueColorAtlas()
 	: UDMMaterialValue(EDMValueType::VT_ColorAtlas)
 	, Value(0)
@@ -156,6 +158,11 @@ FString UDMMaterialValueColorAtlas::GetComponentPathComponent() const
 	return TEXT("ColorAtlasAlpha");
 }
 
+FText UDMMaterialValueColorAtlas::GetComponentDescription() const
+{
+	return LOCTEXT("ColorAtlas", "Color Atlas");
+}
+
 TSharedPtr<FJsonValue> UDMMaterialValueColorAtlas::JsonSerialize() const
 {
 	return FDMJsonUtils::Serialize({
@@ -233,3 +240,5 @@ bool UDMMaterialValueColorAtlas::IsDefaultValue() const
 	return FMath::IsNearlyEqual(Value, DefaultValue);
 }
 #endif
+
+#undef LOCTEXT_NAMESPACE
diff --git a/Engine/Plugins/Experimental/DynamicMaterial/Source/DynamicMaterial/Private/Components/MaterialValues/DMMaterialValueFloat1.cpp b/Engine/Plugins/Experimental/DynamicMaterial/Source/DynamicMaterial/Private/Components/MaterialValues/DMMaterialValueFloat1.cpp
--- a/Engine/Plugins/Experimental/DynamicMaterial/Source/DynamicMaterial/Private/Components/MaterialValues/DMMaterialValueFloat1.cpp
+++ b/Engine/Plugins/Experimental/DynamicMaterial/Source/DynamicMaterial/Private/Components/MaterialValues/DMMaterialValueFloat1.cpp
@@ -75,6 +75,11 @@ FString UDMMaterialValueFloat1::GetComponentPathComponent() const
 	return TEXT("Scalar");
 }
 
+FText UDMMaterialValueFloat1::GetComponentDescription() const
+{
+	return LOCTEXT("Scalar", "Scalar");
+}
+
 TSharedPtr<FJsonValue> UDMMaterialValueFloat1::JsonSerialize() const
 {
 	return FDMJsonUtils::Serialize(Value);
diff --git a/Engine/Plugins/Experimental/DynamicMaterial/Source/DynamicMaterial/Private/Components/MaterialValues/DMMaterialValueFloat2.cpp b/Engine/Plugins/Experimental/DynamicMaterial/Source/DynamicMaterial/Private/Components/MaterialValues/DMMaterialValueFloat2.cpp
--- a/Engine/Plugins/Experimental/DynamicMaterial/Source/DynamicMaterial/Private/Components/MaterialValues/DMMaterialValueFloat2.cpp
+++ b/Engine/Plugins/Experimental/DynamicMaterial/Source/DynamicMaterial/Private/Components/MaterialValues/DMMaterialValueFloat2.cpp
@@ -84,6 +84,11 @@ FString UDMMaterialValueFloat2::GetComponentPathComponent() const
 	return TEXT("Vector2D");
 }
 
+FText UDMMaterialValueFloat2::GetComponentDescription() const
+{
+	return LOCTEXT("Vector2", "Vector 2");
+}
+
 TSharedPtr<FJsonValue> UDMMaterialValueFloat2::JsonSerialize() const
 {
 	return FDMJsonUtils::Serialize(Value);
diff --git a/Engine/Plugins/Experimental/DynamicMaterial/Source/DynamicMaterial/Private/Components/MaterialValues/DMMaterialValueFloat3RGB.cpp b/Engine/Plugins/Experimental/DynamicMaterial/Source/DynamicMaterial/Private/Components/MaterialValues/DMMaterialValueFloat3RGB.cpp
--- a/Engine/Plugins/Experimental/DynamicMaterial/Source/DynamicMaterial/Private/Components/MaterialValues/DMMaterialValueFloat3RGB.cpp
+++ b/Engine/Plugins/Experimental/DynamicMaterial/Source/DynamicMaterial/Private/Components/MaterialValues/DMMaterialValueFloat3RGB.cpp
@@ -81,6 +81,11 @@ FString UDMMaterialValueFloat3RGB::GetComponentPathComponent() const
 	return TEXT("RGB");
 }
 
+FText UDMMaterialValueFloat3RGB::GetComponentDescription() const
+{
+	return LOCTEXT("ColorRGB", "Color (RGB)");
+}
+
 TSharedPtr<FJsonValue> UDMMaterialValueFloat3RGB::JsonSerialize() const
 {
 	return FDMJsonUtils::Serialize(Value);
diff --git a/Engine/Plugins/Experimental/DynamicMaterial/Source/DynamicMaterial/Private/Components/MaterialValues/DMMaterialValueFloat3RPY.cpp b/Engine/Plugins/Experimental/DynamicMaterial/Source/DynamicMaterial/Private/Components/MaterialValues/DMMaterialValueFloat3RPY.cpp
--- a/Engine/Plugins/Experimental/DynamicMaterial/Source/DynamicMaterial/Private/Components/MaterialValues/DMMaterialValueFloat3RPY.cpp
+++ b/Engine/Plugins/Experimental/DynamicMaterial/Source/DynamicMaterial/Private/Components/MaterialValues/DMMaterialValueFloat3RPY.cpp
@@ -81,6 +81,11 @@ FString UDMMaterialValueFloat3RPY::GetComponentPathComponent() const
 	return TEXT("Rotator");
 }
 
+FText UDMMaterialValueFloat3RPY::GetComponentDescription() const
+{
+	return LOCTEXT("Rotator", "Rotator");
+}
+
 TSharedPtr<FJsonValue> UDMMaterialValueFloat3RPY::JsonSerialize() const
 {
 	return FDMJsonUtils::Serialize(Value);
diff --git a/Engine/Plugins/Experimental/DynamicMaterial/Source/DynamicMaterial/Private/Components/MaterialValues/DMMaterialValueFloat3XYZ.cpp b/Engine/Plugins/Experimental/DynamicMaterial/Source/DynamicMaterial/Private/Components/MaterialValues/DMMaterialValueFloat3XYZ.cpp
--- a/Engine/Plugins/Experimental/DynamicMaterial/Source/DynamicMaterial/Private/Components/MaterialValues/DMMaterialValueFloat3XYZ.cpp
+++ b/Engine/Plugins/Experimental/DynamicMaterial/Source/DynamicMaterial/Private/Components/MaterialValues/DMMaterialValueFloat3XYZ.cpp
@@ -81,6 +81,11 @@ FString UDMMaterialValueFloat3XYZ::GetComponentPathComponent() const
 	return TEXT("Vector3D");
 }
 
+FText UDMMaterialValueFloat3XYZ::GetComponentDescription() const
+{
+	return LOCTEXT("Vector3", "Vector 3");
+}
+
 TSharedPtr<FJsonValue> UDMMaterialValueFloat3XYZ::JsonSerialize() const
 {
 	return FDMJsonUtils::Serialize(Value);
diff --git a/Engine/Plugins/Experimental/DynamicMaterial/Source/DynamicMaterial/Private/Components/MaterialValues/DMMaterialValueFloat4.cpp b/Engine/Plugins/Experimental/DynamicMaterial/Source/DynamicMaterial/Private/Components/MaterialValues/DMMaterialValueFloat4.cpp
--- a/Engine/Plugins/Experimental/DynamicMaterial/Source/DynamicMaterial/Private/Components/MaterialValues/DMMaterialValueFloat4.cpp
+++ b/Engine/Plugins/Experimental/DynamicMaterial/Source/DynamicMaterial/Private/Components/MaterialValues/DMMaterialValueFloat4.cpp
@@ -82,6 +82,11 @@ FString UDMMaterialValueFloat4::GetComponentPathComponent() const
 	return TEXT("RGBA");
 }
 
+FText UDMMaterialValueFloat4::GetComponentDescription() const
+{
+	return LOCTEXT("ColorRGBA", "Color (RGBA)");
+}
+
 TSharedPtr<FJsonValue> UDMMaterialValueFloat4::JsonSerialize() const
 {
 	return FDMJsonUtils::Serialize(Value);
diff --git a/Engine/Plugins/Experimental/DynamicMaterial/Source/DynamicMaterial/Private/Components/MaterialValues/DMMaterialValueRenderTarget.cpp b/Engine/Plugins/Experimental/DynamicMaterial/Source/DynamicMaterial/Private/Components/MaterialValues/DMMaterialValueRenderTarget.cpp
--- a/Engine/Plugins/Experimental/DynamicMaterial/Source/DynamicMaterial/Private/Components/MaterialValues/DMMaterialValueRenderTarget.cpp
+++ b/Engine/Plugins/Experimental/DynamicMaterial/Source/DynamicMaterial/Private/Components/MaterialValues/DMMaterialValueRenderTarget.cpp
@@ -10,6 +10,8 @@
 #include "Components/MaterialValuesDynamic/DMMaterialValueRenderTargetDynamic.h"
 #endif
 
+#define LOCTEXT_NAMESPACE "DMMaterialValueRenderTarget"
+
 const FString UDMMaterialValueRenderTarget::RendererPathToken = "Renderer";
 
 UDMMaterialValueRenderTarget::UDMMaterialValueRenderTarget()
@@ -172,6 +174,16 @@ FString UDMMaterialValueRenderTarget::GetComponentPathComponent() const
 	return TEXT("RenderTarget");
 }
 
+FText UDMMaterialValueRenderTarget::GetComponentDescription() const
+{
+	if (IsValid(Renderer))
+	{
+		return Renderer->GetComponentDescription();
+	}
+
+	return LOCTEXT("RenderTarget", "Render Target");
+}
+
 TSharedPtr<FJsonValue> UDMMaterialValueRenderTarget::JsonSerialize() const
 {
 	return FDMJsonUtils::Serialize({
@@ -367,3 +379,5 @@ void UDMMaterialValueRenderTarget::OnComponentRemoved()
 	}
 }
 #endif
+
+#undef LOCTEXT_NAMESPACE
diff --git a/Engine/Plugins/Experimental/DynamicMaterial/Source/DynamicMaterial/Private/Components/MaterialValues/DMMaterialValueTexture.cpp b/Engine/Plugins/Experimental/DynamicMaterial/Source/DynamicMaterial/Private/Components/MaterialValues/DMMaterialValueTexture.cpp
--- a/Engine/Plugins/Experimental/DynamicMaterial/Source/DynamicMaterial/Private/Components/MaterialValues/DMMaterialValueTexture.cpp
+++ b/Engine/Plugins/Experimental/DynamicMaterial/Source/DynamicMaterial/Private/Components/MaterialValues/DMMaterialValueTexture.cpp
@@ -18,6 +18,8 @@
 #include "Utils/DMUtils.h"
 #endif
 
+#define LOCTEXT_NAMESPACE "DMMaterialValueTexture"
+
 #if WITH_EDITOR
 namespace UE::DynamicMaterial::Private
 {
@@ -183,6 +185,11 @@ FString UDMMaterialValueTexture::GetComponentPathComponent() const
 	return TEXT("Texture");
 }
 
+FText UDMMaterialValueTexture::GetComponentDescription() const
+{
+	return LOCTEXT("Texture", "Texture");
+}
+
 TSharedPtr<FJsonValue> UDMMaterialValueTexture::JsonSerialize() const
 {
 	return FDMJsonUtils::Serialize(Value);
@@ -240,3 +247,5 @@ void UDMMaterialValueTexture::SetMIDParameter(UMaterialInstanceDynamic* InMID) c
 
 	InMID->SetTextureParameterValue(GetMaterialParameterName(), Value);
 }
+
+#undef LOCTEXT_NAMESPACE
diff --git a/Engine/Plugins/Experimental/DynamicMaterial/Source/DynamicMaterial/Private/Components/RenderTargetRenderers/DMRenderTargetTextRenderer.cpp b/Engine/Plugins/Experimental/DynamicMaterial/Source/DynamicMaterial/Private/Components/RenderTargetRenderers/DMRenderTargetTextRenderer.cpp
--- a/Engine/Plugins/Experimental/DynamicMaterial/Source/DynamicMaterial/Private/Components/RenderTargetRenderers/DMRenderTargetTextRenderer.cpp
+++ b/Engine/Plugins/Experimental/DynamicMaterial/Source/DynamicMaterial/Private/Components/RenderTargetRenderers/DMRenderTargetTextRenderer.cpp
@@ -247,6 +247,11 @@ bool UDMRenderTargetTextRenderer::JsonDeserialize(const TSharedPtr<FJsonValue>&
 	return true;
 }
 
+FText UDMRenderTargetTextRenderer::GetComponentDescription() const
+{
+	return LOCTEXT("Text", "Text");
+}
+
 void UDMRenderTargetTextRenderer::PostEditChangeProperty(FPropertyChangedEvent& InPropertyChangedEvent)
 {
 	Super::PostEditChangeProperty(InPropertyChangedEvent);
diff --git a/Engine/Plugins/Experimental/DynamicMaterial/Source/DynamicMaterial/Private/Components/RenderTargetRenderers/DMRenderTargetUMGWidgetRenderer.cpp b/Engine/Plugins/Experimental/DynamicMaterial/Source/DynamicMaterial/Private/Components/RenderTargetRenderers/DMRenderTargetUMGWidgetRenderer.cpp
--- a/Engine/Plugins/Experimental/DynamicMaterial/Source/DynamicMaterial/Private/Components/RenderTargetRenderers/DMRenderTargetUMGWidgetRenderer.cpp
+++ b/Engine/Plugins/Experimental/DynamicMaterial/Source/DynamicMaterial/Private/Components/RenderTargetRenderers/DMRenderTargetUMGWidgetRenderer.cpp
@@ -46,6 +46,11 @@ bool UDMRenderTargetUMGWidgetRenderer::JsonDeserialize(const TSharedPtr<FJsonVal
 	return false;
 }
 
+FText UDMRenderTargetUMGWidgetRenderer::GetComponentDescription() const
+{
+	return LOCTEXT("UMGWidget", "UMG Widget");
+}
+
 void UDMRenderTargetUMGWidgetRenderer::PostEditChangeProperty(FPropertyChangedEvent& InPropertyChangedEvent)
 {
 	Super::PostEditChangeProperty(InPropertyChangedEvent);
diff --git a/Engine/Plugins/Experimental/DynamicMaterial/Source/DynamicMaterial/Private/Components/RenderTargetRenderers/DMRenderTargetWidgetRendererBase.cpp b/Engine/Plugins/Experimental/DynamicMaterial/Source/DynamicMaterial/Private/Components/RenderTargetRenderers/DMRenderTargetWidgetRendererBase.cpp
--- a/Engine/Plugins/Experimental/DynamicMaterial/Source/DynamicMaterial/Private/Components/RenderTargetRenderers/DMRenderTargetWidgetRendererBase.cpp
+++ b/Engine/Plugins/Experimental/DynamicMaterial/Source/DynamicMaterial/Private/Components/RenderTargetRenderers/DMRenderTargetWidgetRendererBase.cpp
@@ -16,6 +16,13 @@ UDMRenderTargetWidgetRendererBase::UDMRenderTargetWidgetRendererBase()
 	WidgetRenderer->SetShouldClearTarget(true);
 }
 
+#if WITH_EDITOR
+FText UDMRenderTargetWidgetRendererBase::GetComponentDescription() const
+{
+	return LOCTEXT("Widget", "Widget");
+}
+#endif
+
 void UDMRenderTargetWidgetRendererBase::UpdateRenderTarget_Internal()
 {
 	UDMMaterialValueRenderTarget* RenderTargetValue = GetRenderTargetValue();
diff --git a/Engine/Plugins/Experimental/DynamicMaterial/Source/DynamicMaterial/Public/Components/DMRenderTargetRenderer.h b/Engine/Plugins/Experimental/DynamicMaterial/Source/DynamicMaterial/Public/Components/DMRenderTargetRenderer.h
--- a/Engine/Plugins/Experimental/DynamicMaterial/Source/DynamicMaterial/Public/Components/DMRenderTargetRenderer.h
+++ b/Engine/Plugins/Experimental/DynamicMaterial/Source/DynamicMaterial/Public/Components/DMRenderTargetRenderer.h
@@ -57,6 +57,9 @@ public:
 	//~ End IDMJsonSerializable
 
 	//~ Begin UDMMaterialComponent
+#if WITH_EDITOR
+	DYNAMICMATERIAL_API virtual FText GetComponentDescription() const override;
+#endif
 	DYNAMICMATERIAL_API virtual void Update(UDMMaterialComponent* InSource, EDMUpdateType InUpdateType) override;
 	//~ End UDMMaterialComponent
 
diff --git a/Engine/Plugins/Experimental/DynamicMaterial/Source/DynamicMaterial/Public/Components/MaterialValues/DMMaterialValueBool.h b/Engine/Plugins/Experimental/DynamicMaterial/Source/DynamicMaterial/Public/Components/MaterialValues/DMMaterialValueBool.h
--- a/Engine/Plugins/Experimental/DynamicMaterial/Source/DynamicMaterial/Public/Components/MaterialValues/DMMaterialValueBool.h
+++ b/Engine/Plugins/Experimental/DynamicMaterial/Source/DynamicMaterial/Public/Components/MaterialValues/DMMaterialValueBool.h
@@ -51,6 +51,7 @@ public:
 	//~ Begin UDMMaterialComponent
 #if WITH_EDITOR
 	DYNAMICMATERIAL_API virtual FString GetComponentPathComponent() const override;
+	DYNAMICMATERIAL_API virtual FText GetComponentDescription() const override;
 #endif
 	//~ End UDMMaterialComponent
  
diff --git a/Engine/Plugins/Experimental/DynamicMaterial/Source/DynamicMaterial/Public/Components/MaterialValues/DMMaterialValueColorAtlas.h b/Engine/Plugins/Experimental/DynamicMaterial/Source/DynamicMaterial/Public/Components/MaterialValues/DMMaterialValueColorAtlas.h
--- a/Engine/Plugins/Experimental/DynamicMaterial/Source/DynamicMaterial/Public/Components/MaterialValues/DMMaterialValueColorAtlas.h
+++ b/Engine/Plugins/Experimental/DynamicMaterial/Source/DynamicMaterial/Public/Components/MaterialValues/DMMaterialValueColorAtlas.h
@@ -70,6 +70,7 @@ public:
 	//~ Begin UDMMaterialComponent
 #if WITH_EDITOR
 	DYNAMICMATERIAL_API virtual FString GetComponentPathComponent() const override;
+	DYNAMICMATERIAL_API virtual FText GetComponentDescription() const override;
 #endif
 	//~ End UDMMaterialComponent
 
diff --git a/Engine/Plugins/Experimental/DynamicMaterial/Source/DynamicMaterial/Public/Components/MaterialValues/DMMaterialValueFloat1.h b/Engine/Plugins/Experimental/DynamicMaterial/Source/DynamicMaterial/Public/Components/MaterialValues/DMMaterialValueFloat1.h
--- a/Engine/Plugins/Experimental/DynamicMaterial/Source/DynamicMaterial/Public/Components/MaterialValues/DMMaterialValueFloat1.h
+++ b/Engine/Plugins/Experimental/DynamicMaterial/Source/DynamicMaterial/Public/Components/MaterialValues/DMMaterialValueFloat1.h
@@ -51,6 +51,7 @@ public:
 	//~ Begin UDMMaterialComponent
 #if WITH_EDITOR
 	DYNAMICMATERIAL_API virtual FString GetComponentPathComponent() const override;
+	DYNAMICMATERIAL_API virtual FText GetComponentDescription() const override;
 #endif
 	//~ End UDMMaterialComponent
  
diff --git a/Engine/Plugins/Experimental/DynamicMaterial/Source/DynamicMaterial/Public/Components/MaterialValues/DMMaterialValueFloat2.h b/Engine/Plugins/Experimental/DynamicMaterial/Source/DynamicMaterial/Public/Components/MaterialValues/DMMaterialValueFloat2.h
--- a/Engine/Plugins/Experimental/DynamicMaterial/Source/DynamicMaterial/Public/Components/MaterialValues/DMMaterialValueFloat2.h
+++ b/Engine/Plugins/Experimental/DynamicMaterial/Source/DynamicMaterial/Public/Components/MaterialValues/DMMaterialValueFloat2.h
@@ -52,6 +52,7 @@ public:
 	//~ Begin UDMMaterialComponent
 #if WITH_EDITOR
 	DYNAMICMATERIAL_API virtual FString GetComponentPathComponent() const override;
+	DYNAMICMATERIAL_API virtual FText GetComponentDescription() const override;
 #endif
 	//~ End UDMMaterialComponent
  
diff --git a/Engine/Plugins/Experimental/DynamicMaterial/Source/DynamicMaterial/Public/Components/MaterialValues/DMMaterialValueFloat3RGB.h b/Engine/Plugins/Experimental/DynamicMaterial/Source/DynamicMaterial/Public/Components/MaterialValues/DMMaterialValueFloat3RGB.h
--- a/Engine/Plugins/Experimental/DynamicMaterial/Source/DynamicMaterial/Public/Components/MaterialValues/DMMaterialValueFloat3RGB.h
+++ b/Engine/Plugins/Experimental/DynamicMaterial/Source/DynamicMaterial/Public/Components/MaterialValues/DMMaterialValueFloat3RGB.h
@@ -52,6 +52,7 @@ public:
 	//~ Begin UDMMaterialComponent
 #if WITH_EDITOR
 	DYNAMICMATERIAL_API virtual FString GetComponentPathComponent() const override;
+	DYNAMICMATERIAL_API virtual FText GetComponentDescription() const override;
 #endif
 	//~ End UDMMaterialComponent
 
diff --git a/Engine/Plugins/Experimental/DynamicMaterial/Source/DynamicMaterial/Public/Components/MaterialValues/DMMaterialValueFloat3RPY.h b/Engine/Plugins/Experimental/DynamicMaterial/Source/DynamicMaterial/Public/Components/MaterialValues/DMMaterialValueFloat3RPY.h
--- a/Engine/Plugins/Experimental/DynamicMaterial/Source/DynamicMaterial/Public/Components/MaterialValues/DMMaterialValueFloat3RPY.h
+++ b/Engine/Plugins/Experimental/DynamicMaterial/Source/DynamicMaterial/Public/Components/MaterialValues/DMMaterialValueFloat3RPY.h
@@ -52,6 +52,7 @@ public:
 	//~ Begin UDMMaterialComponent
 #if WITH_EDITOR
 	DYNAMICMATERIAL_API virtual FString GetComponentPathComponent() const override;
+	DYNAMICMATERIAL_API virtual FText GetComponentDescription() const override;
 #endif
 	//~ End UDMMaterialComponent
  
diff --git a/Engine/Plugins/Experimental/DynamicMaterial/Source/DynamicMaterial/Public/Components/MaterialValues/DMMaterialValueFloat3XYZ.h b/Engine/Plugins/Experimental/DynamicMaterial/Source/DynamicMaterial/Public/Components/MaterialValues/DMMaterialValueFloat3XYZ.h
--- a/Engine/Plugins/Experimental/DynamicMaterial/Source/DynamicMaterial/Public/Components/MaterialValues/DMMaterialValueFloat3XYZ.h
+++ b/Engine/Plugins/Experimental/DynamicMaterial/Source/DynamicMaterial/Public/Components/MaterialValues/DMMaterialValueFloat3XYZ.h
@@ -52,6 +52,7 @@ public:
 	//~ Begin UDMMaterialComponent
 #if WITH_EDITOR
 	DYNAMICMATERIAL_API virtual FString GetComponentPathComponent() const override;
+	DYNAMICMATERIAL_API virtual FText GetComponentDescription() const override;
 #endif
 	//~ End UDMMaterialComponent
  
diff --git a/Engine/Plugins/Experimental/DynamicMaterial/Source/DynamicMaterial/Public/Components/MaterialValues/DMMaterialValueFloat4.h b/Engine/Plugins/Experimental/DynamicMaterial/Source/DynamicMaterial/Public/Components/MaterialValues/DMMaterialValueFloat4.h
--- a/Engine/Plugins/Experimental/DynamicMaterial/Source/DynamicMaterial/Public/Components/MaterialValues/DMMaterialValueFloat4.h
+++ b/Engine/Plugins/Experimental/DynamicMaterial/Source/DynamicMaterial/Public/Components/MaterialValues/DMMaterialValueFloat4.h
@@ -53,6 +53,7 @@ public:
 	//~ Begin UDMMaterialComponent
 #if WITH_EDITOR
 	DYNAMICMATERIAL_API virtual FString GetComponentPathComponent() const override;
+	DYNAMICMATERIAL_API virtual FText GetComponentDescription() const override;
 #endif
 	//~ End UDMMaterialComponent
 
diff --git a/Engine/Plugins/Experimental/DynamicMaterial/Source/DynamicMaterial/Public/Components/MaterialValues/DMMaterialValueRenderTarget.h b/Engine/Plugins/Experimental/DynamicMaterial/Source/DynamicMaterial/Public/Components/MaterialValues/DMMaterialValueRenderTarget.h
--- a/Engine/Plugins/Experimental/DynamicMaterial/Source/DynamicMaterial/Public/Components/MaterialValues/DMMaterialValueRenderTarget.h
+++ b/Engine/Plugins/Experimental/DynamicMaterial/Source/DynamicMaterial/Public/Components/MaterialValues/DMMaterialValueRenderTarget.h
@@ -79,6 +79,7 @@ public:
 	DYNAMICMATERIAL_API virtual void Update(UDMMaterialComponent* InSource, EDMUpdateType InUpdateType) override;
 #if WITH_EDITOR
 	DYNAMICMATERIAL_API virtual FString GetComponentPathComponent() const override;
+	DYNAMICMATERIAL_API virtual FText GetComponentDescription() const override;
 	DYNAMICMATERIAL_API virtual void PostEditorDuplicate(UDynamicMaterialModel* InMaterialModel, UDMMaterialComponent* InParent) override;
 #endif
 	//~ End UDMMaterialComponent
diff --git a/Engine/Plugins/Experimental/DynamicMaterial/Source/DynamicMaterial/Public/Components/MaterialValues/DMMaterialValueTexture.h b/Engine/Plugins/Experimental/DynamicMaterial/Source/DynamicMaterial/Public/Components/MaterialValues/DMMaterialValueTexture.h
--- a/Engine/Plugins/Experimental/DynamicMaterial/Source/DynamicMaterial/Public/Components/MaterialValues/DMMaterialValueTexture.h
+++ b/Engine/Plugins/Experimental/DynamicMaterial/Source/DynamicMaterial/Public/Components/MaterialValues/DMMaterialValueTexture.h
@@ -74,6 +74,7 @@ public:
 	//~ Begin UDMMaterialComponent
 #if WITH_EDITOR
 	DYNAMICMATERIAL_API virtual FString GetComponentPathComponent() const override;
+	DYNAMICMATERIAL_API virtual FText GetComponentDescription() const override;
 #endif
 	//~ End UDMMaterialComponent
 
diff --git a/Engine/Plugins/Experimental/DynamicMaterial/Source/DynamicMaterial/Public/Components/RenderTargetRenderers/DMRenderTargetTextRenderer.h b/Engine/Plugins/Experimental/DynamicMaterial/Source/DynamicMaterial/Public/Components/RenderTargetRenderers/DMRenderTargetTextRenderer.h
--- a/Engine/Plugins/Experimental/DynamicMaterial/Source/DynamicMaterial/Public/Components/RenderTargetRenderers/DMRenderTargetTextRenderer.h
+++ b/Engine/Plugins/Experimental/DynamicMaterial/Source/DynamicMaterial/Public/Components/RenderTargetRenderers/DMRenderTargetTextRenderer.h
@@ -187,6 +187,10 @@ public:
 	DYNAMICMATERIAL_API virtual bool JsonDeserialize(const TSharedPtr<FJsonValue>& InJsonValue) override;
 	//~ End IDMJsonSerializable
 
+	//~ Begin UDMMaterialComponent
+	DYNAMICMATERIAL_API virtual FText GetComponentDescription() const override;
+	//~ End UDMMaterialComponent
+
 	//~ Begin UObject
 	DYNAMICMATERIAL_API virtual void PostEditChangeProperty(FPropertyChangedEvent& InPropertyChangedEvent) override;
 	//~ End UObject
diff --git a/Engine/Plugins/Experimental/DynamicMaterial/Source/DynamicMaterial/Public/Components/RenderTargetRenderers/DMRenderTargetUMGWidgetRenderer.h b/Engine/Plugins/Experimental/DynamicMaterial/Source/DynamicMaterial/Public/Components/RenderTargetRenderers/DMRenderTargetUMGWidgetRenderer.h
--- a/Engine/Plugins/Experimental/DynamicMaterial/Source/DynamicMaterial/Public/Components/RenderTargetRenderers/DMRenderTargetUMGWidgetRenderer.h
+++ b/Engine/Plugins/Experimental/DynamicMaterial/Source/DynamicMaterial/Public/Components/RenderTargetRenderers/DMRenderTargetUMGWidgetRenderer.h
@@ -28,6 +28,10 @@ public:
 	DYNAMICMATERIAL_API virtual bool JsonDeserialize(const TSharedPtr<FJsonValue>& InJsonValue) override;
 	//~ End IDMJsonSerializable
 
+	//~ Begin UDMMaterialComponent
+	DYNAMICMATERIAL_API virtual FText GetComponentDescription() const override;
+	//~ End UDMMaterialComponent
+
 	//~ Begin UObject
 	DYNAMICMATERIAL_API virtual void PostEditChangeProperty(FPropertyChangedEvent& InPropertyChangedEvent) override;
 	//~ End UObject
diff --git a/Engine/Plugins/Experimental/DynamicMaterial/Source/DynamicMaterial/Public/Components/RenderTargetRenderers/DMRenderTargetWidgetRendererBase.h b/Engine/Plugins/Experimental/DynamicMaterial/Source/DynamicMaterial/Public/Components/RenderTargetRenderers/DMRenderTargetWidgetRendererBase.h
--- a/Engine/Plugins/Experimental/DynamicMaterial/Source/DynamicMaterial/Public/Components/RenderTargetRenderers/DMRenderTargetWidgetRendererBase.h
+++ b/Engine/Plugins/Experimental/DynamicMaterial/Source/DynamicMaterial/Public/Components/RenderTargetRenderers/DMRenderTargetWidgetRendererBase.h
@@ -20,6 +20,12 @@ class UDMRenderTargetWidgetRendererBase : public UDMRenderTargetRenderer
 public:
 	DYNAMICMATERIAL_API UDMRenderTargetWidgetRendererBase();
 
+#if WITH_EDITOR
+	//~ Begin UDMMaterialComponent
+	DYNAMICMATERIAL_API virtual FText GetComponentDescription() const override;
+	//~ End UDMMaterialComponent
+#endif
+
 protected:
 	TSharedPtr<SWidget> Widget;
 	TSharedPtr<FWidgetRenderer> WidgetRenderer;
diff --git a/Engine/Plugins/Experimental/DynamicMaterial/Source/DynamicMaterialEditor/Private/Components/MaterialStageExpressions/DMMSETextureSampleBase.cpp b/Engine/Plugins/Experimental/DynamicMaterial/Source/DynamicMaterialEditor/Private/Components/MaterialStageExpressions/DMMSETextureSampleBase.cpp
--- a/Engine/Plugins/Experimental/DynamicMaterial/Source/DynamicMaterialEditor/Private/Components/MaterialStageExpressions/DMMSETextureSampleBase.cpp
+++ b/Engine/Plugins/Experimental/DynamicMaterial/Source/DynamicMaterialEditor/Private/Components/MaterialStageExpressions/DMMSETextureSampleBase.cpp
@@ -238,6 +238,27 @@ void UDMMaterialStageExpressionTextureSampleBase::OnInputUpdated(int32 InInputIn
 	}
 }
 
+FText UDMMaterialStageExpressionTextureSampleBase::GetComponentDescription() const
+{
+	if (UDMMaterialStage* Stage = GetStage())
+	{
+		const TArray<UDMMaterialStageInput*>& Inputs = Stage->GetInputs();
+
+		for (UDMMaterialStageInput* StageInput : Inputs)
+		{
+			if (UDMMaterialStageInputValue* InputValue = Cast<UDMMaterialStageInputValue>(StageInput))
+			{
+				if (UDMMaterialValue* Value = Cast<UDMMaterialValueTexture>(InputValue->GetValue()))
+				{
+					return Value->GetComponentDescription();
+				}
+			}
+		}
+	}
+
+	return Super::GetComponentDescription();
+}
+
 void UDMMaterialStageExpressionTextureSampleBase::SetClampTextureEnabled(bool bInValue)
 {
 	if (bClampTexture == bInValue)
diff --git a/Engine/Plugins/Experimental/DynamicMaterial/Source/DynamicMaterialEditor/Private/UI/PropertyGenerators/DMStagePropertyRowGenerator.cpp b/Engine/Plugins/Experimental/DynamicMaterial/Source/DynamicMaterialEditor/Private/UI/PropertyGenerators/DMStagePropertyRowGenerator.cpp
--- a/Engine/Plugins/Experimental/DynamicMaterial/Source/DynamicMaterialEditor/Private/UI/PropertyGenerators/DMStagePropertyRowGenerator.cpp
+++ b/Engine/Plugins/Experimental/DynamicMaterial/Source/DynamicMaterialEditor/Private/UI/PropertyGenerators/DMStagePropertyRowGenerator.cpp
@@ -101,7 +101,7 @@ TSharedRef<SWidget> FDMStagePropertyRowGenerator::CreateSourceTypeEditWidget(con
 			[
 				SNew(STextBlock)
 				.Font(IDetailLayoutBuilder::GetDetailFont())
-				.Text_Static(&FDMStagePropertyRowGenerator::GetSourceTypeEditWidgetText, InStage)
+				.Text(GetSourceTypeEditWidgetText(InStage))
 			]
 			+ SHorizontalBox::Slot()
 			.AutoWidth()
diff --git a/Engine/Plugins/Experimental/DynamicMaterial/Source/DynamicMaterialEditor/Public/Components/MaterialStageExpressions/DMMSETextureSampleBase.h b/Engine/Plugins/Experimental/DynamicMaterial/Source/DynamicMaterialEditor/Public/Components/MaterialStageExpressions/DMMSETextureSampleBase.h
--- a/Engine/Plugins/Experimental/DynamicMaterial/Source/DynamicMaterialEditor/Public/Components/MaterialStageExpressions/DMMSETextureSampleBase.h
+++ b/Engine/Plugins/Experimental/DynamicMaterial/Source/DynamicMaterialEditor/Public/Components/MaterialStageExpressions/DMMSETextureSampleBase.h
@@ -38,6 +38,10 @@ public:
 	DYNAMICMATERIALEDITOR_API virtual void OnInputUpdated(int32 InInputIndex, EDMUpdateType InUpdateType) override;
 	//~ End UDMMaterialStageThroughput
 
+	//~ Begin UDMMaterialComponent
+	virtual FText GetComponentDescription() const override;
+	//~ End UDMMaterialComponent;
+
 	//~ Begin UObject
 	DYNAMICMATERIALEDITOR_API virtual void PostEditChangeProperty(FPropertyChangedEvent& InPropertyChangedEvent) override;
 	//~ End UObject
