Material Designer
- Removed the duplicate button.

diff --git a/Engine/Plugins/Experimental/DynamicMaterial/Source/DynamicMaterialEditor/Private/UI/Widgets/Editor/SDMMaterialSlotEditor.cpp b/Engine/Plugins/Experimental/DynamicMaterial/Source/DynamicMaterialEditor/Private/UI/Widgets/Editor/SDMMaterialSlotEditor.cpp
--- a/Engine/Plugins/Experimental/DynamicMaterial/Source/DynamicMaterialEditor/Private/UI/Widgets/Editor/SDMMaterialSlotEditor.cpp
+++ b/Engine/Plugins/Experimental/DynamicMaterial/Source/DynamicMaterialEditor/Private/UI/Widgets/Editor/SDMMaterialSlotEditor.cpp
@@ -820,25 +820,6 @@ TSharedRef<SWidget> SDMMaterialSlotEditor::CreateSlot_LayerSettings()
 			]
 		]
 
-		+ SHorizontalBox::Slot()
-		.AutoWidth()
-		.HAlign(HAlign_Right)
-		.VAlign(VAlign_Center)
-		.Padding(5.0f, 2.0f, 0.0f, 2.0f)
-		[
-			SNew(SButton)
-			.ContentPadding(4.0f)
-			.ButtonStyle(FDynamicMaterialEditorStyle::Get(), "HoverHintOnly.Bordered.Dark")
-			.ToolTipText(LOCTEXT("DuplicateSelectedLayer", "Duplicate Selected Layer"))
-			.IsEnabled(this, &SDMMaterialSlotEditor::GetLayerRowsButtonsCanDuplicate)
-			.OnClicked(this, &SDMMaterialSlotEditor::OnLayerRowButtonsDuplicateClicked)
-			[
-				SNew(SImage)
-				.Image(FDynamicMaterialEditorStyle::GetBrush("LayerView.DuplicateIcon"))
-				.DesiredSizeOverride(FVector2D(16.0f))
-			]
-		]
-
 		+ SHorizontalBox::Slot()
 		.AutoWidth()
 		.HAlign(HAlign_Right)
