Material Designer
- Non-selected, disabled layer is now a darker red instead of purple.

diff --git a/Engine/Plugins/Experimental/DynamicMaterial/Source/DynamicMaterialEditor/Private/DynamicMaterialEditorStyle.cpp b/Engine/Plugins/Experimental/DynamicMaterial/Source/DynamicMaterialEditor/Private/DynamicMaterialEditorStyle.cpp
--- a/Engine/Plugins/Experimental/DynamicMaterial/Source/DynamicMaterialEditor/Private/DynamicMaterialEditorStyle.cpp
+++ b/Engine/Plugins/Experimental/DynamicMaterial/Source/DynamicMaterialEditor/Private/DynamicMaterialEditorStyle.cpp
@@ -238,7 +238,7 @@ void FDynamicMaterialEditorStyle::SetupStageStyles(const TSharedRef<FSlateStyleS
 
 	const FLinearColor EnabledColor = FStyleColors::Foreground.GetSpecifiedColor();
 	const FLinearColor EnabledSelectedColor = FStyleColors::Primary.GetSpecifiedColor();
-	const FLinearColor DisabledColor = FStyleColors::AccentPurple.GetSpecifiedColor();
+	const FLinearColor DisabledColor = FStyleColors::AccentRed.GetSpecifiedColor() * FLinearColor(0.5f, 0.5f, 0.5f, 1.f);
 	const FLinearColor DisabledSelectedColor = FStyleColors::AccentRed.GetSpecifiedColor();
 
 	Style->Set("Stage.Inactive", new FSlateRoundedBoxBrush(
