Material Designer
- Fixes switching visual bug with selected properties when adding a layer / material recompile.
- Changed order of text properties.

diff --git a/Engine/Plugins/Experimental/DynamicMaterial/Source/DynamicMaterial/Private/Components/RenderTargetRenderers/DMRenderTargetTextRenderer.cpp b/Engine/Plugins/Experimental/DynamicMaterial/Source/DynamicMaterial/Private/Components/RenderTargetRenderers/DMRenderTargetTextRenderer.cpp
--- a/Engine/Plugins/Experimental/DynamicMaterial/Source/DynamicMaterial/Private/Components/RenderTargetRenderers/DMRenderTargetTextRenderer.cpp
+++ b/Engine/Plugins/Experimental/DynamicMaterial/Source/DynamicMaterial/Private/Components/RenderTargetRenderers/DMRenderTargetTextRenderer.cpp
@@ -56,16 +56,16 @@ struct FDMRenderTargetTextRenderer
 
 	static const inline TSet<FName> PropertyNames = {
 		FontInfoName, TextName, TextColorName,
+		JustifyName,
+		FlowDirectionName,
+		LineHeightName,
+		AutoWrapTextName, WrapTextAtName, WrappingPolicyName,
+		PaddingLeftName, PaddingTopName, PaddingRightName, PaddingBottomName,
 		HasHighlightName, HighlightColorName,
 		HasShadowName, ShadowColorName, ShadowOffsetName,
-		AutoWrapTextName, WrapTextAtName, WrappingPolicyName,
-		JustifyName,
 		TransformPolicyName,
-		FlowDirectionName,
 		ShapingMethodName,
 		StrikeBrushName,
-		LineHeightName,
-		PaddingLeftName, PaddingTopName, PaddingRightName, PaddingBottomName,
 		OverrideRenderTargetSizeName
 	};
 };
diff --git a/Engine/Plugins/Experimental/DynamicMaterial/Source/DynamicMaterialEditor/Private/UI/Widgets/SDMMaterialEditor.cpp b/Engine/Plugins/Experimental/DynamicMaterial/Source/DynamicMaterialEditor/Private/UI/Widgets/SDMMaterialEditor.cpp
--- a/Engine/Plugins/Experimental/DynamicMaterial/Source/DynamicMaterialEditor/Private/UI/Widgets/SDMMaterialEditor.cpp
+++ b/Engine/Plugins/Experimental/DynamicMaterial/Source/DynamicMaterialEditor/Private/UI/Widgets/SDMMaterialEditor.cpp
@@ -814,7 +814,7 @@ TSharedRef<SDMMaterialPropertySelector> SDMMaterialEditor::CreateSlot_PropertySe
 {
 	TSharedRef<SDMMaterialPropertySelector> NewPropertySelector = CreateSlot_PropertySelector_Impl();
 
-	if (EditMode == EDMMaterialEditorMode::EditSlot && SelectedMaterialProperty != EDMMaterialPropertyType::None)
+	if (EditMode == EDMMaterialEditorMode::EditSlot && SelectedMaterialProperty == EDMMaterialPropertyType::None)
 	{
 		if (UDynamicMaterialModel* MaterialModel = GetMaterialModel())
 		{
