Material Designer
- Removed axes from material previews
- Removed physics from material preview scenes.

diff --git a/Engine/Plugins/Experimental/DynamicMaterial/Source/DynamicMaterialEditor/Private/UI/Widgets/Editor/DMMaterialPreviewViewportClient.cpp b/Engine/Plugins/Experimental/DynamicMaterial/Source/DynamicMaterialEditor/Private/UI/Widgets/Editor/DMMaterialPreviewViewportClient.cpp
--- a/Engine/Plugins/Experimental/DynamicMaterial/Source/DynamicMaterialEditor/Private/UI/Widgets/Editor/DMMaterialPreviewViewportClient.cpp
+++ b/Engine/Plugins/Experimental/DynamicMaterial/Source/DynamicMaterialEditor/Private/UI/Widgets/Editor/DMMaterialPreviewViewportClient.cpp
@@ -13,6 +13,8 @@ FDMMaterialPreviewViewportClient::FDMMaterialPreviewViewportClient(const TShared
 {
 	PreviewWidget = InPreviewWidget;
 
+	bDrawAxes = false;
+
 	// Setup defaults for the common draw helper.
 	DrawHelper.bDrawPivot = false;
 	DrawHelper.bDrawWorldBox = false;
@@ -29,6 +31,7 @@ FDMMaterialPreviewViewportClient::FDMMaterialPreviewViewportClient(const TShared
 	EngineShowFlags.DisableAdvancedFeatures();
 	EngineShowFlags.SetSnap(0);
 	EngineShowFlags.SetSeparateTranslucency(true);
+	EngineShowFlags.SetPivot(false);
 
 	OverrideNearClipPlane(1.0f);
 	bUsingOrbitCamera = true;
diff --git a/Engine/Plugins/Experimental/DynamicMaterial/Source/DynamicMaterialEditor/Private/UI/Widgets/Editor/SDMMaterialPreview.cpp b/Engine/Plugins/Experimental/DynamicMaterial/Source/DynamicMaterialEditor/Private/UI/Widgets/Editor/SDMMaterialPreview.cpp
--- a/Engine/Plugins/Experimental/DynamicMaterial/Source/DynamicMaterialEditor/Private/UI/Widgets/Editor/SDMMaterialPreview.cpp
+++ b/Engine/Plugins/Experimental/DynamicMaterial/Source/DynamicMaterialEditor/Private/UI/Widgets/Editor/SDMMaterialPreview.cpp
@@ -83,7 +83,11 @@ void SDMMaterialPreview::Construct(const FArguments& InArgs, const TSharedRef<SD
 {
 	EditorWidgetWeak = InEditorWidget;
 	bShowMenu = InArgs._ShowMenu;
-	PreviewScene = MakeShareable(new FAdvancedPreviewScene(FPreviewScene::ConstructionValues()));
+	PreviewScene = MakeShareable(new FAdvancedPreviewScene(
+		FPreviewScene::ConstructionValues()
+		.SetCreatePhysicsScene(false)
+		.ShouldSimulatePhysics(false)
+	));
 
 	PreviewMaterial = nullptr;
 	PreviewMeshComponent = nullptr;
