Material Designer
- Fixed "Create with Material Designer" button not creating a material when not used with motion design meshes.

diff --git a/Engine/Plugins/Experimental/DynamicMaterial/Source/DynamicMaterialEditor/Private/DetailsPanel/Widgets/SDMMaterialListExtensionWidget.cpp b/Engine/Plugins/Experimental/DynamicMaterial/Source/DynamicMaterialEditor/Private/DetailsPanel/Widgets/SDMMaterialListExtensionWidget.cpp
--- a/Engine/Plugins/Experimental/DynamicMaterial/Source/DynamicMaterialEditor/Private/DetailsPanel/Widgets/SDMMaterialListExtensionWidget.cpp
+++ b/Engine/Plugins/Experimental/DynamicMaterial/Source/DynamicMaterialEditor/Private/DetailsPanel/Widgets/SDMMaterialListExtensionWidget.cpp
@@ -115,10 +115,12 @@ void SDMMaterialListExtensionWidget::SetDynamicMaterialInstance(UDynamicMaterial
 						{
 							const FMaterialListItem& ListItem = MaterialItemViewWeak.Pin()->GetMaterialListItem();
 							const FDMObjectMaterialProperty MaterialProperty(CurrentComponent, ListItem.SlotIndex);
-							WorldSubsystem->ExecuteMaterialValueSetterDelegate(MaterialProperty, NewInstance);
-						}
 
-						return;
+							if (WorldSubsystem->ExecuteMaterialValueSetterDelegate(MaterialProperty, NewInstance))
+							{
+								return;
+							}
+						}
 					}
 				}
 			}
