Material Designer
- Added async update callbacks for when the component previews recreate materials based on structural changes.

diff --git a/Engine/Plugins/Experimental/DynamicMaterial/Source/DynamicMaterialEditor/Private/UI/Widgets/Visualizers/SDMMaterialComponentPreview.cpp b/Engine/Plugins/Experimental/DynamicMaterial/Source/DynamicMaterialEditor/Private/UI/Widgets/Visualizers/SDMMaterialComponentPreview.cpp
--- a/Engine/Plugins/Experimental/DynamicMaterial/Source/DynamicMaterialEditor/Private/UI/Widgets/Visualizers/SDMMaterialComponentPreview.cpp
+++ b/Engine/Plugins/Experimental/DynamicMaterial/Source/DynamicMaterialEditor/Private/UI/Widgets/Visualizers/SDMMaterialComponentPreview.cpp
@@ -24,6 +24,12 @@ SDMMaterialComponentPreview::SDMMaterialComponentPreview()
 
 SDMMaterialComponentPreview::~SDMMaterialComponentPreview()
 {
+	if (EndOfFrameDelegateHandle.IsValid())
+	{
+		FCoreDelegates::OnEndFrame.Remove(EndOfFrameDelegateHandle);
+		EndOfFrameDelegateHandle.Reset();
+	}
+
 	if (!FDynamicMaterialModule::AreUObjectsSafe())
 	{
 		return;
@@ -49,19 +55,6 @@ void SDMMaterialComponentPreview::Construct(const FArguments& InArgs, const TSha
 	PreviewMaterialDynamicWeak = InEditorWidget->GetPreviewMaterialManager()->CreatePreviewMaterialDynamic(PreviewMaterialBaseWeak.Get());
 	MaterialModelBaseWeak = InEditorWidget->GetMaterialModelBase();
 
-	if (UDynamicMaterialModel* MaterialModel = Cast<UDynamicMaterialModel>(MaterialModelBaseWeak.Get()))
-	{
-		MaterialModel->ApplyComponents(PreviewMaterialDynamicWeak.Get());
-	}
-	else if (UDynamicMaterialModelDynamic* MaterialModelDynamic = Cast<UDynamicMaterialModelDynamic>(MaterialModelBaseWeak.Get()))
-	{
-		MaterialModelDynamic->ApplyComponents(PreviewMaterialDynamicWeak.Get());
-	}
-	else
-	{
-		return;
-	}
-
 	SetCanTick(true);
 
 	if (ensure(IsValid(InComponent)))
@@ -90,28 +83,27 @@ void SDMMaterialComponentPreview::Tick(const FGeometry& AllottedGeometry, const
 
 void SDMMaterialComponentPreview::OnComponentUpdated(UDMMaterialComponent* InComponent, UDMMaterialComponent* InSource, EDMUpdateType InUpdateType)
 {
-	if (InComponent != ComponentWeak.Get() || !IsValid(InComponent) || !InComponent->IsComponentValid())
+	// We're going to recreate the material anyway so this doesn't need to run.
+	if (EndOfFrameDelegateHandle.IsValid())
 	{
 		return;
 	}
 
-	TSharedPtr<SDMMaterialEditor> EditorWidget = EditorWidgetWeak.Pin();
-
-	if (!EditorWidget.IsValid())
+	if (InComponent != ComponentWeak.Get() || !IsValid(InComponent) || !InComponent->IsComponentValid())
 	{
 		return;
 	}
 
-	UMaterial* PreviewMaterialBase = PreviewMaterialBaseWeak.Get();
+	TSharedPtr<SDMMaterialEditor> EditorWidget = EditorWidgetWeak.Pin();
 
-	if (!PreviewMaterialBase)
+	if (!EditorWidget.IsValid())
 	{
-		PreviewMaterialBase = EditorWidget->GetPreviewMaterialManager()->CreatePreviewMaterial(InComponent);
+		return;
 	}
 
 	UMaterialInstanceDynamic* MID = PreviewMaterialDynamicWeak.Get();
 
-	if (!MID || !EnumHasAnyFlags(InUpdateType, EDMUpdateType::Structure))
+	if (MID && !EnumHasAnyFlags(InUpdateType, EDMUpdateType::Structure))
 	{
 		if (UDMMaterialValue* Value = Cast<UDMMaterialValue>(InSource))
 		{
@@ -130,7 +122,47 @@ void SDMMaterialComponentPreview::OnComponentUpdated(UDMMaterialComponent* InCom
 			TextureUVDynamic->SetMIDParameters(MID);
 		}
 	}
-	else if (UDMMaterialStage* Stage = Cast<UDMMaterialStage>(InComponent))
+	else if (!EndOfFrameDelegateHandle.IsValid())
+	{
+		EndOfFrameDelegateHandle = FCoreDelegates::OnEndFrame.AddSP(this, &SDMMaterialComponentPreview::OnEndOfFrame);
+	}
+}
+
+void SDMMaterialComponentPreview::OnEndOfFrame()
+{
+	if (EndOfFrameDelegateHandle.IsValid())
+	{
+		FCoreDelegates::OnEndFrame.Remove(EndOfFrameDelegateHandle);
+		EndOfFrameDelegateHandle.Reset();
+	}
+
+	RecreateMaterial();
+}
+
+void SDMMaterialComponentPreview::RecreateMaterial()
+{
+	TSharedPtr<SDMMaterialEditor> EditorWidget = EditorWidgetWeak.Pin();
+
+	if (!EditorWidget.IsValid())
+	{
+		return;
+	}
+
+	UDMMaterialComponent* Component = ComponentWeak.Get();
+
+	if (!Component)
+	{
+		return;
+	}
+
+	UMaterial* PreviewMaterialBase = PreviewMaterialBaseWeak.Get();
+
+	if (!PreviewMaterialBase)
+	{
+		PreviewMaterialBase = EditorWidget->GetPreviewMaterialManager()->CreatePreviewMaterial(Component);
+	}
+
+	if (UDMMaterialStage* Stage = Cast<UDMMaterialStage>(Component))
 	{
 		Stage->GeneratePreviewMaterial(PreviewMaterialBase);
 
@@ -150,7 +182,7 @@ void SDMMaterialComponentPreview::OnComponentUpdated(UDMMaterialComponent* InCom
 
 		Brush.SetMaterial(PreviewMaterialDynamicWeak.Get());
 	}
-	else if (UDMMaterialProperty* Property = Cast<UDMMaterialProperty>(InComponent))
+	else if (UDMMaterialProperty* Property = Cast<UDMMaterialProperty>(Component))
 	{
 		Property->GeneratePreviewMaterial(PreviewMaterialBase);
 
diff --git a/Engine/Plugins/Experimental/DynamicMaterial/Source/DynamicMaterialEditor/Private/UI/Widgets/Visualizers/SDMMaterialComponentPreview.h b/Engine/Plugins/Experimental/DynamicMaterial/Source/DynamicMaterialEditor/Private/UI/Widgets/Visualizers/SDMMaterialComponentPreview.h
--- a/Engine/Plugins/Experimental/DynamicMaterial/Source/DynamicMaterialEditor/Private/UI/Widgets/Visualizers/SDMMaterialComponentPreview.h
+++ b/Engine/Plugins/Experimental/DynamicMaterial/Source/DynamicMaterialEditor/Private/UI/Widgets/Visualizers/SDMMaterialComponentPreview.h
@@ -4,6 +4,7 @@
 
 #include "Widgets/SCompoundWidget.h"
 
+#include "Delegates/IDelegateInstance.h"
 #include "DMEDefs.h"
 #include "SlateMaterialBrush.h"
 #include "UObject/WeakObjectPtr.h"
@@ -45,6 +46,11 @@ protected:
 	TWeakObjectPtr<UMaterial> PreviewMaterialBaseWeak;
 	TWeakObjectPtr<UMaterialInstanceDynamic> PreviewMaterialDynamicWeak;
 	FSlateMaterialBrush Brush;
+	FDelegateHandle EndOfFrameDelegateHandle;
 
 	void OnComponentUpdated(UDMMaterialComponent* InComponent, UDMMaterialComponent* InSource, EDMUpdateType InUpdateType);
+
+	void OnEndOfFrame();
+
+	void RecreateMaterial();
 };
