Material Designer
- Renamed "Property Previews" to "Channels".
- Added the ability to use long names on the property selector even in slim layouts.

diff --git a/Engine/Plugins/Experimental/DynamicMaterial/Source/DynamicMaterialEditor/Private/DynamicMaterialEditorSettings.cpp b/Engine/Plugins/Experimental/DynamicMaterial/Source/DynamicMaterialEditor/Private/DynamicMaterialEditorSettings.cpp
--- a/Engine/Plugins/Experimental/DynamicMaterial/Source/DynamicMaterialEditor/Private/DynamicMaterialEditorSettings.cpp
+++ b/Engine/Plugins/Experimental/DynamicMaterial/Source/DynamicMaterialEditor/Private/DynamicMaterialEditorSettings.cpp
@@ -268,6 +268,7 @@ void UDynamicMaterialEditorSettings::PostEditChangeProperty(FPropertyChangedEven
 	const FName PropertyName = InPropertyChangedEvent.GetMemberPropertyName();
 
 	if (PropertyName == GET_MEMBER_NAME_CHECKED(UDynamicMaterialEditorSettings, Layout)
+		|| PropertyName == GET_MEMBER_NAME_CHECKED(UDynamicMaterialEditorSettings, bUseFullChannelNamesInTopSlimLayout)
 		|| PropertyName == GET_MEMBER_NAME_CHECKED(UDynamicMaterialEditorSettings, SplitterLocation)
 		|| PropertyName == GET_MEMBER_NAME_CHECKED(UDynamicMaterialEditorSettings, bPreviewImagesUseTextureUVs)
 		|| PropertyName == GET_MEMBER_NAME_CHECKED(UDynamicMaterialEditorSettings, PreviewMesh)
@@ -299,6 +300,7 @@ void UDynamicMaterialEditorSettings::OpenEditorSettingsWindow() const
 void UDynamicMaterialEditorSettings::ResetAllLayoutSettings()
 {
 	Layout = EDMMaterialEditorLayout::Left;
+	bUseFullChannelNamesInTopSlimLayout = false;
 	SplitterLocation = 240;
 	ThumbnailSize = 256.f;
 	PreviewMesh = EDMMaterialPreviewMesh::Plane;
diff --git a/Engine/Plugins/Experimental/DynamicMaterial/Source/DynamicMaterialEditor/Private/DynamicMaterialEditorSettings.h b/Engine/Plugins/Experimental/DynamicMaterial/Source/DynamicMaterialEditor/Private/DynamicMaterialEditorSettings.h
--- a/Engine/Plugins/Experimental/DynamicMaterial/Source/DynamicMaterialEditor/Private/DynamicMaterialEditorSettings.h
+++ b/Engine/Plugins/Experimental/DynamicMaterial/Source/DynamicMaterialEditor/Private/DynamicMaterialEditorSettings.h
@@ -182,6 +182,9 @@ public:
 	UPROPERTY(Config, EditAnywhere, BlueprintReadWrite, Category = "Layout")
 	bool bUVVisualizerVisible;
 
+	UPROPERTY(Config, EditAnywhere, BlueprintReadWrite, Category = "Layout")
+	bool bUseFullChannelNamesInTopSlimLayout;
+
 	UPROPERTY(Config, EditAnywhere, BlueprintReadWrite, Category = "Thumnnails")
 	double ThumbnailSize;
 
diff --git a/Engine/Plugins/Experimental/DynamicMaterial/Source/DynamicMaterialEditor/Private/UI/Widgets/Editor/PropertySelectorLayouts/SDMMaterialPropertySelector_WrapSlim.cpp b/Engine/Plugins/Experimental/DynamicMaterial/Source/DynamicMaterialEditor/Private/UI/Widgets/Editor/PropertySelectorLayouts/SDMMaterialPropertySelector_WrapSlim.cpp
--- a/Engine/Plugins/Experimental/DynamicMaterial/Source/DynamicMaterialEditor/Private/UI/Widgets/Editor/PropertySelectorLayouts/SDMMaterialPropertySelector_WrapSlim.cpp
+++ b/Engine/Plugins/Experimental/DynamicMaterial/Source/DynamicMaterialEditor/Private/UI/Widgets/Editor/PropertySelectorLayouts/SDMMaterialPropertySelector_WrapSlim.cpp
@@ -4,6 +4,7 @@
 
 #include "Components/DMMaterialProperty.h"
 #include "DetailLayoutBuilder.h"
+#include "DynamicMaterialEditorSettings.h"
 #include "Model/DynamicMaterialModelEditorOnlyData.h"
 #include "UI/Widgets/SDMMaterialEditor.h"
 #include "Utils/DMPrivate.h"
@@ -82,6 +83,14 @@ TSharedRef<SWidget> SDMMaterialPropertySelector_WrapSlim::CreateSlot_PropertyLis
 
 TSharedRef<SWidget> SDMMaterialPropertySelector_WrapSlim::CreateSlot_SelectButton(EDMMaterialEditorMode InEditMode, EDMMaterialPropertyType InMaterialProperty)
 {
+	if (const UDynamicMaterialEditorSettings* Settings = GetDefault<UDynamicMaterialEditorSettings>())
+	{
+		if (Settings->bUseFullChannelNamesInTopSlimLayout)
+		{
+			return SDMMaterialPropertySelector_WrapBase::CreateSlot_SelectButton(InEditMode, InMaterialProperty);
+		}
+	}
+
 	const FText ButtonText = GetSelectButtonText(InEditMode, InMaterialProperty, /* Short Name */ true);
 	const FText ToolTip = GetButtonToolTip(InEditMode, InMaterialProperty);
 
diff --git a/Engine/Plugins/Experimental/DynamicMaterial/Source/DynamicMaterialEditor/Private/UI/Widgets/Editor/SDMMaterialPropertySelector.cpp b/Engine/Plugins/Experimental/DynamicMaterial/Source/DynamicMaterialEditor/Private/UI/Widgets/Editor/SDMMaterialPropertySelector.cpp
--- a/Engine/Plugins/Experimental/DynamicMaterial/Source/DynamicMaterialEditor/Private/UI/Widgets/Editor/SDMMaterialPropertySelector.cpp
+++ b/Engine/Plugins/Experimental/DynamicMaterial/Source/DynamicMaterialEditor/Private/UI/Widgets/Editor/SDMMaterialPropertySelector.cpp
@@ -170,8 +170,8 @@ FText SDMMaterialPropertySelector::GetSelectButtonText(EDMMaterialEditorMode InE
 
 		case EDMMaterialEditorMode::PropertyPreviews:
 			return bInShortName
-				? LOCTEXT("Properties", "Props")
-				: LOCTEXT("PropertyPreviews", "Properties");
+				? LOCTEXT("ChannelsShort", "Chans")
+				: LOCTEXT("Channels", "Channels");
 
 		case EDMMaterialEditorMode::EditSlot:
 			return bInShortName
@@ -194,11 +194,11 @@ FText SDMMaterialPropertySelector::GetButtonToolTip(EDMMaterialEditorMode InEdit
 			return LOCTEXT("GeneralSettingsToolTip", "Edit the Material Global Settings.");
 
 		case EDMMaterialEditorMode::PropertyPreviews:
-			return LOCTEXT("PropertyPreviewsToolTip", "Preview and toggle all the Material Properties.");
+			return LOCTEXT("PropertyPreviewsToolTip", "Preview and toggle the Material Channels.");
 
 		case EDMMaterialEditorMode::EditSlot:
 		{
-			const FText Format = LOCTEXT("PropertySelectFormat", "Edit the {0} property.");
+			const FText Format = LOCTEXT("PropertySelectFormat", "Edit the {0} channel.");
 			return FText::Format(Format, GetSelectButtonText(InEditMode, InMaterialProperty, /* Short Name */ false));
 		}
 
diff --git a/Engine/Plugins/Experimental/DynamicMaterial/Source/DynamicMaterialEditor/Private/UI/Widgets/SDMMaterialEditor.cpp b/Engine/Plugins/Experimental/DynamicMaterial/Source/DynamicMaterialEditor/Private/UI/Widgets/SDMMaterialEditor.cpp
--- a/Engine/Plugins/Experimental/DynamicMaterial/Source/DynamicMaterialEditor/Private/UI/Widgets/SDMMaterialEditor.cpp
+++ b/Engine/Plugins/Experimental/DynamicMaterial/Source/DynamicMaterialEditor/Private/UI/Widgets/SDMMaterialEditor.cpp
@@ -68,6 +68,11 @@ SDMMaterialEditor::~SDMMaterialEditor()
 		EditorOnlyData->GetOnPropertyUpdateDelegate().RemoveAll(this);
 		EditorOnlyData->GetOnSlotListUpdateDelegate().RemoveAll(this);
 	}
+
+	if (UDynamicMaterialEditorSettings* Settings = GetMutableDefault<UDynamicMaterialEditorSettings>())
+	{
+		Settings->GetOnSettingsChanged().RemoveAll(this);
+	}
 }
 
 void SDMMaterialEditor::Construct(const FArguments& InArgs, const TSharedRef<SDMMaterialDesigner>& InDesignerWidget)
@@ -94,6 +99,11 @@ void SDMMaterialEditor::Construct(const FArguments& InArgs, const TSharedRef<SDM
 	}
 
 	FCoreDelegates::OnEnginePreExit.AddSP(this, &SDMMaterialEditor::OnEnginePreExit);
+
+	if (UDynamicMaterialEditorSettings* Settings = GetMutableDefault<UDynamicMaterialEditorSettings>())
+	{
+		Settings->GetOnSettingsChanged().AddSP(this, &SDMMaterialEditor::OnSettingsChanged);
+	}
 }
 
 TSharedPtr<SDMMaterialDesigner> SDMMaterialEditor::GetDesignerWidget() const
@@ -940,4 +950,19 @@ void SDMMaterialEditor::OnSlotListUpdate(UDynamicMaterialModelBase* InMaterialMo
 	PropertySelectorSlot.Invalidate();
 }
 
+void SDMMaterialEditor::OnSettingsChanged(const FPropertyChangedEvent& InPropertyChangedEvent)
+{
+	if (!PropertySelectorSlot.IsValid())
+	{
+		return;
+	}
+
+	const FName MemberName = InPropertyChangedEvent.GetMemberPropertyName();
+
+	if (MemberName == GET_MEMBER_NAME_CHECKED(UDynamicMaterialEditorSettings, bUseFullChannelNamesInTopSlimLayout))
+	{
+		PropertySelectorSlot.Invalidate();
+	}
+}
+
 #undef LOCTEXT_NAMESPACE
diff --git a/Engine/Plugins/Experimental/DynamicMaterial/Source/DynamicMaterialEditor/Private/UI/Widgets/SDMMaterialEditor.h b/Engine/Plugins/Experimental/DynamicMaterial/Source/DynamicMaterialEditor/Private/UI/Widgets/SDMMaterialEditor.h
--- a/Engine/Plugins/Experimental/DynamicMaterial/Source/DynamicMaterialEditor/Private/UI/Widgets/SDMMaterialEditor.h
+++ b/Engine/Plugins/Experimental/DynamicMaterial/Source/DynamicMaterialEditor/Private/UI/Widgets/SDMMaterialEditor.h
@@ -33,6 +33,7 @@ class UDynamicMaterialModelBase;
 class UDynamicMaterialModelEditorOnlyData;
 enum class EDMMaterialPropertyType : uint8;
 enum class EDMUpdateType : uint8;
+struct InPropertyChangedEvent;
 
 namespace UE::DynamicMaterialEditor::Private
 {
@@ -230,4 +231,6 @@ protected:
 	void OnPropertyUpdate(UDynamicMaterialModelBase* InMaterialModelBase);
 
 	void OnSlotListUpdate(UDynamicMaterialModelBase* InMaterialModelBase);
+
+	void OnSettingsChanged(const FPropertyChangedEvent& InPropertyChangedEvent);
 };
