Material Designer
- Changed the "General" category to something more descriptive.

diff --git a/Engine/Plugins/Experimental/DynamicMaterial/Source/DynamicMaterialEditor/Private/UI/Widgets/Editor/SDMMaterialComponentEditor.cpp b/Engine/Plugins/Experimental/DynamicMaterial/Source/DynamicMaterialEditor/Private/UI/Widgets/Editor/SDMMaterialComponentEditor.cpp
--- a/Engine/Plugins/Experimental/DynamicMaterial/Source/DynamicMaterialEditor/Private/UI/Widgets/Editor/SDMMaterialComponentEditor.cpp
+++ b/Engine/Plugins/Experimental/DynamicMaterial/Source/DynamicMaterialEditor/Private/UI/Widgets/Editor/SDMMaterialComponentEditor.cpp
@@ -2,10 +2,15 @@
 
 #include "UI/Widgets/Editor/SDMMaterialComponentEditor.h"
 
-#include "DynamicMaterialEditorModule.h"
 #include "Components/DMMaterialComponent.h"
 #include "Components/DMMaterialLayer.h"
+#include "CustomDetailsViewArgs.h"
+#include "DynamicMaterialEditorModule.h"
 #include "DynamicMaterialModule.h"
+#include "ICustomDetailsView.h"
+#include "Items/ICustomDetailsViewCustomCategoryItem.h"
+#include "Items/ICustomDetailsViewItem.h"
+#include "UI/Utils/DMWidgetStatics.h"
 #include "UI/Widgets/SDMMaterialEditor.h"
 #include "Widgets/Input/SComboButton.h"
 
@@ -61,6 +66,38 @@ void SDMMaterialComponentEditor::OnComponentUpdated(UDMMaterialComponent* InComp
 	}
 }
 
+TSharedRef<ICustomDetailsViewItem> SDMMaterialComponentEditor::GetDefaultCategory(const TSharedRef<ICustomDetailsView>& InDetailsView,
+	const FCustomDetailsViewItemId& InRootId)
+{
+	UDMMaterialComponent* Component = GetComponent();
+
+	if (!Component)
+	{
+		return SDMObjectEditorWidgetBase::GetDefaultCategory(InDetailsView, InRootId);
+	}
+
+	if (!DefaultCategoryItem.IsValid())
+	{
+		const FText ComponentCategoryFormat = LOCTEXT("ComponantCategoryFormat", "{0} Settings");
+		const FText ComponentCategoryText = FText::Format(ComponentCategoryFormat, Component->GetComponentDescription());
+		DefaultCategoryItem = InDetailsView->CreateCustomCategoryItem(DefaultCategoryName, ComponentCategoryText)->AsItem();
+		DefaultCategoryItem->RefreshItemId();
+		InDetailsView->ExtendTree(InRootId, ECustomDetailsTreeInsertPosition::Child, DefaultCategoryItem.ToSharedRef());
+
+		bool bExpansionState = true;
+		FDMWidgetStatics::Get().GetExpansionState(ObjectWeak.Get(), DefaultCategoryName, bExpansionState);
+
+		InDetailsView->SetItemExpansionState(
+			DefaultCategoryItem->GetItemId(),
+			bExpansionState ? ECustomDetailsViewExpansion::SelfExpanded : ECustomDetailsViewExpansion::Collapsed
+		);
+
+		Categories.Add(DefaultCategoryName);
+	}
+
+	return DefaultCategoryItem.ToSharedRef();
+}
+
 TArray<FDMPropertyHandle> SDMMaterialComponentEditor::GetPropertyRows()
 {
 	TArray<FDMPropertyHandle> PropertyRows;
diff --git a/Engine/Plugins/Experimental/DynamicMaterial/Source/DynamicMaterialEditor/Private/UI/Widgets/Editor/SDMMaterialComponentEditor.h b/Engine/Plugins/Experimental/DynamicMaterial/Source/DynamicMaterialEditor/Private/UI/Widgets/Editor/SDMMaterialComponentEditor.h
--- a/Engine/Plugins/Experimental/DynamicMaterial/Source/DynamicMaterialEditor/Private/UI/Widgets/Editor/SDMMaterialComponentEditor.h
+++ b/Engine/Plugins/Experimental/DynamicMaterial/Source/DynamicMaterialEditor/Private/UI/Widgets/Editor/SDMMaterialComponentEditor.h
@@ -26,6 +26,8 @@ protected:
 	void OnComponentUpdated(UDMMaterialComponent* InComponent, UDMMaterialComponent* InSource, EDMUpdateType InUpdateType);
 
 	//~ Begin SDMObjectEditorWidgetBase
+	virtual TSharedRef<ICustomDetailsViewItem> GetDefaultCategory(const TSharedRef<ICustomDetailsView>& InDetailsView,
+		const FCustomDetailsViewItemId& InRootId) override;
 	virtual TArray<FDMPropertyHandle> GetPropertyRows() override;
 	virtual void OnUndo() override;
 	//~ End SDMObjectEditorWidgetBase
diff --git a/Engine/Plugins/Experimental/DynamicMaterial/Source/DynamicMaterialEditor/Private/UI/Widgets/Editor/SDMObjectEditorWidgetBase.cpp b/Engine/Plugins/Experimental/DynamicMaterial/Source/DynamicMaterialEditor/Private/UI/Widgets/Editor/SDMObjectEditorWidgetBase.cpp
--- a/Engine/Plugins/Experimental/DynamicMaterial/Source/DynamicMaterialEditor/Private/UI/Widgets/Editor/SDMObjectEditorWidgetBase.cpp
+++ b/Engine/Plugins/Experimental/DynamicMaterial/Source/DynamicMaterialEditor/Private/UI/Widgets/Editor/SDMObjectEditorWidgetBase.cpp
@@ -146,8 +146,6 @@ TSharedRef<ICustomDetailsViewItem> SDMObjectEditorWidgetBase::GetDefaultCategory
 {
 	if (!DefaultCategoryItem.IsValid())
 	{
-		constexpr const TCHAR* DefaultCategoryName = TEXT("General");
-
 		DefaultCategoryItem = InDetailsView->CreateCustomCategoryItem(DefaultCategoryName, LOCTEXT("General", "General"))->AsItem();
 		DefaultCategoryItem->RefreshItemId();
 		InDetailsView->ExtendTree(InRootId, ECustomDetailsTreeInsertPosition::Child, DefaultCategoryItem.ToSharedRef());
diff --git a/Engine/Plugins/Experimental/DynamicMaterial/Source/DynamicMaterialEditor/Private/UI/Widgets/Editor/SDMObjectEditorWidgetBase.h b/Engine/Plugins/Experimental/DynamicMaterial/Source/DynamicMaterialEditor/Private/UI/Widgets/Editor/SDMObjectEditorWidgetBase.h
--- a/Engine/Plugins/Experimental/DynamicMaterial/Source/DynamicMaterialEditor/Private/UI/Widgets/Editor/SDMObjectEditorWidgetBase.h
+++ b/Engine/Plugins/Experimental/DynamicMaterial/Source/DynamicMaterialEditor/Private/UI/Widgets/Editor/SDMObjectEditorWidgetBase.h
@@ -45,6 +45,8 @@ public:
 	//~ End FUndoClient
 
 protected:
+	static constexpr const TCHAR* DefaultCategoryName = TEXT("General");
+
 	TWeakPtr<SDMMaterialEditor> EditorWidgetWeak;
 	TWeakObjectPtr<UObject> ObjectWeak;
 
@@ -57,7 +59,7 @@ protected:
 
 	TSharedRef<SWidget> CreateWidget();
 
-	TSharedRef<ICustomDetailsViewItem> GetDefaultCategory(const TSharedRef<ICustomDetailsView>& InDetailsView,
+	virtual TSharedRef<ICustomDetailsViewItem> GetDefaultCategory(const TSharedRef<ICustomDetailsView>& InDetailsView,
 		const FCustomDetailsViewItemId& InRootId);
 
 	TSharedRef<ICustomDetailsViewItem> GetCategoryForRow(const TSharedRef<ICustomDetailsView>& InDetailsView, 
