Material Designer
- In the top-vertical layout, the material preview area can now be resized with a splitter.

diff --git a/Engine/Plugins/Experimental/DynamicMaterial/Source/DynamicMaterialEditor/Private/DynamicMaterialEditorSettings.cpp b/Engine/Plugins/Experimental/DynamicMaterial/Source/DynamicMaterialEditor/Private/DynamicMaterialEditorSettings.cpp
--- a/Engine/Plugins/Experimental/DynamicMaterial/Source/DynamicMaterialEditor/Private/DynamicMaterialEditorSettings.cpp
+++ b/Engine/Plugins/Experimental/DynamicMaterial/Source/DynamicMaterialEditor/Private/DynamicMaterialEditorSettings.cpp
@@ -301,7 +301,8 @@ void UDynamicMaterialEditorSettings::ResetAllLayoutSettings()
 {
 	Layout = EDMMaterialEditorLayout::Left;
 	bUseFullChannelNamesInTopSlimLayout = false;
-	SplitterLocation = 240;
+	SplitterLocation = 0.5;
+	PreviewSplitterLocation = 0.333;
 	ThumbnailSize = 256.f;
 	PreviewMesh = EDMMaterialPreviewMesh::Plane;
 	bShowPreviewBackground = true;
diff --git a/Engine/Plugins/Experimental/DynamicMaterial/Source/DynamicMaterialEditor/Private/DynamicMaterialEditorSettings.h b/Engine/Plugins/Experimental/DynamicMaterial/Source/DynamicMaterialEditor/Private/DynamicMaterialEditorSettings.h
--- a/Engine/Plugins/Experimental/DynamicMaterial/Source/DynamicMaterialEditor/Private/DynamicMaterialEditorSettings.h
+++ b/Engine/Plugins/Experimental/DynamicMaterial/Source/DynamicMaterialEditor/Private/DynamicMaterialEditorSettings.h
@@ -179,6 +179,11 @@ public:
 		ClampMin = "0.05", UIMin = "0.05", ClampMax = "0.95", UIMax = "0.95"))
 	float SplitterLocation;
 
+	/** Adjusts the vertical size of the material layer view. */
+	UPROPERTY(Config, EditAnywhere, BlueprintReadWrite, Category = "Layout", meta = (
+		ClampMin = "0.05", UIMin = "0.05", ClampMax = "0.95", UIMax = "0.95"))
+	float PreviewSplitterLocation;
+
 	UPROPERTY(Config, EditAnywhere, BlueprintReadWrite, Category = "Layout")
 	bool bUVVisualizerVisible;
 
diff --git a/Engine/Plugins/Experimental/DynamicMaterial/Source/DynamicMaterialEditor/Private/UI/Widgets/Editor/EditorLayouts/SDMMaterialEditor_TopVertical.cpp b/Engine/Plugins/Experimental/DynamicMaterial/Source/DynamicMaterialEditor/Private/UI/Widgets/Editor/EditorLayouts/SDMMaterialEditor_TopVertical.cpp
--- a/Engine/Plugins/Experimental/DynamicMaterial/Source/DynamicMaterialEditor/Private/UI/Widgets/Editor/EditorLayouts/SDMMaterialEditor_TopVertical.cpp
+++ b/Engine/Plugins/Experimental/DynamicMaterial/Source/DynamicMaterialEditor/Private/UI/Widgets/Editor/EditorLayouts/SDMMaterialEditor_TopVertical.cpp
@@ -2,14 +2,22 @@
 
 #include "UI/Widgets/Editor/EditorLayouts/SDMMaterialEditor_TopVertical.h"
 
+#include "DynamicMaterialEditorSettings.h"
 #include "Styling/AppStyle.h"
 #include "UI/Widgets/Editor/PropertySelectorLayouts/SDMMaterialPropertySelector_Wrap.h"
 #include "UI/Widgets/Editor/SDMMaterialPreview.h"
 #include "UI/Widgets/Editor/SDMMaterialPropertySelector.h"
 #include "UI/Widgets/Editor/SDMMaterialSlotEditor.h"
 #include "Widgets/Layout/SBorder.h"
+#include "Widgets/Layout/SBox.h"
+#include "Widgets/Layout/SSplitter.h"
 #include "Widgets/SBoxPanel.h"
 
+SDMMaterialEditor_TopVertical::SDMMaterialEditor_TopVertical()
+	: SplitterSlot_Top(nullptr)
+{
+}
+
 void SDMMaterialEditor_TopVertical::Construct(const FArguments& InArgs, const TSharedRef<SDMMaterialDesigner>& InDesignerWidget)
 {
 	SDMMaterialEditor_TopBase::Construct(
@@ -34,7 +42,7 @@ TSharedRef<SWidget> SDMMaterialEditor_TopVertical::CreateSlot_Top()
 			SNew(SVerticalBox)
 			+ SVerticalBox::Slot()
 			.Expose(MaterialPreviewSlotPtr)
-			.AutoHeight()
+			.FillHeight(1.0f)
 			.Padding(0.f)
 			[
 				SNullWidget::NullWidget
@@ -42,7 +50,7 @@ TSharedRef<SWidget> SDMMaterialEditor_TopVertical::CreateSlot_Top()
 
 			+ SVerticalBox::Slot()
 			.Expose(PropertySelectorSlotPtr)
-			.FillHeight(1.0f)
+			.AutoHeight()
 			.Padding(0.f, 5.f, 0.f, 0.f)
 			[
 				SNullWidget::NullWidget
@@ -55,7 +63,75 @@ TSharedRef<SWidget> SDMMaterialEditor_TopVertical::CreateSlot_Top()
 	return NewLeft;
 }
 
+TSharedRef<SWidget> SDMMaterialEditor_TopVertical::CreateSlot_Main()
+{
+	using namespace UE::DynamicMaterialEditor::Private;
+
+	float SplitterValue = 300;
+
+	if (UDynamicMaterialEditorSettings* Settings = UDynamicMaterialEditorSettings::Get())
+	{
+		SplitterValue = Settings->PreviewSplitterLocation;
+	}
+
+	SSplitter::FSlot* SplitterSlotPtr = nullptr;
+	SSplitter::FSlot* BottomSlotPtr = nullptr;
+	TSharedPtr<SBox> TopBox;
+
+	TSharedRef<SSplitter> NewMain = SNew(SSplitter)
+		.Style(FAppStyle::Get(), "DetailsView.Splitter")
+		.Orientation(Orient_Vertical)
+		.ResizeMode(ESplitterResizeMode::Fill)
+		.PhysicalSplitterHandleSize(5.0f)
+		.HitDetectionSplitterHandleSize(5.0f)
+		.OnSplitterFinishedResizing(this, &SDMMaterialEditor_TopVertical::OnTopSplitterResized)
+
+		+ SSplitter::Slot()
+		.Expose(SplitterSlotPtr)
+		.Resizable(true)
+		.SizeRule(SSplitter::ESizeRule::FractionOfParent)
+		.MinSize(100)
+		.Value(SplitterValue)
+		[
+			SAssignNew(TopBox, SBox)
+			.VAlign(EVerticalAlignment::VAlign_Fill)
+			.Clipping(EWidgetClipping::ClipToBoundsAlways)
+			[
+				SNullWidget::NullWidget
+			]
+		]
+
+		+ SSplitter::Slot()
+		.Expose(BottomSlotPtr)
+		.Resizable(true)
+		.SizeRule(SSplitter::ESizeRule::FractionOfParent)
+		.MinSize(250)
+		.Value(1.f - SplitterValue)
+		[
+			SNullWidget::NullWidget
+		];
+
+	SplitterSlot_Top = SplitterSlotPtr;
+	TopSlot = TDMWidgetSlot<SWidget>(TopBox.ToSharedRef(), 0, CreateSlot_Top());
+	BottomSlot = TDMWidgetSlot<SWidget>(BottomSlotPtr, CreateSlot_Bottom());
+
+	return NewMain;
+}
+
 TSharedRef<SDMMaterialPropertySelector> SDMMaterialEditor_TopVertical::CreateSlot_PropertySelector_Impl()
 {
 	return SNew(SDMMaterialPropertySelector_Wrap, SharedThis(this));
 }
+
+void SDMMaterialEditor_TopVertical::OnTopSplitterResized()
+{
+	if (SplitterSlot)
+	{
+		if (UDynamicMaterialEditorSettings* Settings = UDynamicMaterialEditorSettings::Get())
+		{
+			const float SplitterLocation = static_cast<SSplitter::FSlot*>(SplitterSlot)->GetSizeValue();
+			Settings->PreviewSplitterLocation = SplitterLocation;
+			Settings->SaveConfig();
+		}
+	}
+}
diff --git a/Engine/Plugins/Experimental/DynamicMaterial/Source/DynamicMaterialEditor/Private/UI/Widgets/Editor/EditorLayouts/SDMMaterialEditor_TopVertical.h b/Engine/Plugins/Experimental/DynamicMaterial/Source/DynamicMaterialEditor/Private/UI/Widgets/Editor/EditorLayouts/SDMMaterialEditor_TopVertical.h
--- a/Engine/Plugins/Experimental/DynamicMaterial/Source/DynamicMaterialEditor/Private/UI/Widgets/Editor/EditorLayouts/SDMMaterialEditor_TopVertical.h
+++ b/Engine/Plugins/Experimental/DynamicMaterial/Source/DynamicMaterialEditor/Private/UI/Widgets/Editor/EditorLayouts/SDMMaterialEditor_TopVertical.h
@@ -15,16 +15,23 @@ class SDMMaterialEditor_TopVertical : public SDMMaterialEditor_TopBase
 	SLATE_END_ARGS()
 
 public:
+	SDMMaterialEditor_TopVertical();
+
 	virtual ~SDMMaterialEditor_TopVertical() override = default;
 
 	void Construct(const FArguments& InArgs, const TSharedRef<SDMMaterialDesigner>& InDesignerWidget);
 
 protected:
+	FSlotBase* SplitterSlot_Top;
+
+	void OnTopSplitterResized();
+
 	//~ Begin SDMMaterialEditor_TopBase
 	virtual TSharedRef<SWidget> CreateSlot_Top() override;
 	//~ End SDMMaterialEditor_TopBase
 
 	//~ Begin SDMMaterialEditor
+	virtual TSharedRef<SWidget> CreateSlot_Main() override;
 	virtual TSharedRef<SDMMaterialPropertySelector> CreateSlot_PropertySelector_Impl() override;
 	//~ End SDMMaterialEditor
 };
diff --git a/Engine/Plugins/Experimental/DynamicMaterial/Source/DynamicMaterialEditor/Private/UI/Widgets/SDMMaterialEditor.cpp b/Engine/Plugins/Experimental/DynamicMaterial/Source/DynamicMaterialEditor/Private/UI/Widgets/SDMMaterialEditor.cpp
--- a/Engine/Plugins/Experimental/DynamicMaterial/Source/DynamicMaterialEditor/Private/UI/Widgets/SDMMaterialEditor.cpp
+++ b/Engine/Plugins/Experimental/DynamicMaterial/Source/DynamicMaterialEditor/Private/UI/Widgets/SDMMaterialEditor.cpp
@@ -45,7 +45,8 @@ void SDMMaterialEditor::PrivateRegisterAttributes(FSlateAttributeDescriptor::FIn
 }
 
 SDMMaterialEditor::SDMMaterialEditor()
-	: CommandList(MakeShared<FUICommandList>())
+	: SplitterSlot(nullptr)
+	, CommandList(MakeShared<FUICommandList>())
 	, PreviewMaterialManager(MakeShared<FDMPreviewMaterialManager>())
 	, EditMode(EDMMaterialEditorMode::GlobalSettings)
 {
@@ -911,13 +912,14 @@ void SDMMaterialEditor::OnEnginePreExit()
 
 void SDMMaterialEditor::OnEditorSplitterResized()
 {
-	UDynamicMaterialEditorSettings* Settings = UDynamicMaterialEditorSettings::Get();
-
 	if (SplitterSlot)
 	{
-		const float SplitterLocation = static_cast<SSplitter::FSlot*>(SplitterSlot)->GetSizeValue();
-		Settings->SplitterLocation = SplitterLocation;
-		Settings->SaveConfig();
+		if (UDynamicMaterialEditorSettings* Settings = UDynamicMaterialEditorSettings::Get())
+		{
+			const float SplitterLocation = static_cast<SSplitter::FSlot*>(SplitterSlot)->GetSizeValue();
+			Settings->SplitterLocation = SplitterLocation;
+			Settings->SaveConfig();
+		}
 	}
 }
 
