Material Designer
- Added page history and ability to go backwards and forwards with thumb buttons.

diff --git a/Engine/Plugins/Experimental/DynamicMaterial/Source/DynamicMaterialEditor/Private/DynamicMaterialEditorCommands.cpp b/Engine/Plugins/Experimental/DynamicMaterial/Source/DynamicMaterialEditor/Private/DynamicMaterialEditorCommands.cpp
--- a/Engine/Plugins/Experimental/DynamicMaterial/Source/DynamicMaterialEditor/Private/DynamicMaterialEditorCommands.cpp
+++ b/Engine/Plugins/Experimental/DynamicMaterial/Source/DynamicMaterialEditor/Private/DynamicMaterialEditorCommands.cpp
@@ -63,6 +63,18 @@ void FDynamicMaterialEditorCommands::RegisterCommands()
 		, "Sets the preview mesh to the Shader Ball."
 		, EUserInterfaceActionType::ToggleButton
 		, FInputChord());
+
+	UI_COMMAND(NavigateForward
+		, "Navigate Forward"
+		, "Navigates forward in the page history."
+		, EUserInterfaceActionType::Button
+		, FInputChord(EKeys::ThumbMouseButton2));
+
+	UI_COMMAND(NavigateBack
+		, "Navigate Back"
+		, "Navigates back in the page history."
+		, EUserInterfaceActionType::Button
+		, FInputChord(EKeys::ThumbMouseButton));
 }
 
 #undef LOCTEXT_NAMESPACE
diff --git a/Engine/Plugins/Experimental/DynamicMaterial/Source/DynamicMaterialEditor/Private/DynamicMaterialEditorCommands.h b/Engine/Plugins/Experimental/DynamicMaterial/Source/DynamicMaterialEditor/Private/DynamicMaterialEditorCommands.h
--- a/Engine/Plugins/Experimental/DynamicMaterial/Source/DynamicMaterialEditor/Private/DynamicMaterialEditorCommands.h
+++ b/Engine/Plugins/Experimental/DynamicMaterial/Source/DynamicMaterialEditor/Private/DynamicMaterialEditorCommands.h
@@ -22,4 +22,7 @@ public:
 	TSharedPtr<FUICommandInfo> MoveLayerDown;
 
 	TSharedPtr<FUICommandInfo> SetShaderBallPreview;
+
+	TSharedPtr<FUICommandInfo> NavigateForward;
+	TSharedPtr<FUICommandInfo> NavigateBack;
 };
diff --git a/Engine/Plugins/Experimental/DynamicMaterial/Source/DynamicMaterialEditor/Private/UI/Widgets/Editor/EditorLayouts/SDMMaterialEditor_LeftBase.cpp b/Engine/Plugins/Experimental/DynamicMaterial/Source/DynamicMaterialEditor/Private/UI/Widgets/Editor/EditorLayouts/SDMMaterialEditor_LeftBase.cpp
--- a/Engine/Plugins/Experimental/DynamicMaterial/Source/DynamicMaterialEditor/Private/UI/Widgets/Editor/EditorLayouts/SDMMaterialEditor_LeftBase.cpp
+++ b/Engine/Plugins/Experimental/DynamicMaterial/Source/DynamicMaterialEditor/Private/UI/Widgets/Editor/EditorLayouts/SDMMaterialEditor_LeftBase.cpp
@@ -67,19 +67,19 @@ void SDMMaterialEditor_LeftBase::EditGlobalSettings(bool bInForceRefresh)
 	SDMMaterialEditor::EditGlobalSettings(bInForceRefresh);
 }
 
-void SDMMaterialEditor_LeftBase::ShowPropertyPreviews(bool bInForceRefresh)
+void SDMMaterialEditor_LeftBase::EditProperties(bool bInForceRefresh)
 {
-	if (EditMode == EDMMaterialEditorMode::PropertyPreviews && !bInForceRefresh)
+	if (EditMode == EDMMaterialEditorMode::Properties && !bInForceRefresh)
 	{
 		return;
 	}
 
-	if (EditMode != EDMMaterialEditorMode::PropertyPreviews)
+	if (EditMode != EDMMaterialEditorMode::Properties)
 	{
 		RightSlot.Invalidate();
 	}
 
-	SDMMaterialEditor::ShowPropertyPreviews(bInForceRefresh);
+	SDMMaterialEditor::EditProperties(bInForceRefresh);
 }
 
 void SDMMaterialEditor_LeftBase::ValidateSlots_Main()
@@ -186,7 +186,7 @@ TSharedRef<SWidget> SDMMaterialEditor_LeftBase::CreateSlot_Right()
 			Content = CreateSlot_Right_GlobalSettings();
 			break;
 
-		case EDMMaterialEditorMode::PropertyPreviews:
+		case EDMMaterialEditorMode::Properties:
 			Content = CreateSlot_Right_PropertyPreviews();
 			break;
 
diff --git a/Engine/Plugins/Experimental/DynamicMaterial/Source/DynamicMaterialEditor/Private/UI/Widgets/Editor/EditorLayouts/SDMMaterialEditor_LeftBase.h b/Engine/Plugins/Experimental/DynamicMaterial/Source/DynamicMaterialEditor/Private/UI/Widgets/Editor/EditorLayouts/SDMMaterialEditor_LeftBase.h
--- a/Engine/Plugins/Experimental/DynamicMaterial/Source/DynamicMaterialEditor/Private/UI/Widgets/Editor/EditorLayouts/SDMMaterialEditor_LeftBase.h
+++ b/Engine/Plugins/Experimental/DynamicMaterial/Source/DynamicMaterialEditor/Private/UI/Widgets/Editor/EditorLayouts/SDMMaterialEditor_LeftBase.h
@@ -23,7 +23,7 @@ public:
 	virtual void EditSlot(UDMMaterialSlot* InSlot, bool bInForceRefresh = false) override;
 	virtual void EditComponent(UDMMaterialComponent* InComponent, bool bInForceRefresh = false) override;
 	virtual void EditGlobalSettings(bool bInForceRefresh = false) override;
-	virtual void ShowPropertyPreviews(bool bInForceRefresh = false) override;
+	virtual void EditProperties(bool bInForceRefresh = false) override;
 	//~ End SDMMaterialEditor
 
 protected:
diff --git a/Engine/Plugins/Experimental/DynamicMaterial/Source/DynamicMaterialEditor/Private/UI/Widgets/Editor/EditorLayouts/SDMMaterialEditor_TopBase.cpp b/Engine/Plugins/Experimental/DynamicMaterial/Source/DynamicMaterialEditor/Private/UI/Widgets/Editor/EditorLayouts/SDMMaterialEditor_TopBase.cpp
--- a/Engine/Plugins/Experimental/DynamicMaterial/Source/DynamicMaterialEditor/Private/UI/Widgets/Editor/EditorLayouts/SDMMaterialEditor_TopBase.cpp
+++ b/Engine/Plugins/Experimental/DynamicMaterial/Source/DynamicMaterialEditor/Private/UI/Widgets/Editor/EditorLayouts/SDMMaterialEditor_TopBase.cpp
@@ -65,19 +65,19 @@ void SDMMaterialEditor_TopBase::EditGlobalSettings(bool bInForceRefresh)
 	SDMMaterialEditor::EditGlobalSettings(bInForceRefresh);
 }
 
-void SDMMaterialEditor_TopBase::ShowPropertyPreviews(bool bInForceRefresh)
+void SDMMaterialEditor_TopBase::EditProperties(bool bInForceRefresh)
 {
-	if (EditMode == EDMMaterialEditorMode::PropertyPreviews && !bInForceRefresh)
+	if (EditMode == EDMMaterialEditorMode::Properties && !bInForceRefresh)
 	{
 		return;
 	}
 
-	if (EditMode != EDMMaterialEditorMode::PropertyPreviews)
+	if (EditMode != EDMMaterialEditorMode::Properties)
 	{
 		BottomSlot.Invalidate();
 	}
 
-	SDMMaterialEditor::ShowPropertyPreviews(bInForceRefresh);
+	SDMMaterialEditor::EditProperties(bInForceRefresh);
 }
 
 void SDMMaterialEditor_TopBase::ValidateSlots_Main()
@@ -184,7 +184,7 @@ TSharedRef<SWidget> SDMMaterialEditor_TopBase::CreateSlot_Bottom()
 			Content = CreateSlot_Bottom_GlobalSettings();
 			break;
 
-		case EDMMaterialEditorMode::PropertyPreviews:
+		case EDMMaterialEditorMode::Properties:
 			Content = CreateSlot_Bottom_PropertyPreviews();
 			break;
 
diff --git a/Engine/Plugins/Experimental/DynamicMaterial/Source/DynamicMaterialEditor/Private/UI/Widgets/Editor/EditorLayouts/SDMMaterialEditor_TopBase.h b/Engine/Plugins/Experimental/DynamicMaterial/Source/DynamicMaterialEditor/Private/UI/Widgets/Editor/EditorLayouts/SDMMaterialEditor_TopBase.h
--- a/Engine/Plugins/Experimental/DynamicMaterial/Source/DynamicMaterialEditor/Private/UI/Widgets/Editor/EditorLayouts/SDMMaterialEditor_TopBase.h
+++ b/Engine/Plugins/Experimental/DynamicMaterial/Source/DynamicMaterialEditor/Private/UI/Widgets/Editor/EditorLayouts/SDMMaterialEditor_TopBase.h
@@ -23,7 +23,7 @@ public:
 	virtual void EditSlot(UDMMaterialSlot* InSlot, bool bInForceRefresh = false) override;
 	virtual void EditComponent(UDMMaterialComponent* InComponent, bool bInForceRefresh = false) override;
 	virtual void EditGlobalSettings(bool bInForceRefresh = false) override;
-	virtual void ShowPropertyPreviews(bool bInForceRefresh = false) override;
+	virtual void EditProperties(bool bInForceRefresh = false) override;
 	//~ End SDMMaterialEditor
 
 protected:
diff --git a/Engine/Plugins/Experimental/DynamicMaterial/Source/DynamicMaterialEditor/Private/UI/Widgets/Editor/PropertySelectorLayouts/SDMMaterialPropertySelector_Vertical.cpp b/Engine/Plugins/Experimental/DynamicMaterial/Source/DynamicMaterialEditor/Private/UI/Widgets/Editor/PropertySelectorLayouts/SDMMaterialPropertySelector_Vertical.cpp
--- a/Engine/Plugins/Experimental/DynamicMaterial/Source/DynamicMaterialEditor/Private/UI/Widgets/Editor/PropertySelectorLayouts/SDMMaterialPropertySelector_Vertical.cpp
+++ b/Engine/Plugins/Experimental/DynamicMaterial/Source/DynamicMaterialEditor/Private/UI/Widgets/Editor/PropertySelectorLayouts/SDMMaterialPropertySelector_Vertical.cpp
@@ -4,6 +4,7 @@
 
 #include "DetailLayoutBuilder.h"
 #include "DMDefs.h"
+#include "UI/Widgets/SDMMaterialEditor.h"
 #include "Widgets/Images/SImage.h"
 #include "Widgets/Input/SCheckBox.h"
 #include "Widgets/Layout/SBox.h"
@@ -21,19 +22,18 @@ void SDMMaterialPropertySelector_Vertical::Construct(const FArguments& InArgs, c
 	);
 }
 
-TSharedRef<SWidget> SDMMaterialPropertySelector_Vertical::CreateSlot_SelectButton(EDMMaterialEditorMode InEditMode, 
-	EDMMaterialPropertyType InMaterialProperty)
+TSharedRef<SWidget> SDMMaterialPropertySelector_Vertical::CreateSlot_SelectButton(const FDMMaterialEditorPage& InPage)
 {
-	const FText ButtonText = GetSelectButtonText(InEditMode, InMaterialProperty, /* Short Name */ false);
-	const FText ToolTip = GetButtonToolTip(InEditMode, InMaterialProperty);
+	const FText ButtonText = GetSelectButtonText(InPage, /* Short Name */ false);
+	const FText ToolTip = GetButtonToolTip(InPage);
 
 	return SNew(SCheckBox)
 		.Style(FAppStyle::Get(), "DetailsView.SectionButton")
 		.HAlign(EHorizontalAlignment::HAlign_Center)
 		.Padding(0.f)
-		.IsEnabled(this, &SDMMaterialPropertySelector_Vertical::GetPropertySelectEnabled, InEditMode, InMaterialProperty)
-		.IsChecked(this, &SDMMaterialPropertySelector_Vertical::GetPropertySelectState, InEditMode, InMaterialProperty)
-		.OnCheckStateChanged(this, &SDMMaterialPropertySelector_Vertical::OnPropertySelectStateChanged, InEditMode, InMaterialProperty)
+		.IsEnabled(this, &SDMMaterialPropertySelector_Vertical::GetPropertySelectEnabled, InPage)
+		.IsChecked(this, &SDMMaterialPropertySelector_Vertical::GetPropertySelectState, InPage)
+		.OnCheckStateChanged(this, &SDMMaterialPropertySelector_Vertical::OnPropertySelectStateChanged, InPage)
 		.ToolTipText(ToolTip)
 		.Content()
 		[
@@ -48,7 +48,7 @@ TSharedRef<SWidget> SDMMaterialPropertySelector_Vertical::CreateSlot_SelectButto
 				[
 					SNew(SImage)
 					.Image(FAppStyle::Get().GetBrush("FilterBar.FilterImage"))
-					.ColorAndOpacity(this, &SDMMaterialPropertySelector_Vertical::GetPropertySelectButtonChipColor, InEditMode, InMaterialProperty)
+					.ColorAndOpacity(this, &SDMMaterialPropertySelector_Vertical::GetPropertySelectButtonChipColor, InPage)
 						.DesiredSizeOverride(FVector2D(8, 17))
 				]
 				+SHorizontalBox::Slot()
diff --git a/Engine/Plugins/Experimental/DynamicMaterial/Source/DynamicMaterialEditor/Private/UI/Widgets/Editor/PropertySelectorLayouts/SDMMaterialPropertySelector_Vertical.h b/Engine/Plugins/Experimental/DynamicMaterial/Source/DynamicMaterialEditor/Private/UI/Widgets/Editor/PropertySelectorLayouts/SDMMaterialPropertySelector_Vertical.h
--- a/Engine/Plugins/Experimental/DynamicMaterial/Source/DynamicMaterialEditor/Private/UI/Widgets/Editor/PropertySelectorLayouts/SDMMaterialPropertySelector_Vertical.h
+++ b/Engine/Plugins/Experimental/DynamicMaterial/Source/DynamicMaterialEditor/Private/UI/Widgets/Editor/PropertySelectorLayouts/SDMMaterialPropertySelector_Vertical.h
@@ -4,6 +4,8 @@
 
 #include "UI/Widgets/Editor/PropertySelectorLayouts/SDMMaterialPropertySelector_VerticalBase.h"
 
+struct FDMMaterialEditorPage;
+
 class SDMMaterialPropertySelector_Vertical : public SDMMaterialPropertySelector_VerticalBase
 {
 	SLATE_BEGIN_ARGS(SDMMaterialPropertySelector_Vertical) {}
@@ -16,6 +18,6 @@ public:
 
 protected:
 	//~ Begin SDMMaterialPropertySelector
-	virtual TSharedRef<SWidget> CreateSlot_SelectButton(EDMMaterialEditorMode InEditMode, EDMMaterialPropertyType InMaterialProperty) override;
+	virtual TSharedRef<SWidget> CreateSlot_SelectButton(const FDMMaterialEditorPage& InPage) override;
 	//~ End SDMMaterialPropertySelector
 };
diff --git a/Engine/Plugins/Experimental/DynamicMaterial/Source/DynamicMaterialEditor/Private/UI/Widgets/Editor/PropertySelectorLayouts/SDMMaterialPropertySelector_VerticalBase.cpp b/Engine/Plugins/Experimental/DynamicMaterial/Source/DynamicMaterialEditor/Private/UI/Widgets/Editor/PropertySelectorLayouts/SDMMaterialPropertySelector_VerticalBase.cpp
--- a/Engine/Plugins/Experimental/DynamicMaterial/Source/DynamicMaterialEditor/Private/UI/Widgets/Editor/PropertySelectorLayouts/SDMMaterialPropertySelector_VerticalBase.cpp
+++ b/Engine/Plugins/Experimental/DynamicMaterial/Source/DynamicMaterialEditor/Private/UI/Widgets/Editor/PropertySelectorLayouts/SDMMaterialPropertySelector_VerticalBase.cpp
@@ -42,7 +42,7 @@ TSharedRef<SWidget> SDMMaterialPropertySelector_VerticalBase::CreateSlot_Propert
 	NewSlotList->AddSlot(PropertySelectorColumns::Select, Row)
 		.Padding(Padding)
 		[
-			CreateSlot_SelectButton(EDMMaterialEditorMode::GlobalSettings, EDMMaterialPropertyType::None)
+			CreateSlot_SelectButton(FDMMaterialEditorPage::GlobalSettings)
 		];
 
 	++Row;
@@ -50,7 +50,7 @@ TSharedRef<SWidget> SDMMaterialPropertySelector_VerticalBase::CreateSlot_Propert
 	NewSlotList->AddSlot(PropertySelectorColumns::Select, Row)
 		.Padding(Padding)
 		[
-			CreateSlot_SelectButton(EDMMaterialEditorMode::PropertyPreviews, EDMMaterialPropertyType::None)
+			CreateSlot_SelectButton(FDMMaterialEditorPage::Properties)
 		];
 
 	++Row;
@@ -83,7 +83,7 @@ TSharedRef<SWidget> SDMMaterialPropertySelector_VerticalBase::CreateSlot_Propert
 		NewSlotList->AddSlot(PropertySelectorColumns::Select, Row)
 			.Padding(Padding)
 			[
-				CreateSlot_SelectButton(EDMMaterialEditorMode::EditSlot, PropertyPair.Key)
+				CreateSlot_SelectButton({EDMMaterialEditorMode::EditSlot, PropertyPair.Key})
 			];
 
 		++Row;
@@ -117,7 +117,7 @@ TSharedRef<SWidget> SDMMaterialPropertySelector_VerticalBase::CreateSlot_Propert
 		NewSlotList->AddSlot(PropertySelectorColumns::Select, Row)
 			.Padding(Padding)
 			[
-				CreateSlot_SelectButton(EDMMaterialEditorMode::EditSlot, PropertyPair.Key)
+				CreateSlot_SelectButton({EDMMaterialEditorMode::EditSlot, PropertyPair.Key})
 			];
 
 		++Row;
@@ -145,7 +145,7 @@ TSharedRef<SWidget> SDMMaterialPropertySelector_VerticalBase::CreateSlot_Propert
 		NewSlotList->AddSlot(PropertySelectorColumns::Select, Row)
 			.Padding(Padding)
 			[
-				CreateSlot_SelectButton(EDMMaterialEditorMode::EditSlot, PropertyPair.Key)
+				CreateSlot_SelectButton({EDMMaterialEditorMode::EditSlot, PropertyPair.Key})
 			];
 
 		++Row;
diff --git a/Engine/Plugins/Experimental/DynamicMaterial/Source/DynamicMaterialEditor/Private/UI/Widgets/Editor/PropertySelectorLayouts/SDMMaterialPropertySelector_VerticalSlim.cpp b/Engine/Plugins/Experimental/DynamicMaterial/Source/DynamicMaterialEditor/Private/UI/Widgets/Editor/PropertySelectorLayouts/SDMMaterialPropertySelector_VerticalSlim.cpp
--- a/Engine/Plugins/Experimental/DynamicMaterial/Source/DynamicMaterialEditor/Private/UI/Widgets/Editor/PropertySelectorLayouts/SDMMaterialPropertySelector_VerticalSlim.cpp
+++ b/Engine/Plugins/Experimental/DynamicMaterial/Source/DynamicMaterialEditor/Private/UI/Widgets/Editor/PropertySelectorLayouts/SDMMaterialPropertySelector_VerticalSlim.cpp
@@ -39,7 +39,7 @@ TSharedRef<SWidget> SDMMaterialPropertySelector_VerticalSlim::CreateSlot_Propert
 
 	const FMargin Padding = FMargin(0.f, 1.f);
 
-	TSharedRef<SWidget> PreviewButton = CreateSlot_SelectButton(EDMMaterialEditorMode::MaterialPreview, EDMMaterialPropertyType::None);
+	TSharedRef<SWidget> PreviewButton = CreateSlot_SelectButton(FDMMaterialEditorPage::Preview);
 	SetupMaterialPreviewButton(PreviewButton);
 
 	NewSlotList->AddSlot()
@@ -53,14 +53,14 @@ TSharedRef<SWidget> SDMMaterialPropertySelector_VerticalSlim::CreateSlot_Propert
 		.AutoHeight()
 		.Padding(Padding)
 		[
-			CreateSlot_SelectButton(EDMMaterialEditorMode::GlobalSettings, EDMMaterialPropertyType::None)
+			CreateSlot_SelectButton(FDMMaterialEditorPage::GlobalSettings)
 		];
 
 	NewSlotList->AddSlot()
 		.AutoHeight()
 		.Padding(Padding)
 		[
-			CreateSlot_SelectButton(EDMMaterialEditorMode::PropertyPreviews, EDMMaterialPropertyType::None)
+			CreateSlot_SelectButton(FDMMaterialEditorPage::Properties)
 		];
 
 	// Valid model properties
@@ -87,7 +87,7 @@ TSharedRef<SWidget> SDMMaterialPropertySelector_VerticalSlim::CreateSlot_Propert
 			.AutoHeight()
 			.Padding(Padding)
 			[
-				CreateSlot_SelectButton(EDMMaterialEditorMode::EditSlot, PropertyPair.Key)
+				CreateSlot_SelectButton({EDMMaterialEditorMode::EditSlot, PropertyPair.Key})
 			];
 	}
 
@@ -115,7 +115,7 @@ TSharedRef<SWidget> SDMMaterialPropertySelector_VerticalSlim::CreateSlot_Propert
 			.AutoHeight()
 			.Padding(Padding)
 			[
-				CreateSlot_SelectButton(EDMMaterialEditorMode::EditSlot, PropertyPair.Key)
+				CreateSlot_SelectButton({EDMMaterialEditorMode::EditSlot, PropertyPair.Key})
 			];
 	}
 
@@ -137,26 +137,25 @@ TSharedRef<SWidget> SDMMaterialPropertySelector_VerticalSlim::CreateSlot_Propert
 			.AutoHeight()
 			.Padding(Padding)
 			[
-				CreateSlot_SelectButton(EDMMaterialEditorMode::EditSlot, PropertyPair.Key)
+				CreateSlot_SelectButton({EDMMaterialEditorMode::EditSlot, PropertyPair.Key})
 			];
 	}
 
 	return NewSlotList;
 }
 
-TSharedRef<SWidget> SDMMaterialPropertySelector_VerticalSlim::CreateSlot_SelectButton(EDMMaterialEditorMode InEditMode, 
-	EDMMaterialPropertyType InMaterialProperty)
+TSharedRef<SWidget> SDMMaterialPropertySelector_VerticalSlim::CreateSlot_SelectButton(const FDMMaterialEditorPage& InPage)
 {
-	const FText ButtonText = GetSelectButtonText(InEditMode, InMaterialProperty, /* Short Name */ true);
-	const FText ToolTip = GetButtonToolTip(InEditMode, InMaterialProperty);
+	const FText ButtonText = GetSelectButtonText(InPage, /* Short Name */ true);
+	const FText ToolTip = GetButtonToolTip(InPage);
 
 	return SNew(SCheckBox)
 		.Style(FAppStyle::Get(), "DetailsView.SectionButton")
 		.HAlign(EHorizontalAlignment::HAlign_Center)
 		.Padding(0.f)
-		.IsEnabled(this, &SDMMaterialPropertySelector_VerticalSlim::GetPropertySelectEnabled, InEditMode, InMaterialProperty)
-		.IsChecked(this, &SDMMaterialPropertySelector_VerticalSlim::GetPropertySelectState, InEditMode, InMaterialProperty)
-		.OnCheckStateChanged(this, &SDMMaterialPropertySelector_VerticalSlim::OnPropertySelectStateChanged, InEditMode, InMaterialProperty)
+		.IsEnabled(this, &SDMMaterialPropertySelector_VerticalSlim::GetPropertySelectEnabled, InPage)
+		.IsChecked(this, &SDMMaterialPropertySelector_VerticalSlim::GetPropertySelectState, InPage)
+		.OnCheckStateChanged(this, &SDMMaterialPropertySelector_VerticalSlim::OnPropertySelectStateChanged, InPage)
 		.ToolTipText(ToolTip)
 		.Content()
 		[
@@ -171,7 +170,7 @@ TSharedRef<SWidget> SDMMaterialPropertySelector_VerticalSlim::CreateSlot_SelectB
 				[
 					SNew(SImage)
 					.Image(FAppStyle::Get().GetBrush("FilterBar.FilterImage"))
-					.ColorAndOpacity(this, &SDMMaterialPropertySelector_VerticalSlim::GetPropertySelectButtonChipColor, InEditMode, InMaterialProperty)
+					.ColorAndOpacity(this, &SDMMaterialPropertySelector_VerticalSlim::GetPropertySelectButtonChipColor, InPage)
 					.DesiredSizeOverride(FVector2D(8, 17))
 				]
 				+SHorizontalBox::Slot()
diff --git a/Engine/Plugins/Experimental/DynamicMaterial/Source/DynamicMaterialEditor/Private/UI/Widgets/Editor/PropertySelectorLayouts/SDMMaterialPropertySelector_VerticalSlim.h b/Engine/Plugins/Experimental/DynamicMaterial/Source/DynamicMaterialEditor/Private/UI/Widgets/Editor/PropertySelectorLayouts/SDMMaterialPropertySelector_VerticalSlim.h
--- a/Engine/Plugins/Experimental/DynamicMaterial/Source/DynamicMaterialEditor/Private/UI/Widgets/Editor/PropertySelectorLayouts/SDMMaterialPropertySelector_VerticalSlim.h
+++ b/Engine/Plugins/Experimental/DynamicMaterial/Source/DynamicMaterialEditor/Private/UI/Widgets/Editor/PropertySelectorLayouts/SDMMaterialPropertySelector_VerticalSlim.h
@@ -4,6 +4,8 @@
 
 #include "UI/Widgets/Editor/PropertySelectorLayouts/SDMMaterialPropertySelector_VerticalBase.h"
 
+struct FDMMaterialEditorPage;
+
 class SDMMaterialPropertySelector_VerticalSlim : public SDMMaterialPropertySelector_VerticalBase
 {
 	SLATE_BEGIN_ARGS(SDMMaterialPropertySelector_VerticalSlim) {}
@@ -17,6 +19,6 @@ public:
 protected:
 	//~ Begin SDMMaterialPropertySelector
 	virtual TSharedRef<SWidget> CreateSlot_PropertyList() override;
-	virtual TSharedRef<SWidget> CreateSlot_SelectButton(EDMMaterialEditorMode InEditMode, EDMMaterialPropertyType InMaterialProperty) override;
+	virtual TSharedRef<SWidget> CreateSlot_SelectButton(const FDMMaterialEditorPage& InPage) override;
 	//~ End SDMMaterialPropertySelector
 };
diff --git a/Engine/Plugins/Experimental/DynamicMaterial/Source/DynamicMaterialEditor/Private/UI/Widgets/Editor/PropertySelectorLayouts/SDMMaterialPropertySelector_WrapBase.cpp b/Engine/Plugins/Experimental/DynamicMaterial/Source/DynamicMaterialEditor/Private/UI/Widgets/Editor/PropertySelectorLayouts/SDMMaterialPropertySelector_WrapBase.cpp
--- a/Engine/Plugins/Experimental/DynamicMaterial/Source/DynamicMaterialEditor/Private/UI/Widgets/Editor/PropertySelectorLayouts/SDMMaterialPropertySelector_WrapBase.cpp
+++ b/Engine/Plugins/Experimental/DynamicMaterial/Source/DynamicMaterialEditor/Private/UI/Widgets/Editor/PropertySelectorLayouts/SDMMaterialPropertySelector_WrapBase.cpp
@@ -45,13 +45,13 @@ TSharedRef<SWidget> SDMMaterialPropertySelector_WrapBase::CreateSlot_PropertyLis
 	NewSlotList->AddSlot()
 		.Padding(Padding)
 		[
-			CreateSlot_SelectButton(EDMMaterialEditorMode::GlobalSettings, EDMMaterialPropertyType::None)
+			CreateSlot_SelectButton(FDMMaterialEditorPage::GlobalSettings)
 		];
 
 	NewSlotList->AddSlot()
 		.Padding(Padding)
 		[
-			CreateSlot_SelectButton(EDMMaterialEditorMode::PropertyPreviews, EDMMaterialPropertyType::None)
+			CreateSlot_SelectButton(FDMMaterialEditorPage::Properties)
 		];
 
 	for (const TPair<EDMMaterialPropertyType, UDMMaterialProperty*>& PropertyPair : EditorOnlyData->GetMaterialProperties())
@@ -66,26 +66,25 @@ TSharedRef<SWidget> SDMMaterialPropertySelector_WrapBase::CreateSlot_PropertyLis
 		NewSlotList->AddSlot()
 			.Padding(Padding)
 			[
-				CreateSlot_SelectButton(EDMMaterialEditorMode::EditSlot, PropertyPair.Key)
+				CreateSlot_SelectButton({EDMMaterialEditorMode::EditSlot, PropertyPair.Key})
 			];
 	}
 
 	return NewSlotList;
 }
 
-TSharedRef<SWidget> SDMMaterialPropertySelector_WrapBase::CreateSlot_SelectButton(EDMMaterialEditorMode InEditMode, 
-	EDMMaterialPropertyType InMaterialProperty)
+TSharedRef<SWidget> SDMMaterialPropertySelector_WrapBase::CreateSlot_SelectButton(const FDMMaterialEditorPage& InPage)
 {
-	const FText ButtonText = GetSelectButtonText(InEditMode, InMaterialProperty, /* Short Name */ false);
-	const FText ToolTip = GetButtonToolTip(InEditMode, InMaterialProperty);
+	const FText ButtonText = GetSelectButtonText(InPage, /* Short Name */ false);
+	const FText ToolTip = GetButtonToolTip(InPage);
 
 	return SNew(SCheckBox)
 		.Style(FAppStyle::Get(), "DetailsView.SectionButton")
 		.HAlign(EHorizontalAlignment::HAlign_Center)
 		.Padding(0.f)
-		.IsEnabled(this, &SDMMaterialPropertySelector_WrapBase::GetPropertySelectEnabled, InEditMode, InMaterialProperty)
-		.IsChecked(this, &SDMMaterialPropertySelector_WrapBase::GetPropertySelectState, InEditMode, InMaterialProperty)
-		.OnCheckStateChanged(this, &SDMMaterialPropertySelector_WrapBase::OnPropertySelectStateChanged, InEditMode, InMaterialProperty)
+		.IsEnabled(this, &SDMMaterialPropertySelector_WrapBase::GetPropertySelectEnabled, InPage)
+		.IsChecked(this, &SDMMaterialPropertySelector_WrapBase::GetPropertySelectState, InPage)
+		.OnCheckStateChanged(this, &SDMMaterialPropertySelector_WrapBase::OnPropertySelectStateChanged, InPage)
 		.ToolTipText(ToolTip)
 		.Content()
 		[
@@ -96,7 +95,7 @@ TSharedRef<SWidget> SDMMaterialPropertySelector_WrapBase::CreateSlot_SelectButto
 			[
 				SNew(SImage)
 				.Image(FAppStyle::Get().GetBrush("FilterBar.FilterImage"))
-				.ColorAndOpacity(this, &SDMMaterialPropertySelector_WrapBase::GetPropertySelectButtonChipColor, InEditMode, InMaterialProperty)
+				.ColorAndOpacity(this, &SDMMaterialPropertySelector_WrapBase::GetPropertySelectButtonChipColor, InPage)
 				.DesiredSizeOverride(FVector2D(8, 17))
 			]
 			+SHorizontalBox::Slot()
diff --git a/Engine/Plugins/Experimental/DynamicMaterial/Source/DynamicMaterialEditor/Private/UI/Widgets/Editor/PropertySelectorLayouts/SDMMaterialPropertySelector_WrapBase.h b/Engine/Plugins/Experimental/DynamicMaterial/Source/DynamicMaterialEditor/Private/UI/Widgets/Editor/PropertySelectorLayouts/SDMMaterialPropertySelector_WrapBase.h
--- a/Engine/Plugins/Experimental/DynamicMaterial/Source/DynamicMaterialEditor/Private/UI/Widgets/Editor/PropertySelectorLayouts/SDMMaterialPropertySelector_WrapBase.h
+++ b/Engine/Plugins/Experimental/DynamicMaterial/Source/DynamicMaterialEditor/Private/UI/Widgets/Editor/PropertySelectorLayouts/SDMMaterialPropertySelector_WrapBase.h
@@ -4,6 +4,8 @@
 
 #include "UI/Widgets/Editor/SDMMaterialPropertySelector.h"
 
+struct FDMMaterialEditorPage;
+
 class SDMMaterialPropertySelector_WrapBase : public SDMMaterialPropertySelector
 {
 	SLATE_BEGIN_ARGS(SDMMaterialPropertySelector_WrapBase) {}
@@ -17,6 +19,6 @@ public:
 protected:
 	//~ Begin SDMMaterialPropertySelector
 	virtual TSharedRef<SWidget> CreateSlot_PropertyList() override;
-	virtual TSharedRef<SWidget> CreateSlot_SelectButton(EDMMaterialEditorMode InEditMode, EDMMaterialPropertyType InMaterialProperty) override;
+	virtual TSharedRef<SWidget> CreateSlot_SelectButton(const FDMMaterialEditorPage& InPage) override;
 	//~ End SDMMaterialPropertySelector
 };
diff --git a/Engine/Plugins/Experimental/DynamicMaterial/Source/DynamicMaterialEditor/Private/UI/Widgets/Editor/PropertySelectorLayouts/SDMMaterialPropertySelector_WrapSlim.cpp b/Engine/Plugins/Experimental/DynamicMaterial/Source/DynamicMaterialEditor/Private/UI/Widgets/Editor/PropertySelectorLayouts/SDMMaterialPropertySelector_WrapSlim.cpp
--- a/Engine/Plugins/Experimental/DynamicMaterial/Source/DynamicMaterialEditor/Private/UI/Widgets/Editor/PropertySelectorLayouts/SDMMaterialPropertySelector_WrapSlim.cpp
+++ b/Engine/Plugins/Experimental/DynamicMaterial/Source/DynamicMaterialEditor/Private/UI/Widgets/Editor/PropertySelectorLayouts/SDMMaterialPropertySelector_WrapSlim.cpp
@@ -41,7 +41,7 @@ TSharedRef<SWidget> SDMMaterialPropertySelector_WrapSlim::CreateSlot_PropertyLis
 
 	const FMargin Padding = FMargin(0.f, 1.f);
 
-	TSharedRef<SWidget> PreviewButton = CreateSlot_SelectButton(EDMMaterialEditorMode::MaterialPreview, EDMMaterialPropertyType::None);
+	TSharedRef<SWidget> PreviewButton = CreateSlot_SelectButton(FDMMaterialEditorPage::Preview);
 	SetupMaterialPreviewButton(PreviewButton);
 
 	NewSlotList->AddSlot()
@@ -53,13 +53,13 @@ TSharedRef<SWidget> SDMMaterialPropertySelector_WrapSlim::CreateSlot_PropertyLis
 	NewSlotList->AddSlot()
 		.Padding(Padding)
 		[
-			CreateSlot_SelectButton(EDMMaterialEditorMode::GlobalSettings, EDMMaterialPropertyType::None)
+			CreateSlot_SelectButton(FDMMaterialEditorPage::GlobalSettings)
 		];
 
 	NewSlotList->AddSlot()
 		.Padding(Padding)
 		[
-			CreateSlot_SelectButton(EDMMaterialEditorMode::PropertyPreviews, EDMMaterialPropertyType::None)
+			CreateSlot_SelectButton(FDMMaterialEditorPage::Properties)
 		];
 
 	for (const TPair<EDMMaterialPropertyType, UDMMaterialProperty*>& PropertyPair : EditorOnlyData->GetMaterialProperties())
@@ -74,33 +74,34 @@ TSharedRef<SWidget> SDMMaterialPropertySelector_WrapSlim::CreateSlot_PropertyLis
 		NewSlotList->AddSlot()
 			.Padding(Padding)
 			[
-				CreateSlot_SelectButton(EDMMaterialEditorMode::EditSlot, PropertyPair.Key)
+				CreateSlot_SelectButton({EDMMaterialEditorMode::EditSlot, PropertyPair.Key})
 			];
 	}
 
 	return NewSlotList;
 }
 
-TSharedRef<SWidget> SDMMaterialPropertySelector_WrapSlim::CreateSlot_SelectButton(EDMMaterialEditorMode InEditMode, EDMMaterialPropertyType InMaterialProperty)
+TSharedRef<SWidget> SDMMaterialPropertySelector_WrapSlim::CreateSlot_SelectButton(const FDMMaterialEditorPage& InPage)
 {
 	if (const UDynamicMaterialEditorSettings* Settings = GetDefault<UDynamicMaterialEditorSettings>())
 	{
 		if (Settings->bUseFullChannelNamesInTopSlimLayout)
 		{
-			return SDMMaterialPropertySelector_WrapBase::CreateSlot_SelectButton(InEditMode, InMaterialProperty);
+			return SDMMaterialPropertySelector_WrapBase::CreateSlot_SelectButton(InPage);
 		}
 	}
 
-	const FText ButtonText = GetSelectButtonText(InEditMode, InMaterialProperty, /* Short Name */ true);
-	const FText ToolTip = GetButtonToolTip(InEditMode, InMaterialProperty);
+	const FText ButtonText = GetSelectButtonText(InPage, /* Short Name */ true);
+	const FText ToolTip = GetButtonToolTip(InPage);
+
 
 	return SNew(SCheckBox)
 		.Style(FAppStyle::Get(), "DetailsView.SectionButton")
 		.HAlign(EHorizontalAlignment::HAlign_Center)
 		.Padding(0.f)
-		.IsEnabled(this, &SDMMaterialPropertySelector_WrapSlim::GetPropertySelectEnabled, InEditMode, InMaterialProperty)
-		.IsChecked(this, &SDMMaterialPropertySelector_WrapSlim::GetPropertySelectState, InEditMode, InMaterialProperty)
-		.OnCheckStateChanged(this, &SDMMaterialPropertySelector_WrapSlim::OnPropertySelectStateChanged, InEditMode, InMaterialProperty)
+		.IsEnabled(this, &SDMMaterialPropertySelector_WrapSlim::GetPropertySelectEnabled, InPage)
+		.IsChecked(this, &SDMMaterialPropertySelector_WrapSlim::GetPropertySelectState, InPage)
+		.OnCheckStateChanged(this, &SDMMaterialPropertySelector_WrapSlim::OnPropertySelectStateChanged, InPage)
 		.ToolTipText(ToolTip)
 		.Content()
 		[
@@ -111,7 +112,7 @@ TSharedRef<SWidget> SDMMaterialPropertySelector_WrapSlim::CreateSlot_SelectButto
 			[
 				SNew(SImage)
 				.Image(FAppStyle::Get().GetBrush("FilterBar.FilterImage"))
-				.ColorAndOpacity(this, &SDMMaterialPropertySelector_WrapSlim::GetPropertySelectButtonChipColor, InEditMode, InMaterialProperty)
+				.ColorAndOpacity(this, &SDMMaterialPropertySelector_WrapSlim::GetPropertySelectButtonChipColor, InPage)
 				.DesiredSizeOverride(FVector2D(8, 17))
 			]
 			+SHorizontalBox::Slot()
diff --git a/Engine/Plugins/Experimental/DynamicMaterial/Source/DynamicMaterialEditor/Private/UI/Widgets/Editor/PropertySelectorLayouts/SDMMaterialPropertySelector_WrapSlim.h b/Engine/Plugins/Experimental/DynamicMaterial/Source/DynamicMaterialEditor/Private/UI/Widgets/Editor/PropertySelectorLayouts/SDMMaterialPropertySelector_WrapSlim.h
--- a/Engine/Plugins/Experimental/DynamicMaterial/Source/DynamicMaterialEditor/Private/UI/Widgets/Editor/PropertySelectorLayouts/SDMMaterialPropertySelector_WrapSlim.h
+++ b/Engine/Plugins/Experimental/DynamicMaterial/Source/DynamicMaterialEditor/Private/UI/Widgets/Editor/PropertySelectorLayouts/SDMMaterialPropertySelector_WrapSlim.h
@@ -17,6 +17,6 @@ public:
 protected:
 	//~ Begin SDMMaterialPropertySelector
 	virtual TSharedRef<SWidget> CreateSlot_PropertyList() override;
-	virtual TSharedRef<SWidget> CreateSlot_SelectButton(EDMMaterialEditorMode InEditMode, EDMMaterialPropertyType InMaterialProperty) override;
+	virtual TSharedRef<SWidget> CreateSlot_SelectButton(const FDMMaterialEditorPage& InPage) override;
 	//~ End SDMMaterialPropertySelector
 };
diff --git a/Engine/Plugins/Experimental/DynamicMaterial/Source/DynamicMaterialEditor/Private/UI/Widgets/Editor/SDMMaterialPropertyPreviews.cpp b/Engine/Plugins/Experimental/DynamicMaterial/Source/DynamicMaterialEditor/Private/UI/Widgets/Editor/SDMMaterialPropertyPreviews.cpp
--- a/Engine/Plugins/Experimental/DynamicMaterial/Source/DynamicMaterialEditor/Private/UI/Widgets/Editor/SDMMaterialPropertyPreviews.cpp
+++ b/Engine/Plugins/Experimental/DynamicMaterial/Source/DynamicMaterialEditor/Private/UI/Widgets/Editor/SDMMaterialPropertyPreviews.cpp
@@ -228,8 +228,9 @@ void SDMMaterialPropertyPreviews::AddPropertyPreview(const TSharedRef<SWrapBox>&
 
 TSharedRef<SWidget> SDMMaterialPropertyPreviews::CreateSlot_EnabledButton(EDMMaterialPropertyType InMaterialProperty)
 {
+	const FDMMaterialEditorPage Page = {EDMMaterialEditorMode::EditSlot, InMaterialProperty};
 	const FText Format = LOCTEXT("PropertyEnableFormat", "Toggle the {0} property.\n\nProperty must be valid for the Material Type.");
-	const FText ToolTip = FText::Format(Format, SDMMaterialPropertySelector::GetSelectButtonText(EDMMaterialEditorMode::EditSlot, InMaterialProperty, /* Short Name */ false));
+	const FText ToolTip = FText::Format(Format, SDMMaterialPropertySelector::GetSelectButtonText(Page, /* Short Name */ false));
 
 	return SNew(SCheckBox)
 		.IsEnabled(this, &SDMMaterialPropertyPreviews::GetPropertyEnabledEnabled, InMaterialProperty)
@@ -240,10 +241,12 @@ TSharedRef<SWidget> SDMMaterialPropertyPreviews::CreateSlot_EnabledButton(EDMMat
 
 TSharedRef<SWidget> SDMMaterialPropertyPreviews::CreateSlot_PropertyName(EDMMaterialPropertyType InMaterialProperty)
 {
+	const FDMMaterialEditorPage Page = {EDMMaterialEditorMode::EditSlot, InMaterialProperty};
+
 	return SNew(STextBlock)
 		.Font(IDetailLayoutBuilder::GetDetailFont())
-		.Text(SDMMaterialPropertySelector::GetSelectButtonText(EDMMaterialEditorMode::EditSlot, InMaterialProperty, /* Short Name */ true))
-		.ToolTipText(SDMMaterialPropertySelector::GetSelectButtonText(EDMMaterialEditorMode::EditSlot, InMaterialProperty, /* Short Name */ false));
+		.Text(SDMMaterialPropertySelector::GetSelectButtonText(Page, /* Short Name */ true))
+		.ToolTipText(SDMMaterialPropertySelector::GetSelectButtonText(Page, /* Short Name */ false));
 }
 
 bool SDMMaterialPropertyPreviews::GetPropertyEnabledEnabled(EDMMaterialPropertyType InMaterialProperty) const
@@ -361,7 +364,7 @@ void SDMMaterialPropertyPreviews::OnPropertyEnabledStateChanged(ECheckBoxState I
 	Content.Invalidate();
 
 	// Make sure we go back to the property previews
-	EditorWidget->ShowPropertyPreviews();
+	EditorWidget->EditProperties();
 }
 
 FReply SDMMaterialPropertyPreviews::OnPreviewClicked(const FGeometry& InGeometry, const FPointerEvent& InPointerEvent, 
diff --git a/Engine/Plugins/Experimental/DynamicMaterial/Source/DynamicMaterialEditor/Private/UI/Widgets/Editor/SDMMaterialPropertySelector.cpp b/Engine/Plugins/Experimental/DynamicMaterial/Source/DynamicMaterialEditor/Private/UI/Widgets/Editor/SDMMaterialPropertySelector.cpp
--- a/Engine/Plugins/Experimental/DynamicMaterial/Source/DynamicMaterialEditor/Private/UI/Widgets/Editor/SDMMaterialPropertySelector.cpp
+++ b/Engine/Plugins/Experimental/DynamicMaterial/Source/DynamicMaterialEditor/Private/UI/Widgets/Editor/SDMMaterialPropertySelector.cpp
@@ -47,7 +47,7 @@ TSharedPtr<SDMMaterialEditor> SDMMaterialPropertySelector::GetEditorWidget() con
 	return EditorWidgetWeak.Pin();
 }
 
-void SDMMaterialPropertySelector::SetSelectedProperty(EDMMaterialEditorMode InEditMode, EDMMaterialPropertyType InMaterialProperty)
+void SDMMaterialPropertySelector::SetSelectedProperty(const FDMMaterialEditorPage& InPage)
 {
 	TSharedPtr<SDMMaterialEditor> EditorWidget = GetEditorWidget();
 
@@ -56,25 +56,14 @@ void SDMMaterialPropertySelector::SetSelectedProperty(EDMMaterialEditorMode InEd
 		return;
 	}
 
-	switch (InEditMode)
+	switch (InPage.EditMode)
 	{
 		case EDMMaterialEditorMode::MaterialPreview:
 			OpenMaterialPreviewTab();
 			break;
 
-		case EDMMaterialEditorMode::GlobalSettings:
-			EditorWidget->EditGlobalSettings();
-			break;
-
-		case EDMMaterialEditorMode::PropertyPreviews:
-			EditorWidget->ShowPropertyPreviews();
-			break;
-
-		case EDMMaterialEditorMode::EditSlot:
-			EditorWidget->SelectProperty(InMaterialProperty);
-			break;
-
 		default:
+			EditorWidget->SetActivePage(InPage);
 			break;
 	}
 }
@@ -144,7 +133,7 @@ TSharedPtr<SDMMaterialSlotEditor> SDMMaterialPropertySelector::GetSlotEditorWidg
 TSharedRef<SWidget> SDMMaterialPropertySelector::CreateSlot_EnabledButton(EDMMaterialPropertyType InMaterialProperty)
 {
 	const FText Format = LOCTEXT("PropertyEnableFormat", "Toggle the {0} property.\n\nProperty must be valid for the Material Type.");
-	const FText ToolTip = FText::Format(Format, GetSelectButtonText(EDMMaterialEditorMode::EditSlot, InMaterialProperty, /* Short Name */ false));
+	const FText ToolTip = FText::Format(Format, GetSelectButtonText({EDMMaterialEditorMode::EditSlot, InMaterialProperty}, /* Short Name */ false));
 
 	return SNew(SCheckBox)
 		.IsEnabled(this, &SDMMaterialPropertySelector::GetPropertyEnabledEnabled, InMaterialProperty)
@@ -153,10 +142,10 @@ TSharedRef<SWidget> SDMMaterialPropertySelector::CreateSlot_EnabledButton(EDMMat
 		.ToolTipText(ToolTip);
 }
 
-FText SDMMaterialPropertySelector::GetSelectButtonText(EDMMaterialEditorMode InEditMode, EDMMaterialPropertyType InMaterialProperty, 
+FText SDMMaterialPropertySelector::GetSelectButtonText(const FDMMaterialEditorPage& InPage, 
 	bool bInShortName)
 {
-	switch (InEditMode)
+	switch (InPage.EditMode)
 	{
 		case EDMMaterialEditorMode::MaterialPreview:
 			return bInShortName
@@ -168,24 +157,24 @@ FText SDMMaterialPropertySelector::GetSelectButtonText(EDMMaterialEditorMode InE
 				? LOCTEXT("GlobalSettingsShort", "Global")
 				: LOCTEXT("GlobalSettings", "Global Settings");
 
-		case EDMMaterialEditorMode::PropertyPreviews:
+		case EDMMaterialEditorMode::Properties:
 			return bInShortName
 				? LOCTEXT("ChannelsShort", "Chans")
 				: LOCTEXT("Channels", "Channels");
 
 		case EDMMaterialEditorMode::EditSlot:
 			return bInShortName
-				? UE::DynamicMaterialEditor::Private::GetMaterialPropertyShortDisplayName(InMaterialProperty)
-				: UE::DynamicMaterialEditor::Private::GetMaterialPropertyLongDisplayName(InMaterialProperty);
+				? UE::DynamicMaterialEditor::Private::GetMaterialPropertyShortDisplayName(InPage.MaterialProperty)
+				: UE::DynamicMaterialEditor::Private::GetMaterialPropertyLongDisplayName(InPage.MaterialProperty);
 
 		default:
 			return FText::GetEmpty();
 	}
 }
 
-FText SDMMaterialPropertySelector::GetButtonToolTip(EDMMaterialEditorMode InEditMode, EDMMaterialPropertyType InMaterialProperty)
+FText SDMMaterialPropertySelector::GetButtonToolTip(const FDMMaterialEditorPage& InPage)
 {
-	switch (InEditMode)
+	switch (InPage.EditMode)
 	{
 		case EDMMaterialEditorMode::MaterialPreview:
 			return LOCTEXT("MaterialPreviewToolTip", "Show a preview of the Material.");
@@ -193,13 +182,13 @@ FText SDMMaterialPropertySelector::GetButtonToolTip(EDMMaterialEditorMode InEdit
 		case EDMMaterialEditorMode::GlobalSettings:
 			return LOCTEXT("GeneralSettingsToolTip", "Edit the Material Global Settings.");
 
-		case EDMMaterialEditorMode::PropertyPreviews:
+		case EDMMaterialEditorMode::Properties:
 			return LOCTEXT("PropertyPreviewsToolTip", "Preview and toggle the Material Channels.");
 
 		case EDMMaterialEditorMode::EditSlot:
 		{
 			const FText Format = LOCTEXT("PropertySelectFormat", "Edit the {0} channel.");
-			return FText::Format(Format, GetSelectButtonText(InEditMode, InMaterialProperty, /* Short Name */ false));
+			return FText::Format(Format, GetSelectButtonText(InPage, /* Short Name */ false));
 		}
 
 		default:
@@ -241,7 +230,7 @@ bool SDMMaterialPropertySelector::SetPropertyEnabled(EDMMaterialPropertyType InM
 	{
 		if (InMaterialProperty == EditorWidget->GetSelectedPropertyType())
 		{
-			SetSelectedProperty(EDMMaterialEditorMode::GlobalSettings, EDMMaterialPropertyType::None);
+			SetSelectedProperty(FDMMaterialEditorPage::GlobalSettings);
 		}
 	}
 
@@ -317,25 +306,24 @@ void SDMMaterialPropertySelector::OnPropertyEnabledStateChanged(ECheckBoxState I
 	{
 		if (bSetEnabled)
 		{
-			SetSelectedProperty(EDMMaterialEditorMode::EditSlot, InMaterialProperty);
+			SetSelectedProperty({EDMMaterialEditorMode::EditSlot, InMaterialProperty});
 		}
 	}
 }
 
-bool SDMMaterialPropertySelector::GetPropertySelectEnabled(EDMMaterialEditorMode InEditMode, EDMMaterialPropertyType InMaterialProperty) const
+bool SDMMaterialPropertySelector::GetPropertySelectEnabled(FDMMaterialEditorPage InPage) const
 {
-	switch (InEditMode)
+	switch (InPage.EditMode)
 	{
 		case EDMMaterialEditorMode::EditSlot:
-			return GetPropertyEnabledEnabled(InMaterialProperty) && DoesPropertySlotExist(InMaterialProperty);
+			return GetPropertyEnabledEnabled(InPage.MaterialProperty) && DoesPropertySlotExist(InPage.MaterialProperty);
 
 		default:
 			return true;
 	}	
 }
 
-ECheckBoxState SDMMaterialPropertySelector::GetPropertySelectState(EDMMaterialEditorMode InEditMode, EDMMaterialPropertyType 
-	InMaterialProperty) const
+ECheckBoxState SDMMaterialPropertySelector::GetPropertySelectState(FDMMaterialEditorPage InPage) const
 {
 	TSharedPtr<SDMMaterialEditor> EditorWidget = GetEditorWidget();
 
@@ -344,10 +332,10 @@ ECheckBoxState SDMMaterialPropertySelector::GetPropertySelectState(EDMMaterialEd
 		return ECheckBoxState::Undetermined;
 	}
 
-	switch (InEditMode)
+	switch (InPage.EditMode)
 	{
 		default:
-			return EditorWidget->GetEditMode() == InEditMode
+			return EditorWidget->GetEditMode() == InPage.EditMode
 				? ECheckBoxState::Checked
 				: ECheckBoxState::Unchecked;
 
@@ -355,35 +343,33 @@ ECheckBoxState SDMMaterialPropertySelector::GetPropertySelectState(EDMMaterialEd
 			return ECheckBoxState::Unchecked;
 
 		case EDMMaterialEditorMode::EditSlot:
-			return InMaterialProperty == EditorWidget->GetSelectedPropertyType()
+			return InPage.MaterialProperty == EditorWidget->GetSelectedPropertyType()
 				? ECheckBoxState::Checked
 				: ECheckBoxState::Unchecked;
 	}
 }
 
-void SDMMaterialPropertySelector::OnPropertySelectStateChanged(ECheckBoxState InState, EDMMaterialEditorMode InEditMode, 
-	EDMMaterialPropertyType InMaterialProperty)
+void SDMMaterialPropertySelector::OnPropertySelectStateChanged(ECheckBoxState InState, FDMMaterialEditorPage InPage)
 {
 	if (InState != ECheckBoxState::Checked)
 	{
 		return;
 	}
 
-	SetSelectedProperty(InEditMode, InMaterialProperty);
+	SetSelectedProperty(InPage);
 }
 
-FSlateColor SDMMaterialPropertySelector::GetPropertySelectButtonChipColor(EDMMaterialEditorMode InEditMode, 
-	EDMMaterialPropertyType InMaterialProperty) const
+FSlateColor SDMMaterialPropertySelector::GetPropertySelectButtonChipColor(FDMMaterialEditorPage InPage) const
 {
-	switch (InEditMode)
+	switch (InPage.EditMode)
 	{
 		case EDMMaterialEditorMode::MaterialPreview:
 		case EDMMaterialEditorMode::GlobalSettings:
-		case EDMMaterialEditorMode::PropertyPreviews:
+		case EDMMaterialEditorMode::Properties:
 			return FStyleColors::AccentGreen;
 
 		case EDMMaterialEditorMode::EditSlot:
-			if (GetPropertySelectEnabled(EDMMaterialEditorMode::EditSlot, InMaterialProperty))
+			if (GetPropertySelectEnabled({EDMMaterialEditorMode::EditSlot, InPage.MaterialProperty}))
 			{
 				return FStyleColors::Primary;
 			}
diff --git a/Engine/Plugins/Experimental/DynamicMaterial/Source/DynamicMaterialEditor/Private/UI/Widgets/Editor/SDMMaterialPropertySelector.h b/Engine/Plugins/Experimental/DynamicMaterial/Source/DynamicMaterialEditor/Private/UI/Widgets/Editor/SDMMaterialPropertySelector.h
--- a/Engine/Plugins/Experimental/DynamicMaterial/Source/DynamicMaterialEditor/Private/UI/Widgets/Editor/SDMMaterialPropertySelector.h
+++ b/Engine/Plugins/Experimental/DynamicMaterial/Source/DynamicMaterialEditor/Private/UI/Widgets/Editor/SDMMaterialPropertySelector.h
@@ -16,6 +16,7 @@ class UDynamicMaterialModelEditorOnlyData;
 enum class ECheckBoxState : uint8;
 enum class EDMMaterialEditorMode : uint8;
 enum class EDMMaterialPropertyType : uint8;
+struct FDMMaterialEditorPage;
 
 namespace UE::DynamicMaterialEditor::Private
 {
@@ -34,9 +35,9 @@ class SDMMaterialPropertySelector : public SCompoundWidget
 	SLATE_END_ARGS()
 
 public:
-	static FText GetSelectButtonText(EDMMaterialEditorMode InEditMode, EDMMaterialPropertyType InMaterialProperty, bool bInShortName);
+	static FText GetSelectButtonText(const FDMMaterialEditorPage& InPage, bool bInShortName);
 
-	static FText GetButtonToolTip(EDMMaterialEditorMode InEditMode, EDMMaterialPropertyType InMaterialProperty);
+	static FText GetButtonToolTip(const FDMMaterialEditorPage& InPage);
 
 	virtual ~SDMMaterialPropertySelector() override = default;
 
@@ -55,7 +56,7 @@ protected:
 
 	TSharedRef<SWidget> CreateSlot_EnabledButton(EDMMaterialPropertyType InMaterialProperty);
 
-	virtual TSharedRef<SWidget> CreateSlot_SelectButton(EDMMaterialEditorMode InEditMode, EDMMaterialPropertyType InMaterialProperty) = 0;
+	virtual TSharedRef<SWidget> CreateSlot_SelectButton(const FDMMaterialEditorPage& InPage) = 0;
 
 	bool IsPropertyEnabled(EDMMaterialPropertyType InMaterialProperty) const;
 
@@ -67,12 +68,12 @@ protected:
 	ECheckBoxState GetPropertyEnabledState(EDMMaterialPropertyType InMaterialProperty) const;
 	void OnPropertyEnabledStateChanged(ECheckBoxState InState, EDMMaterialPropertyType InMaterialProperty);
 
-	bool GetPropertySelectEnabled(EDMMaterialEditorMode InEditMode, EDMMaterialPropertyType InMaterialProperty) const;
-	ECheckBoxState GetPropertySelectState(EDMMaterialEditorMode InEditMode, EDMMaterialPropertyType InMaterialProperty) const;
-	void OnPropertySelectStateChanged(ECheckBoxState InState, EDMMaterialEditorMode InEditMode, EDMMaterialPropertyType InMaterialProperty);
-	FSlateColor GetPropertySelectButtonChipColor(EDMMaterialEditorMode InEditMode, EDMMaterialPropertyType InMaterialProperty) const;
+	bool GetPropertySelectEnabled(FDMMaterialEditorPage InPage) const;
+	ECheckBoxState GetPropertySelectState(FDMMaterialEditorPage InPage) const;
+	void OnPropertySelectStateChanged(ECheckBoxState InState, FDMMaterialEditorPage InPage);
+	FSlateColor GetPropertySelectButtonChipColor(FDMMaterialEditorPage InPage) const;
 
-	void SetSelectedProperty(EDMMaterialEditorMode InEditMode, EDMMaterialPropertyType InMaterialProperty);
+	void SetSelectedProperty(const FDMMaterialEditorPage& InPage);
 
 	void SetupMaterialPreviewButton(const TSharedRef<SWidget>& InSelectButton);
 
diff --git a/Engine/Plugins/Experimental/DynamicMaterial/Source/DynamicMaterialEditor/Private/UI/Widgets/SDMMaterialDesigner.cpp b/Engine/Plugins/Experimental/DynamicMaterial/Source/DynamicMaterialEditor/Private/UI/Widgets/SDMMaterialDesigner.cpp
--- a/Engine/Plugins/Experimental/DynamicMaterial/Source/DynamicMaterialEditor/Private/UI/Widgets/SDMMaterialDesigner.cpp
+++ b/Engine/Plugins/Experimental/DynamicMaterial/Source/DynamicMaterialEditor/Private/UI/Widgets/SDMMaterialDesigner.cpp
@@ -478,8 +478,8 @@ void SDMMaterialDesigner::OnLayoutChanged()
 			NewEditor->EditGlobalSettings();
 			break;
 
-		case EDMMaterialEditorMode::PropertyPreviews:
-			NewEditor->ShowPropertyPreviews();
+		case EDMMaterialEditorMode::Properties:
+			NewEditor->EditProperties();
 			break;
 
 		case EDMMaterialEditorMode::EditSlot:
diff --git a/Engine/Plugins/Experimental/DynamicMaterial/Source/DynamicMaterialEditor/Private/UI/Widgets/SDMMaterialEditor.cpp b/Engine/Plugins/Experimental/DynamicMaterial/Source/DynamicMaterialEditor/Private/UI/Widgets/SDMMaterialEditor.cpp
--- a/Engine/Plugins/Experimental/DynamicMaterial/Source/DynamicMaterialEditor/Private/UI/Widgets/SDMMaterialEditor.cpp
+++ b/Engine/Plugins/Experimental/DynamicMaterial/Source/DynamicMaterialEditor/Private/UI/Widgets/SDMMaterialEditor.cpp
@@ -40,6 +40,15 @@
 
 #define LOCTEXT_NAMESPACE "SDMMaterialEditor"
 
+FDMMaterialEditorPage FDMMaterialEditorPage::Preview = {EDMMaterialEditorMode::MaterialPreview, EDMMaterialPropertyType::None};
+FDMMaterialEditorPage FDMMaterialEditorPage::GlobalSettings = {EDMMaterialEditorMode::GlobalSettings, EDMMaterialPropertyType::None};
+FDMMaterialEditorPage FDMMaterialEditorPage::Properties = {EDMMaterialEditorMode::Properties, EDMMaterialPropertyType::None};
+
+bool FDMMaterialEditorPage::operator==(const FDMMaterialEditorPage& InOther) const
+{
+	return EditMode == InOther.EditMode && MaterialProperty == InOther.MaterialProperty;
+}
+
 void SDMMaterialEditor::PrivateRegisterAttributes(FSlateAttributeDescriptor::FInitializer&)
 {
 }
@@ -49,7 +58,12 @@ SDMMaterialEditor::SDMMaterialEditor()
 	, CommandList(MakeShared<FUICommandList>())
 	, PreviewMaterialManager(MakeShared<FDMPreviewMaterialManager>())
 	, EditMode(EDMMaterialEditorMode::GlobalSettings)
+	, PageHistoryActive(0)
+	, PageHistoryCount(1)
 {
+	// Some small number to get us going
+	PageHistory.Reserve(20);
+	PageHistory.Add(FDMMaterialEditorPage::GlobalSettings);
 }
 
 SDMMaterialEditor::~SDMMaterialEditor()
@@ -251,6 +265,8 @@ void SDMMaterialEditor::SelectProperty(EDMMaterialPropertyType InProperty, bool
 	}
 
 	EditSlot(Slot);
+
+	PageHistoryAdd({EDMMaterialEditorMode::EditSlot, InProperty});
 }
 
 const TSharedRef<FUICommandList>& SDMMaterialEditor::GetCommandList() const
@@ -332,16 +348,18 @@ void SDMMaterialEditor::EditGlobalSettings(bool bInForceRefresh)
 	SelectedMaterialProperty = EDMMaterialPropertyType::None;
 
 	GlobalSettingsEditorSlot.Invalidate();
+
+	PageHistoryAdd(FDMMaterialEditorPage::GlobalSettings);
 }
 
-void SDMMaterialEditor::ShowPropertyPreviews(bool bInForceRefresh)
+void SDMMaterialEditor::EditProperties(bool bInForceRefresh)
 {
-	if (EditMode == EDMMaterialEditorMode::PropertyPreviews && !bInForceRefresh)
+	if (EditMode == EDMMaterialEditorMode::Properties && !bInForceRefresh)
 	{
 		return;
 	}
 
-	if (EditMode != EDMMaterialEditorMode::PropertyPreviews)
+	if (EditMode != EDMMaterialEditorMode::Properties)
 	{
 		SlotEditorSlot.Invalidate();
 		SplitterSlot = nullptr;
@@ -349,10 +367,12 @@ void SDMMaterialEditor::ShowPropertyPreviews(bool bInForceRefresh)
 		GlobalSettingsEditorSlot.Invalidate();
 	}
 
-	EditMode = EDMMaterialEditorMode::PropertyPreviews;
+	EditMode = EDMMaterialEditorMode::Properties;
 	SelectedMaterialProperty = EDMMaterialPropertyType::None;
 
 	MaterialPropertyPreviewsSlot.Invalidate();
+
+	PageHistoryAdd(FDMMaterialEditorPage::Properties);
 }
 
 void SDMMaterialEditor::OpenMaterialPreviewTab()
@@ -563,6 +583,16 @@ FReply SDMMaterialEditor::OnKeyDown(const FGeometry& InMyGeometry, const FKeyEve
 	return FReply::Unhandled();
 }
 
+FReply SDMMaterialEditor::OnMouseButtonDown(const FGeometry& InGeometry, const FPointerEvent& InPointerEvent)
+{
+	if (CommandList->ProcessCommandBindings(InPointerEvent))
+	{
+		return FReply::Handled();
+	}
+
+	return SCompoundWidget::OnMouseButtonDown(InGeometry, InPointerEvent);
+}
+
 void SDMMaterialEditor::PostUndo(bool bInSuccess)
 {
 	OnUndo();
@@ -576,9 +606,22 @@ void SDMMaterialEditor::PostRedo(bool bInSuccess)
 void SDMMaterialEditor::BindCommands(SDMMaterialSlotEditor* InSlotEditor)
 {
 	const FGenericCommands& GenericCommands = FGenericCommands::Get();
+	const FDynamicMaterialEditorCommands& DMEditorCommands = FDynamicMaterialEditorCommands::Get();
 
 	CommandList = MakeShared<FUICommandList>();
 
+	CommandList->MapAction(
+		DMEditorCommands.NavigateForward,
+		FExecuteAction::CreateSP(this, &SDMMaterialEditor::NavigateForward_Execute),
+		FCanExecuteAction::CreateSP(this, &SDMMaterialEditor::NavigateForward_CanExecute)
+	);
+
+	CommandList->MapAction(
+		DMEditorCommands.NavigateBack,
+		FExecuteAction::CreateSP(this, &SDMMaterialEditor::NavigateBack_Execute),
+		FCanExecuteAction::CreateSP(this, &SDMMaterialEditor::NavigateBack_CanExecute)
+	);
+
 	CommandList->MapAction(
 		FDynamicMaterialEditorCommands::Get().AddDefaultLayer,
 		FExecuteAction::CreateSP(InSlotEditor, &SDMMaterialSlotEditor::AddNewLayer),
@@ -700,7 +743,7 @@ void SDMMaterialEditor::ValidateSlots()
 				GlobalSettingsEditorSlot->Validate();
 			}
 		}
-		else if (EditMode == EDMMaterialEditorMode::PropertyPreviews)
+		else if (EditMode == EDMMaterialEditorMode::Properties)
 		{
 			if (MaterialPropertyPreviewsSlot.HasBeenInvalidated())
 			{
@@ -755,6 +798,93 @@ void SDMMaterialEditor::ClearSlots()
 	ClearSlots_Main();
 }
 
+void SDMMaterialEditor::PageHistoryAdd(const FDMMaterialEditorPage& InPage)
+{
+	if (PageHistory.IsValidIndex(PageHistoryActive) && PageHistory[PageHistoryActive] == InPage)
+	{
+		return;
+	}
+
+	const int32 NewPageIndex = PageHistoryActive + 1;
+
+	if (!PageHistory.IsValidIndex(NewPageIndex))
+	{
+		PageHistory.Add(InPage);
+	}
+	else
+	{
+		PageHistory[NewPageIndex] = InPage;
+	}	
+
+	PageHistoryActive = NewPageIndex;
+	PageHistoryCount = NewPageIndex + 1;
+}
+
+bool SDMMaterialEditor::SetActivePage(const FDMMaterialEditorPage& InPage)
+{
+	switch (InPage.EditMode)
+	{
+		// This is not a valid page
+		case EDMMaterialEditorMode::MaterialPreview:
+		default:
+			return false;
+
+		case EDMMaterialEditorMode::GlobalSettings:
+			EditGlobalSettings();
+			return true;
+
+		case EDMMaterialEditorMode::Properties:
+			EditProperties();
+			return true;
+
+		case EDMMaterialEditorMode::EditSlot:
+			SelectProperty(InPage.MaterialProperty);
+			return true;
+	}
+}
+
+bool SDMMaterialEditor::PageHistoryBack()
+{
+	const int32 NewPageIndex = PageHistoryActive - 1;
+
+	if (!PageHistory.IsValidIndex(NewPageIndex))
+	{
+		return false;
+	}
+
+	const int32 OldPageIndex = PageHistoryActive;
+	PageHistoryActive = NewPageIndex;
+
+	if (!SetActivePage(PageHistory[NewPageIndex]))
+	{
+		PageHistoryActive = OldPageIndex;
+		return false;
+	}
+
+	return true;
+}
+
+bool SDMMaterialEditor::PageHistoryForward()
+{
+	const int32 NewPageIndex = PageHistoryActive + 1;
+
+	if (NewPageIndex >= PageHistoryCount || !PageHistory.IsValidIndex(NewPageIndex))
+	{
+		return false;
+	}
+
+	const int32 OldPageIndex = PageHistoryActive;
+	PageHistoryActive = NewPageIndex;
+
+	if (!SetActivePage(PageHistory[NewPageIndex]))
+	{
+		PageHistoryActive = OldPageIndex;
+		return false;
+	}
+
+	return true;
+}
+
 void SDMMaterialEditor::CreateLayout()
 {
 	ContentSlot << CreateSlot_Container();
@@ -967,4 +1097,24 @@ void SDMMaterialEditor::OnSettingsChanged(const FPropertyChangedEvent& InPropert
 	}
 }
 
+void SDMMaterialEditor::NavigateForward_Execute()
+{
+	PageHistoryForward();
+}
+
+bool SDMMaterialEditor::NavigateForward_CanExecute()
+{
+	return (PageHistoryActive + 1) < PageHistoryCount;
+}
+
+void SDMMaterialEditor::NavigateBack_Execute()
+{
+	PageHistoryBack();
+}
+
+bool SDMMaterialEditor::NavigateBack_CanExecute()
+{
+	return PageHistoryActive > 0;
+}
+
 #undef LOCTEXT_NAMESPACE
diff --git a/Engine/Plugins/Experimental/DynamicMaterial/Source/DynamicMaterialEditor/Private/UI/Widgets/SDMMaterialEditor.h b/Engine/Plugins/Experimental/DynamicMaterial/Source/DynamicMaterialEditor/Private/UI/Widgets/SDMMaterialEditor.h
--- a/Engine/Plugins/Experimental/DynamicMaterial/Source/DynamicMaterialEditor/Private/UI/Widgets/SDMMaterialEditor.h
+++ b/Engine/Plugins/Experimental/DynamicMaterial/Source/DynamicMaterialEditor/Private/UI/Widgets/SDMMaterialEditor.h
@@ -44,11 +44,23 @@ namespace UE::DynamicMaterialEditor::Private
 enum class EDMMaterialEditorMode : uint8
 {
 	GlobalSettings,
-	PropertyPreviews,
+	Properties,
 	EditSlot,
 	MaterialPreview
 };
 
+struct FDMMaterialEditorPage
+{
+	static FDMMaterialEditorPage Preview;
+	static FDMMaterialEditorPage GlobalSettings;
+	static FDMMaterialEditorPage Properties;
+
+	EDMMaterialEditorMode EditMode;
+	EDMMaterialPropertyType MaterialProperty;
+
+	bool operator==(const FDMMaterialEditorPage& InOther) const;
+};
+
 class SDMMaterialEditor : public SCompoundWidget, public FSelfRegisteringEditorUndoClient
 {
 	SLATE_DECLARE_WIDGET(SDMMaterialEditor, SCompoundWidget)
@@ -107,7 +119,7 @@ public:
 
 	virtual void EditGlobalSettings(bool bInForceRefresh = false);
 
-	virtual void ShowPropertyPreviews(bool bInForceRefresh = false);
+	virtual void EditProperties(bool bInForceRefresh = false);
 
 	void OpenMaterialPreviewTab();
 
@@ -125,9 +137,16 @@ public:
 	DECLARE_MULTICAST_DELEGATE_TwoParams(FOnEditedComponentChanged, const TSharedRef<SDMMaterialComponentEditor>&, UDMMaterialComponent*);
 	FOnEditedComponentChanged::RegistrationType& GetOnEditedComponentChanged();
 
+	bool PageHistoryBack();
+
+	bool PageHistoryForward();
+
+	bool SetActivePage(const FDMMaterialEditorPage& InPage);
+
 	//~ Begin SWidget
 	virtual bool SupportsKeyboardFocus() const override;
 	virtual FReply OnKeyDown(const FGeometry& InMyGeometry, const FKeyEvent& InKeyEvent) override;
+	virtual FReply OnMouseButtonDown(const FGeometry& InGeometry, const FPointerEvent& InPointerEvent) override;
 	//~ End SWidget
 
 	//~ Begin FUndoClient
@@ -165,6 +184,10 @@ protected:
 	TWeakObjectPtr<UDMMaterialSlot> SlotToEdit;
 	TWeakObjectPtr<UDMMaterialComponent> ComponentToEdit;
 
+	TArray<FDMMaterialEditorPage> PageHistory;
+	int32 PageHistoryActive;
+	int32 PageHistoryCount;
+
 	FOnEditedSlotChanged OnEditedSlotChanged;
 	FOnEditedComponentChanged OnEditedComponentChanged;
 
@@ -191,6 +214,8 @@ protected:
 
 	virtual void ClearSlots_Main() = 0;
 
+	void PageHistoryAdd(const FDMMaterialEditorPage& InPage);
+
 	/** Slots */
 	void CreateLayout();
 
@@ -233,4 +258,10 @@ protected:
 	void OnSlotListUpdate(UDynamicMaterialModelBase* InMaterialModelBase);
 
 	void OnSettingsChanged(const FPropertyChangedEvent& InPropertyChangedEvent);
+
+	void NavigateForward_Execute();
+	bool NavigateForward_CanExecute();
+
+	void NavigateBack_Execute();
+	bool NavigateBack_CanExecute();
 };
