Material Designer
- Tidied up the source select menus.

diff --git a/Engine/Plugins/Experimental/DynamicMaterial/Source/DynamicMaterialEditor/Private/Components/MaterialStageExpressions/DMMSESceneTexture.cpp b/Engine/Plugins/Experimental/DynamicMaterial/Source/DynamicMaterialEditor/Private/Components/MaterialStageExpressions/DMMSESceneTexture.cpp
--- a/Engine/Plugins/Experimental/DynamicMaterial/Source/DynamicMaterialEditor/Private/Components/MaterialStageExpressions/DMMSESceneTexture.cpp
+++ b/Engine/Plugins/Experimental/DynamicMaterial/Source/DynamicMaterialEditor/Private/Components/MaterialStageExpressions/DMMSESceneTexture.cpp
@@ -7,7 +7,7 @@
 
 UDMMaterialStageExpressionSceneTexture::UDMMaterialStageExpressionSceneTexture()
 	: UDMMaterialStageExpression(
-		LOCTEXT("SceneTexture", "Scene Texture"),
+		LOCTEXT("SceneTexture", "Post Process"),
 		UDMMaterialStageExpression::FindClass(TEXT("MaterialExpressionSceneTexture"))
 	)
 {
diff --git a/Engine/Plugins/Experimental/DynamicMaterial/Source/DynamicMaterialEditor/Private/UI/Menus/DMMaterialSlotLayerMenus.cpp b/Engine/Plugins/Experimental/DynamicMaterial/Source/DynamicMaterialEditor/Private/UI/Menus/DMMaterialSlotLayerMenus.cpp
--- a/Engine/Plugins/Experimental/DynamicMaterial/Source/DynamicMaterialEditor/Private/UI/Menus/DMMaterialSlotLayerMenus.cpp
+++ b/Engine/Plugins/Experimental/DynamicMaterial/Source/DynamicMaterialEditor/Private/UI/Menus/DMMaterialSlotLayerMenus.cpp
@@ -112,36 +112,6 @@
 		return;
 	}
 
-	bool bHasValidSlot = false;
-
-	if constexpr (UE::DynamicMaterialEditor::bAdvancedSlotsEnabled)
-	{
-		const TArray<UDMMaterialSlot*>& Slots = ModelEditorOnlyData->GetSlots();
-
-		for (UDMMaterialSlot* SlotIter : Slots)
-		{
-			if (Slot == SlotIter)
-			{
-				continue;
-			}
-
-			if (SlotIter->GetLayers().IsEmpty())
-			{
-				continue;
-			}
-
-			TArray<EDMMaterialPropertyType> SlotProperties = ModelEditorOnlyData->GetMaterialPropertiesForSlot(SlotIter);
-
-			if (SlotProperties.IsEmpty())
-			{
-				continue;
-			}
-
-			bHasValidSlot = true;
-			break;
-		}
-	}
-
 	FToolMenuSection& NewSection = InMenu->AddSection(SlotLayerAddSectionName, LOCTEXT("AddLayer", "Add Layer"));
 
 	NewSection.AddMenuEntry(
@@ -160,49 +130,4 @@
 
 	NewSection.AddMenuEntry(
 		NAME_None,
-		LOCTEXT("AddTextureSampleBaseOnly", "Texture (No alpha)"),
-		LOCTEXT("AddTextureSampleBaseOnlyTooltip", "Add a Material Stage based on a Texture with the Alpha disabled."),
-		FSlateIcon(),
-		FUIAction(FExecuteAction::CreateWeakLambda(
-			Slot,
-			[Slot, ExpressionClass = TSubclassOf<UDMMaterialStageExpression>(UDMMaterialStageExpressionTextureSample::StaticClass())]
-			{
-				UDMMaterialLayerObject* NewLayer = UDMMaterialSlotFunctionLibrary::AddNewLayer_Expression(Slot, ExpressionClass);
-
-				NewLayer->ForEachValidStage(
-					EDMMaterialLayerStage::All,
-					[NewLayer](UDMMaterialStage* InStage)
-					{
-						InStage->SetEnabled(InStage->SetEnabled(NewLayer->GetStageType(InStage) != EDMMaterialLayerStage::Mask));
-					});
-			}
-		))
-	);
-
-	if (Slot->GetLayers().IsEmpty() == false)
-	{
-		NewSection.AddMenuEntry(
-			NAME_None,
-			LOCTEXT("AddAlphaOnly", "Alpha Only"),
-			LOCTEXT("AddAlphaOnlyTooltip", "Add an Alpha-Only Material Layer.\n\nThe base layer will be disabled by default. It can still be re-enabled later."),
-			FSlateIcon(),
-			FUIAction(FExecuteAction::CreateWeakLambda(
-				Slot,
-				[Slot, ExpressionClass = TSubclassOf<UDMMaterialStageExpression>(UDMMaterialStageExpressionTextureSample::StaticClass())]
-				{
-					UDMMaterialLayerObject* NewLayer = UDMMaterialSlotFunctionLibrary::AddNewLayer_Expression(Slot, ExpressionClass);
-
-					NewLayer->ForEachValidStage(
-						EDMMaterialLayerStage::All,
-						[NewLayer](UDMMaterialStage* InStage)
-						{
-							InStage->SetEnabled(InStage->SetEnabled(NewLayer->GetStageType(InStage) == EDMMaterialLayerStage::Mask));
-						});
-				}
-			))
-		);
-	}
-
-	NewSection.AddMenuEntry(
-		NAME_None,
 		LOCTEXT("AddColor", "Solid Color"),
@@ -219,18 +144,4 @@
 
 	NewSection.AddMenuEntry(
 		NAME_None,
-		LOCTEXT("AddColorAtlas", "Color Atlas"),
-		LOCTEXT("AddColorAtlasTooltip", "Add a new Material Layer with a Color Atlas."),
-		FSlateIcon(),
-		FUIAction(FExecuteAction::CreateWeakLambda(
-			Slot,
-			[Slot]
-			{
-				UDMMaterialSlotFunctionLibrary::AddNewLayer_NewLocalValue(Slot, EDMValueType::VT_ColorAtlas);
-			}
-		))
-	);
-
-	NewSection.AddMenuEntry(
-		NAME_None,
 		LOCTEXT("AddEdgeColor", "Texture Edge Color"),
@@ -247,46 +158,4 @@
 
 	NewSection.AddMenuEntry(
 		NAME_None,
-		LOCTEXT("AddSceneTexture", "Scene Texture"),
-		LOCTEXT("AddSceneTextureTooltip", "Add a new Material Layer that represents the Scene Texture for a post process material."),
-		FSlateIcon(),
-		FUIAction(FExecuteAction::CreateWeakLambda(
-			Slot,
-			[Slot]
-			{
-				UDMMaterialSlotFunctionLibrary::AddNewLayer_SceneTexture(Slot);
-			}
-		))
-	);
-
-	NewSection.AddMenuEntry(
-		NAME_None,
-		LOCTEXT("AddText", "Text"),
-		LOCTEXT("AddTextTooltip", "Add a Material Stage based on a Text Renderer."),
-		FSlateIcon(),
-		FUIAction(FExecuteAction::CreateWeakLambda(
-			Slot,
-			[Slot]
-			{
-				UDMMaterialSlotFunctionLibrary::AddNewLayer_Renderer(Slot, TSubclassOf<UDMRenderTargetRenderer>(UDMRenderTargetTextRenderer::StaticClass()));
-			}
-		))
-	);
-
-	NewSection.AddMenuEntry(
-		NAME_None,
-		LOCTEXT("AddWidget", "Widget"),
-		LOCTEXT("AddWidgetTooltip", "Add a Material Stage based on a Widget Renderer."),
-		FSlateIcon(),
-		FUIAction(FExecuteAction::CreateWeakLambda(
-			Slot,
-			[Slot]
-			{
-				UDMMaterialSlotFunctionLibrary::AddNewLayer_Renderer(Slot, TSubclassOf<UDMRenderTargetRenderer>(UDMRenderTargetUMGWidgetRenderer::StaticClass()));
-			}
-		))
-	);
-
-	NewSection.AddMenuEntry(
-		NAME_None,
 		LOCTEXT("AddNoise", "Noise"),
@@ -301,17 +170,21 @@
 		))
 	);
 
-	if constexpr (UE::DynamicMaterialEditor::bAdvancedSlotsEnabled)
+	if (ModelEditorOnlyData->GetDomain() == EMaterialDomain::MD_PostProcess)
 	{
-		if (bHasValidSlot)
-		{
-			NewSection.AddSubMenu(
-				NAME_None,
-				LOCTEXT("AddSlotStage", "Slot Output"),
-				LOCTEXT("AddSlotStageTooltip", "Add a Material Stage based on the output of another Material Slot."),
-				FNewToolMenuDelegate::CreateStatic(&AddLayerInputsMenu_Slots)
-			);
-		}
+		NewSection.AddMenuEntry(
+			NAME_None,
+			LOCTEXT("AddSceneTexture", "Post Process"),
+			LOCTEXT("AddSceneTextureTooltip", "Add a new Material Layer that represents the Scene Texture for a post process material."),
+			FSlateIcon(),
+			FUIAction(FExecuteAction::CreateWeakLambda(
+				Slot,
+				[Slot]
+				{
+					UDMMaterialSlotFunctionLibrary::AddNewLayer_SceneTexture(Slot);
+				}
+			))
+		);
 	}
 
 	const TArray<TStrongObjectPtr<UClass>>& Gradients = UDMMaterialStageGradient::GetAvailableGradients();
@@ -326,24 +199,17 @@
 		);
 	}
 
-	NewSection.AddMenuEntry(
-		NAME_None,
-		LOCTEXT("AddMaterialFunction", "Material Function"),
-		LOCTEXT("AddMaterialFunctionTooltip", "Add a new Material Layer based on a Material Function."),
-		FSlateIcon(),
-		FUIAction(FExecuteAction::CreateWeakLambda(
-			Slot,
-			[Slot]
-			{
-				UDMMaterialSlotFunctionLibrary::AddNewLayer_MaterialFunction(Slot);
-			}
-		))
-	);
-
 	if constexpr (UE::DynamicMaterialEditor::bGlobalValuesEnabled)
 	{
 		AddGlobalValueSection(InMenu);
 	}
+
+	NewSection.AddSubMenu(
+		NAME_None,
+		LOCTEXT("AddAdvancedStage", "Advanced"),
+		LOCTEXT("AddAdvancedStageTooltip", "Add an advanced Material Stage."),
+		FNewToolMenuDelegate::CreateStatic(&AddLayerMenu_Advanced)
+	);
 }
 
 void FDMMaterialSlotLayerMenus::AddLayerModifySection(UToolMenu* InMenu)
@@ -737,6 +603,154 @@
 			)
 		);
 	}
+
+	InMenu->AddMenuEntry(
+		NAME_None,
+		FToolMenuEntry::InitMenuEntry(
+			NAME_None,
+			LOCTEXT("AddColorAtlas", "Color Atlas"),
+			LOCTEXT("AddColorAtlasTooltip", "Add a new Material Layer with a Color Atlas."),
+			FSlateIcon(),
+			FUIAction(FExecuteAction::CreateWeakLambda(
+				Slot,
+				[Slot]
+				{
+					UDMMaterialSlotFunctionLibrary::AddNewLayer_NewLocalValue(Slot, EDMValueType::VT_ColorAtlas);
+				}
+			))
+		)
+	);
 }
 
+void FDMMaterialSlotLayerMenus::AddLayerMenu_Advanced(UToolMenu* InMenu)
+{
+	using namespace UE::DynamicMaterialEditor::Private;
+
+	if (!IsValid(InMenu))
+	{
+		return;
+	}
+
+	UDMMenuContext* MenuContext = InMenu->FindContext<UDMMenuContext>();
+
+	if (!MenuContext)
+	{
+		return;
+	}
+
+	TSharedPtr<SDMMaterialEditor> EditorWidget = MenuContext->GetEditorWidget();
+
+	if (!EditorWidget.IsValid())
+	{
+		return;
+	}
+
+	UDMMaterialSlot* Slot = EditorWidget->GetSlotEditorWidget()->GetSlot();
+
+	if (!Slot)
+	{
+		return;
+	}
+
+	UDynamicMaterialModelEditorOnlyData* ModelEditorOnlyData = Slot->GetMaterialModelEditorOnlyData();
+
+	if (!ModelEditorOnlyData)
+	{
+		return;
+	}
+
+	UDynamicMaterialModel* MaterialModel = ModelEditorOnlyData->GetMaterialModel();
+
+	if (!MaterialModel)
+	{
+		return;
+	}
+
+	FToolMenuSection& NewSection = InMenu->AddSection(
+		NAME_None,
+		LOCTEXT("Advanced", "Advanced")
+	);
+
+	NewSection.AddMenuEntry(
+		NAME_None,
+		LOCTEXT("AddText", "Text"),
+		LOCTEXT("AddTextTooltip", "Add a Material Stage based on a Text Renderer."),
+		FSlateIcon(),
+		FUIAction(FExecuteAction::CreateWeakLambda(
+			Slot,
+			[Slot]
+			{
+				UDMMaterialSlotFunctionLibrary::AddNewLayer_Renderer(Slot, TSubclassOf<UDMRenderTargetRenderer>(UDMRenderTargetTextRenderer::StaticClass()));
+			}
+		))
+	);
+
+	NewSection.AddMenuEntry(
+		NAME_None,
+		LOCTEXT("AddWidget", "Widget"),
+		LOCTEXT("AddWidgetTooltip", "Add a Material Stage based on a Widget Renderer."),
+		FSlateIcon(),
+		FUIAction(FExecuteAction::CreateWeakLambda(
+			Slot,
+			[Slot]
+			{
+				UDMMaterialSlotFunctionLibrary::AddNewLayer_Renderer(Slot, TSubclassOf<UDMRenderTargetRenderer>(UDMRenderTargetUMGWidgetRenderer::StaticClass()));
+			}
+		))
+	);
+
+	NewSection.AddMenuEntry(
+		NAME_None,
+		LOCTEXT("AddMaterialFunction", "Material Function"),
+		LOCTEXT("AddMaterialFunctionTooltip", "Add a new Material Layer based on a Material Function."),
+		FSlateIcon(),
+		FUIAction(FExecuteAction::CreateWeakLambda(
+			Slot,
+			[Slot]
+			{
+				UDMMaterialSlotFunctionLibrary::AddNewLayer_MaterialFunction(Slot);
+			}
+		))
+	);
+
+	if constexpr (UE::DynamicMaterialEditor::bAdvancedSlotsEnabled)
+	{
+		bool bHasValidSlot = false;
+		const TArray<UDMMaterialSlot*>& Slots = ModelEditorOnlyData->GetSlots();
+
+		for (UDMMaterialSlot* SlotIter : Slots)
+		{
+			if (Slot == SlotIter)
+			{
+				continue;
+			}
+
+			if (SlotIter->GetLayers().IsEmpty())
+			{
+				continue;
+			}
+
+			TArray<EDMMaterialPropertyType> SlotProperties = ModelEditorOnlyData->GetMaterialPropertiesForSlot(SlotIter);
+
+			if (SlotProperties.IsEmpty())
+			{
+				continue;
+			}
+
+			bHasValidSlot = true;
+			break;
+		}
+
+		if (bHasValidSlot)
+		{
+			NewSection.AddSubMenu(
+				NAME_None,
+				LOCTEXT("AddSlotStage", "Slot Output"),
+				LOCTEXT("AddSlotStageTooltip", "Add a Material Stage based on the output of another Material Slot."),
+				FNewToolMenuDelegate::CreateStatic(&AddLayerInputsMenu_Slots)
+			);
+		}
+	}
+}
+
 #undef LOCTEXT_NAMESPACE
diff --git a/Engine/Plugins/Experimental/DynamicMaterial/Source/DynamicMaterialEditor/Private/UI/Menus/DMMaterialSlotLayerMenus.h b/Engine/Plugins/Experimental/DynamicMaterial/Source/DynamicMaterialEditor/Private/UI/Menus/DMMaterialSlotLayerMenus.h
--- a/Engine/Plugins/Experimental/DynamicMaterial/Source/DynamicMaterialEditor/Private/UI/Menus/DMMaterialSlotLayerMenus.h
+++ b/Engine/Plugins/Experimental/DynamicMaterial/Source/DynamicMaterialEditor/Private/UI/Menus/DMMaterialSlotLayerMenus.h
@@ -32,4 +32,6 @@ private:
 	static void AddLayerInputsMenu_Slots(UToolMenu* InMenu);
 
 	static void AddLayerMenu_Gradients(UToolMenu* InMenu);
+
+	static void AddLayerMenu_Advanced(UToolMenu* InMenu);
 };
diff --git a/Engine/Plugins/Experimental/DynamicMaterial/Source/DynamicMaterialEditor/Private/UI/Menus/DMMaterialStageInputMenus.cpp b/Engine/Plugins/Experimental/DynamicMaterial/Source/DynamicMaterialEditor/Private/UI/Menus/DMMaterialStageInputMenus.cpp
deleted file mode 100644
diff --git a/Engine/Plugins/Experimental/DynamicMaterial/Source/DynamicMaterialEditor/Private/UI/Menus/DMMaterialStageInputMenus.h b/Engine/Plugins/Experimental/DynamicMaterial/Source/DynamicMaterialEditor/Private/UI/Menus/DMMaterialStageInputMenus.h
deleted file mode 100644
diff --git a/Engine/Plugins/Experimental/DynamicMaterial/Source/DynamicMaterialEditor/Private/UI/Menus/DMMaterialStageSourceMenus.cpp b/Engine/Plugins/Experimental/DynamicMaterial/Source/DynamicMaterialEditor/Private/UI/Menus/DMMaterialStageSourceMenus.cpp
--- a/Engine/Plugins/Experimental/DynamicMaterial/Source/DynamicMaterialEditor/Private/UI/Menus/DMMaterialStageSourceMenus.cpp
+++ b/Engine/Plugins/Experimental/DynamicMaterial/Source/DynamicMaterialEditor/Private/UI/Menus/DMMaterialStageSourceMenus.cpp
@@ -708,59 +708,190 @@ void FDMMaterialStageSourceMenus::GenerateChangeSourceMenu_Gradients(UToolMenu*
 		{
 			const FText MenuName = GradientCDO->GetDescription();
 
-			InMenu->AddMenuEntry(NAME_None, FToolMenuEntry::InitMenuEntry(NAME_None,
-				MenuName,
-				LOCTEXT("ChangeSourceGradientTooltip", "Change the source of this stage to a Material Gradient."),
-				FSlateIcon(),
-				FUIAction(FExecuteAction::CreateWeakLambda(
-					MenuContext,
-					[MenuContext, GradientClass]()
-					{
-						UDMMaterialStage* const Stage = MenuContext->GetStage();
-						if (!Stage)
+			InMenu->AddMenuEntry(
+				NAME_None,
+				FToolMenuEntry::InitMenuEntry(
+					NAME_None,
+					MenuName,
+					LOCTEXT("ChangeSourceGradientTooltip", "Change the source of this stage to a Material Gradient."),
+					FSlateIcon(),
+					FUIAction(FExecuteAction::CreateWeakLambda(
+						MenuContext,
+						[MenuContext, GradientClass]()
 						{
-							return;
-						}
+							UDMMaterialStage* const Stage = MenuContext->GetStage();
+							if (!Stage)
+							{
+								return;
+							}
 
-						UDMMaterialStageSource* const StageSource = Stage->GetSource();
-						if (!StageSource)
-						{
-							return;
-						}
+							UDMMaterialStageSource* const StageSource = Stage->GetSource();
+							if (!StageSource)
+							{
+								return;
+							}
 
-						if (StageSource->IsA<UDMMaterialStageBlend>())
-						{
-							FScopedTransaction Transaction(LOCTEXT("SetStageInputBase", "Set Material Designer Base Source"));
-							Stage->Modify();
+							if (StageSource->IsA<UDMMaterialStageBlend>())
+							{
+								FScopedTransaction Transaction(LOCTEXT("SetStageInputBase", "Set Material Designer Base Source"));
+								Stage->Modify();
 
-							UDMMaterialStageInputGradient::ChangeStageInput_Gradient(
-								Stage,
-								GradientClass.Get(), UDMMaterialStageBlend::InputB,
-								FDMMaterialStageConnectorChannel::WHOLE_CHANNEL,
-								FDMMaterialStageConnectorChannel::WHOLE_CHANNEL
-							);
-						}
-						else if (StageSource->IsA<UDMMaterialStageThroughputLayerBlend>())
-						{
-							FScopedTransaction Transaction(LOCTEXT("SetStageInputMask", "Set Material Designer Mask Source"));
-							Stage->Modify();
-							UDMMaterialStageInputGradient::ChangeStageInput_Gradient(
-								Stage,
-								GradientClass.Get(), 
-								UDMMaterialStageThroughputLayerBlend::InputMaskSource,
-								FDMMaterialStageConnectorChannel::WHOLE_CHANNEL,
-								FDMMaterialStageConnectorChannel::WHOLE_CHANNEL
-							);
-						}
-						else
-						{
-							ensureMsgf(false, TEXT("Invalid stage type (%s)"), *StageSource->GetClass()->GetName());
-						}
-					})
-				)
+								UDMMaterialStageInputGradient::ChangeStageInput_Gradient(
+									Stage,
+									GradientClass.Get(), UDMMaterialStageBlend::InputB,
+									FDMMaterialStageConnectorChannel::WHOLE_CHANNEL,
+									FDMMaterialStageConnectorChannel::WHOLE_CHANNEL
+								);
+							}
+							else if (StageSource->IsA<UDMMaterialStageThroughputLayerBlend>())
+							{
+								FScopedTransaction Transaction(LOCTEXT("SetStageInputMask", "Set Material Designer Mask Source"));
+								Stage->Modify();
+								UDMMaterialStageInputGradient::ChangeStageInput_Gradient(
+									Stage,
+									GradientClass.Get(), 
+									UDMMaterialStageThroughputLayerBlend::InputMaskSource,
+									FDMMaterialStageConnectorChannel::WHOLE_CHANNEL,
+									FDMMaterialStageConnectorChannel::WHOLE_CHANNEL
+								);
+							}
+							else
+							{
+								ensureMsgf(false, TEXT("Invalid stage type (%s)"), *StageSource->GetClass()->GetName());
+							}
+						})
+					)
 			));
 		}
 	}
+
+	InMenu->AddMenuEntry(
+		NAME_None,
+		FToolMenuEntry::InitMenuEntry(
+			NAME_None,
+			LOCTEXT("ChangeSourceColorAtlas", "Color Atlas"),
+			LOCTEXT("ChangeSourceColorAtlasTooltip", "Change the source of this stage to a Color Atlas."),
+			FSlateIcon(),
+			FUIAction(
+				FExecuteAction::CreateStatic(
+					&ChangeSourceToColorAtlasFromContext,
+					MenuContext
+				)
+			)
+		)
+	);
+}
+
+void FDMMaterialStageSourceMenus::GenerateChangeSourceMenu_Advanced(UToolMenu* const InMenu)
+{
+	if (!ensure(IsValid(InMenu)))
+	{
+		return;
+	}
+
+	UDMMenuContext* MenuContext = InMenu->FindContext<UDMMenuContext>();
+	if (!ensure(MenuContext))
+	{
+		return;
+	}
+
+	UDMMaterialSlot* Slot = MenuContext->GetSlot();
+
+	UDynamicMaterialModel* MaterialModel = MenuContext->GetModel();
+
+	if (!ensure(MaterialModel))
+	{
+		return;
+	}
+
+	UDynamicMaterialModelEditorOnlyData* ModelEditorOnlyData = UDynamicMaterialModelEditorOnlyData::Get(MaterialModel);
+
+	if (!ensure(ModelEditorOnlyData))
+	{
+		return;
+	}
+
+	FToolMenuSection& NewSection = InMenu->AddSection(
+		NAME_None,
+		LOCTEXT("Advanced", "Advanced")
+	);
+
+	NewSection.AddMenuEntry("Text",
+		LOCTEXT("ChangeSourceText", "Text"),
+		LOCTEXT("ChangeSourceTextTooltip", "Change the source of this stage to a Text Renderer."),
+		FSlateIcon(),
+		FUIAction(
+			FExecuteAction::CreateStatic(
+				&ChangeSourceToTextFromContext,
+				MenuContext
+			)
+		)
+	);
+
+	NewSection.AddMenuEntry("Widget",
+		LOCTEXT("ChangeSourceWidget", "Widget"),
+		LOCTEXT("ChangeSourceWidgetTooltip", "Change the source of this stage to a Widget Renderer."),
+		FSlateIcon(),
+		FUIAction(
+			FExecuteAction::CreateStatic(
+				&ChangeSourceToWidgetFromContext,
+				MenuContext
+			)
+		)
+	);
+
+	NewSection.AddMenuEntry("MaterialFunction",
+		LOCTEXT("ChangeSourceMaterialFunction", "Material Function"),
+		LOCTEXT("ChangeSourceMaterialFunctionTooltip", "Change the source of this stage to a Material Function."),
+		FSlateIcon(),
+		FUIAction(
+			FExecuteAction::CreateStatic(
+				&ChangeSourceToMaterialFunctionFromContext,
+				MenuContext
+			)
+		)
+	);
+
+	if constexpr (UE::DynamicMaterialEditor::bAdvancedSlotsEnabled)
+	{
+		bool bHasValidSlot = false;
+
+		if constexpr (UE::DynamicMaterialEditor::bAdvancedSlotsEnabled)
+		{
+			const TArray<UDMMaterialSlot*>& Slots = ModelEditorOnlyData->GetSlots();
+
+			for (UDMMaterialSlot* SlotIter : Slots)
+			{
+				if (SlotIter == Slot)
+				{
+					continue;
+				}
+
+				if (SlotIter->GetLayers().IsEmpty())
+				{
+					continue;
+				}
+
+				bHasValidSlot = true;
+				break;
+			}
+		}
+
+		if (bHasValidSlot)
+		{
+			NewSection.AddDynamicEntry(
+				NAME_None,
+				FNewToolMenuSectionDelegate::CreateLambda(
+					[](FToolMenuSection& InSection)
+					{
+						InSection.AddSubMenu("SlotOutput",
+							LOCTEXT("ChangeSourceSlotOuptut", "Slot Output"),
+							LOCTEXT("ChangeSourceSlotOutputTooltip", "Change the source of this stage to the output from another Material Slot."),
+							FNewToolMenuDelegate::CreateStatic(&GenerateChangeSourceMenu_Slots)
+							);
+					}));
+		}
+	}
 }
 
 void FDMMaterialStageSourceMenus::ChangeSourceToTextureSampleFromContext(UDMMenuContext* InMenuContext)
@@ -1328,29 +1459,6 @@ void FDMMaterialStageSourceMenus::CreateChangeMaterialStageSource(FToolMenuSecti
 
 	const TArray<TStrongObjectPtr<UClass>>& Gradients = UDMMaterialStageGradient::GetAvailableGradients();
 
-	bool bHasValidSlot = false;
-
-	if constexpr (UE::DynamicMaterialEditor::bAdvancedSlotsEnabled)
-	{
-		const TArray<UDMMaterialSlot*>& Slots = ModelEditorOnlyData->GetSlots();
-
-		for (UDMMaterialSlot* SlotIter : Slots)
-		{
-			if (SlotIter == Slot)
-			{
-				continue;
-			}
-
-			if (SlotIter->GetLayers().IsEmpty())
-			{
-				continue;
-			}
-
-			bHasValidSlot = true;
-			break;
-		}
-	}
-
 	InSection.AddMenuEntry("TextureSample",
 		LOCTEXT("TextureSample", "Texture"),
 		LOCTEXT("ChangeSourceTextureSampleTooltip", "Change the source of this stage to a texture."),
@@ -1375,32 +1483,32 @@ void FDMMaterialStageSourceMenus::CreateChangeMaterialStageSource(FToolMenuSecti
 		)
 	);
 
-	InSection.AddMenuEntry("ColorAtlas",
-		LOCTEXT("ChangeSourceColorAtlas", "Color Atlas"),
-		LOCTEXT("ChangeSourceColorAtlasTooltip", "Change the source of this stage to a Color Atlas."),
+	InSection.AddMenuEntry("TextureSample_EdgeColor",
+		LOCTEXT("AddTextureSampleEgdeColor", "Texture Edge Color"),
+		LOCTEXT("ChangeSourceTextureSampleEdgeColorTooltip", "Change the source of this stage to the edge color of a texture."),
 		FSlateIcon(),
 		FUIAction(
 			FExecuteAction::CreateStatic(
-				&ChangeSourceToColorAtlasFromContext,
+				&ChangeSourceToTextureSampleEdgeColorFromContext,
 				MenuContext
 			)
 		)
 	);
 
-	InSection.AddMenuEntry("TextureSample_EdgeColor",
-		LOCTEXT("AddTextureSampleEgdeColor", "Texture Edge Color"),
-		LOCTEXT("ChangeSourceTextureSampleEdgeColorTooltip", "Change the source of this stage to the edge color of a texture."),
+	InSection.AddMenuEntry("Noise",
+		LOCTEXT("ChangeSourceNoise", "Noise"),
+		LOCTEXT("ChangeSourceNoiseTooltip", "Change the source of this stage to a Noise Renderer."),
 		FSlateIcon(),
 		FUIAction(
 			FExecuteAction::CreateStatic(
-				&ChangeSourceToTextureSampleEdgeColorFromContext,
+				&ChangeSourceToNoiseFromContext,
 				MenuContext
 			)
 		)
 	);
 
 	InSection.AddMenuEntry("SceneTexture",
-		LOCTEXT("AddSceneTexture", "Scene Texture"),
+		LOCTEXT("AddSceneTexture", "Post Process"),
 		LOCTEXT("ChangeSourceSceneTextureTooltip", "Change the source of this stage to Scene Texture in post process materials."),
 		FSlateIcon(),
 		FUIAction(
@@ -1415,42 +1523,6 @@ void FDMMaterialStageSourceMenus::CreateChangeMaterialStageSource(FToolMenuSecti
 		)
 	);
 
-	InSection.AddMenuEntry("Text",
-		LOCTEXT("ChangeSourceText", "Text"),
-		LOCTEXT("ChangeSourceTextTooltip", "Change the source of this stage to a Text Renderer."),
-		FSlateIcon(),
-		FUIAction(
-			FExecuteAction::CreateStatic(
-				&ChangeSourceToTextFromContext,
-				MenuContext
-			)
-		)
-	);
-
-	InSection.AddMenuEntry("Widget",
-		LOCTEXT("ChangeSourceWidget", "Widget"),
-		LOCTEXT("ChangeSourceWidgetTooltip", "Change the source of this stage to a Widget Renderer."),
-		FSlateIcon(),
-		FUIAction(
-			FExecuteAction::CreateStatic(
-				&ChangeSourceToWidgetFromContext,
-				MenuContext
-			)
-		)
-	);
-
-	InSection.AddMenuEntry("Noise",
-		LOCTEXT("ChangeSourceNoise", "Noise"),
-		LOCTEXT("ChangeSourceNoiseTooltip", "Change the source of this stage to a Noise Renderer."),
-		FSlateIcon(),
-		FUIAction(
-			FExecuteAction::CreateStatic(
-				&ChangeSourceToNoiseFromContext,
-				MenuContext
-			)
-		)
-	);
-
 	if constexpr (UE::DynamicMaterialEditor::bGlobalValuesEnabled)
 	{
 		const TArray<UDMMaterialValue*>& Values = MaterialModel->GetValues();
@@ -1486,35 +1558,11 @@ void FDMMaterialStageSourceMenus::CreateChangeMaterialStageSource(FToolMenuSecti
 		);
 	}
 
-	InSection.AddMenuEntry("MaterialFunction",
-		LOCTEXT("ChangeSourceMaterialFunction", "Material Function"),
-		LOCTEXT("ChangeSourceMaterialFunctionTooltip", "Change the source of this stage to a Material Function."),
-		FSlateIcon(),
-		FUIAction(
-			FExecuteAction::CreateStatic(
-				&ChangeSourceToMaterialFunctionFromContext,
-				MenuContext
-			)
-		)
+	InSection.AddSubMenu("Advanced",
+		LOCTEXT("ChangeSourceAdvanced", "Advanced"),
+		LOCTEXT("ChangeSourceAdvancedTooltip", "Add an advanced Material Stage."),
+		FNewToolMenuDelegate::CreateStatic(&GenerateChangeSourceMenu_Advanced)
 	);
-
-	if constexpr (UE::DynamicMaterialEditor::bAdvancedSlotsEnabled)
-	{
-		if (bHasValidSlot)
-		{
-			InSection.AddDynamicEntry(
-				NAME_None,
-				FNewToolMenuSectionDelegate::CreateLambda(
-					[](FToolMenuSection& InSection)
-					{
-						InSection.AddSubMenu("SlotOutput",
-						LOCTEXT("ChangeSourceSlotOuptut", "Slot Output"),
-						LOCTEXT("ChangeSourceSlotOutputTooltip", "Change the source of this stage to the output from another Material Slot."),
-						FNewToolMenuDelegate::CreateStatic(&GenerateChangeSourceMenu_Slots)
-						);
-					}));
-		}
-	}
 }
 
 #undef LOCTEXT_NAMESPACE
diff --git a/Engine/Plugins/Experimental/DynamicMaterial/Source/DynamicMaterialEditor/Private/UI/Menus/DMMaterialStageSourceMenus.h b/Engine/Plugins/Experimental/DynamicMaterial/Source/DynamicMaterialEditor/Private/UI/Menus/DMMaterialStageSourceMenus.h
--- a/Engine/Plugins/Experimental/DynamicMaterial/Source/DynamicMaterialEditor/Private/UI/Menus/DMMaterialStageSourceMenus.h
+++ b/Engine/Plugins/Experimental/DynamicMaterial/Source/DynamicMaterialEditor/Private/UI/Menus/DMMaterialStageSourceMenus.h
@@ -41,6 +41,8 @@ private:
 
 	static void GenerateChangeSourceMenu_Gradients(UToolMenu* const InMenu);
 
+	static void GenerateChangeSourceMenu_Advanced(UToolMenu* const InMenu);
+
 	static void ChangeSourceToTextureSampleFromContext(UDMMenuContext* InMenuContext);
 
 	static void ChangeSourceToNoiseFromContext(UDMMenuContext* InMenuContext);
