Material Designer
- Fixed preview splitter assigning wrong value.

diff --git a/Engine/Plugins/Experimental/DynamicMaterial/Source/DynamicMaterialEditor/Private/UI/Widgets/Editor/EditorLayouts/SDMMaterialEditor_TopVertical.cpp b/Engine/Plugins/Experimental/DynamicMaterial/Source/DynamicMaterialEditor/Private/UI/Widgets/Editor/EditorLayouts/SDMMaterialEditor_TopVertical.cpp
--- a/Engine/Plugins/Experimental/DynamicMaterial/Source/DynamicMaterialEditor/Private/UI/Widgets/Editor/EditorLayouts/SDMMaterialEditor_TopVertical.cpp
+++ b/Engine/Plugins/Experimental/DynamicMaterial/Source/DynamicMaterialEditor/Private/UI/Widgets/Editor/EditorLayouts/SDMMaterialEditor_TopVertical.cpp
@@ -67,7 +67,7 @@ TSharedRef<SWidget> SDMMaterialEditor_TopVertical::CreateSlot_Main()
 {
 	using namespace UE::DynamicMaterialEditor::Private;
 
-	float SplitterValue = 300;
+	float SplitterValue = 0.333;
 
 	if (UDynamicMaterialEditorSettings* Settings = UDynamicMaterialEditorSettings::Get())
 	{
@@ -125,11 +125,11 @@ TSharedRef<SDMMaterialPropertySelector> SDMMaterialEditor_TopVertical::CreateSlo
 
 void SDMMaterialEditor_TopVertical::OnTopSplitterResized()
 {
-	if (SplitterSlot)
+	if (SplitterSlot_Top)
 	{
 		if (UDynamicMaterialEditorSettings* Settings = UDynamicMaterialEditorSettings::Get())
 		{
-			const float SplitterLocation = static_cast<SSplitter::FSlot*>(SplitterSlot)->GetSizeValue();
+			const float SplitterLocation = static_cast<SSplitter::FSlot*>(SplitterSlot_Top)->GetSizeValue();
 			Settings->PreviewSplitterLocation = SplitterLocation;
 			Settings->SaveConfig();
 		}
