Material Designer
- Changed the properties page to have more utility.

diff --git a/Engine/Plugins/Experimental/DynamicMaterial/Source/DynamicMaterialEditor/Private/UI/Widgets/Editor/EditorLayouts/SDMMaterialEditor_LeftBase.cpp b/Engine/Plugins/Experimental/DynamicMaterial/Source/DynamicMaterialEditor/Private/UI/Widgets/Editor/EditorLayouts/SDMMaterialEditor_LeftBase.cpp
--- a/Engine/Plugins/Experimental/DynamicMaterial/Source/DynamicMaterialEditor/Private/UI/Widgets/Editor/EditorLayouts/SDMMaterialEditor_LeftBase.cpp
+++ b/Engine/Plugins/Experimental/DynamicMaterial/Source/DynamicMaterialEditor/Private/UI/Widgets/Editor/EditorLayouts/SDMMaterialEditor_LeftBase.cpp
@@ -7,7 +7,7 @@
 #include "UI/Widgets/Editor/SDMMaterialComponentEditor.h"
 #include "UI/Widgets/Editor/SDMMaterialGlobalSettingsEditor.h"
 #include "UI/Widgets/Editor/SDMMaterialPreview.h"
-#include "UI/Widgets/Editor/SDMMaterialPropertyPreviews.h"
+#include "UI/Widgets/Editor/SDMMaterialProperties.h"
 #include "UI/Widgets/Editor/SDMMaterialPropertySelector.h"
 #include "UI/Widgets/Editor/SDMMaterialSlotEditor.h"
 #include "Widgets/Layout/SBorder.h"
@@ -246,7 +246,7 @@ TSharedRef<SWidget> SDMMaterialEditor_LeftBase::CreateSlot_Right_PropertyPreview
 			]
 		];
 
-	MaterialPropertyPreviewsSlot = TDMWidgetSlot<SDMMaterialPropertyPreviews>(PropertyPreviewsSlotPtr, CreateSlot_MaterialPropertyPreviews());
+	MaterialPropertiesSlot = TDMWidgetSlot<SDMMaterialProperties>(PropertyPreviewsSlotPtr, CreateSlot_MaterialProperties());
 
 	return NewRight;
 }
diff --git a/Engine/Plugins/Experimental/DynamicMaterial/Source/DynamicMaterialEditor/Private/UI/Widgets/Editor/EditorLayouts/SDMMaterialEditor_TopBase.cpp b/Engine/Plugins/Experimental/DynamicMaterial/Source/DynamicMaterialEditor/Private/UI/Widgets/Editor/EditorLayouts/SDMMaterialEditor_TopBase.cpp
--- a/Engine/Plugins/Experimental/DynamicMaterial/Source/DynamicMaterialEditor/Private/UI/Widgets/Editor/EditorLayouts/SDMMaterialEditor_TopBase.cpp
+++ b/Engine/Plugins/Experimental/DynamicMaterial/Source/DynamicMaterialEditor/Private/UI/Widgets/Editor/EditorLayouts/SDMMaterialEditor_TopBase.cpp
@@ -6,7 +6,7 @@
 #include "UI/Widgets/Editor/SDMMaterialComponentEditor.h"
 #include "UI/Widgets/Editor/SDMMaterialGlobalSettingsEditor.h"
 #include "UI/Widgets/Editor/SDMMaterialPreview.h"
-#include "UI/Widgets/Editor/SDMMaterialPropertyPreviews.h"
+#include "UI/Widgets/Editor/SDMMaterialProperties.h"
 #include "UI/Widgets/Editor/SDMMaterialPropertySelector.h"
 #include "UI/Widgets/Editor/SDMMaterialSlotEditor.h"
 #include "Widgets/Layout/SBorder.h"
@@ -244,7 +244,7 @@ TSharedRef<SWidget> SDMMaterialEditor_TopBase::CreateSlot_Bottom_PropertyPreview
 			]
 		];
 
-	MaterialPropertyPreviewsSlot = TDMWidgetSlot<SDMMaterialPropertyPreviews>(PropertyPreviewsSlotPtr, CreateSlot_MaterialPropertyPreviews());
+	MaterialPropertiesSlot = TDMWidgetSlot<SDMMaterialProperties>(PropertyPreviewsSlotPtr, CreateSlot_MaterialProperties());
 
 	return NewBottom;
 }
diff --git a/Engine/Plugins/Experimental/DynamicMaterial/Source/DynamicMaterialEditor/Private/UI/Widgets/Editor/SDMMaterialProperties.cpp b/Engine/Plugins/Experimental/DynamicMaterial/Source/DynamicMaterialEditor/Private/UI/Widgets/Editor/SDMMaterialProperties.cpp
new file mode 100644
--- /dev/null
+++ b/Engine/Plugins/Experimental/DynamicMaterial/Source/DynamicMaterialEditor/Private/UI/Widgets/Editor/SDMMaterialProperties.cpp
@@ -0,0 +1,598 @@
+// Copyright Epic Games, Inc. All Rights Reserved.
+
+#include "UI/Widgets/Editor/SDMMaterialProperties.h"
+
+#include "Components/DMMaterialProperty.h"
+#include "Components/DMMaterialSlot.h"
+#include "Components/MaterialValues/DMMaterialValueFloat1.h"
+#include "CustomDetailsViewArgs.h"
+#include "CustomDetailsViewModule.h"
+#include "DetailLayoutBuilder.h"
+#include "DMWorldSubsystem.h"
+#include "Engine/World.h"
+#include "ICustomDetailsView.h"
+#include "Items/ICustomDetailsViewItem.h"
+#include "Model/DynamicMaterialModelBase.h"
+#include "Model/DynamicMaterialModelDynamic.h"
+#include "Model/DynamicMaterialModelEditorOnlyData.h"
+#include "Styling/StyleColors.h"
+#include "UI/Utils/DMWidgetStatics.h"
+#include "UI/Widgets/Editor/SDMMaterialPropertySelector.h"
+#include "UI/Widgets/SDMMaterialEditor.h"
+#include "UI/Widgets/Visualizers/SDMMaterialComponentPreview.h"
+#include "Utils/DMPrivate.h"
+#include "Widgets/Input/SCheckBox.h"
+#include "Widgets/Layout/SBox.h"
+#include "Widgets/SBoxPanel.h"
+#include "Widgets/SNullWidget.h"
+#include "Widgets/Text/STextBlock.h"
+
+#define LOCTEXT_NAMESPACE "SDMMaterialProperties"
+
+void SDMMaterialProperties::PrivateRegisterAttributes(FSlateAttributeDescriptor::FInitializer&)
+{
+}
+
+void SDMMaterialProperties::Construct(const FArguments& InArgs, const TSharedRef<SDMMaterialEditor>& InEditorWidget)
+{
+	EditorWidgetWeak = InEditorWidget;
+
+	Content = TDMWidgetSlot<SWidget>(SharedThis(this), 0, CreateSlot_Content());
+}
+
+void SDMMaterialProperties::Validate()
+{
+	if (Content.HasBeenInvalidated())
+	{
+		Content << CreateSlot_Content();
+	}
+}
+
+TSharedRef<SWidget> SDMMaterialProperties::CreateSlot_Content()
+{
+	TSharedPtr<SDMMaterialEditor> EditorWidget = EditorWidgetWeak.Pin();
+
+	if (!EditorWidget.IsValid())
+	{
+		return SNullWidget::NullWidget;
+	}
+
+	UDynamicMaterialModel* MaterialModel = EditorWidget->GetMaterialModel();
+
+	if (!MaterialModel)
+	{
+		return SNullWidget::NullWidget;
+	}
+
+	UDynamicMaterialModelEditorOnlyData* EditorOnlyData = UDynamicMaterialModelEditorOnlyData::Get(MaterialModel);
+
+	if (!EditorOnlyData)
+	{
+		return SNullWidget::NullWidget;
+	}
+
+	TSharedRef<SVerticalBox> List = SNew(SVerticalBox);
+
+	// Active properties first.
+	UE::DynamicMaterial::ForEachMaterialPropertyType(
+		[this, &List, EditorOnlyData](EDMMaterialPropertyType InMaterialProperty)
+		{
+			if (UE::DynamicMaterialEditor::Private::IsCustomMaterialProperty(InMaterialProperty))
+			{
+				return EDMIterationResult::Continue;
+			}
+
+			if (UDMMaterialProperty* MaterialProperty = EditorOnlyData->GetMaterialProperty(InMaterialProperty))
+			{
+				UDMMaterialSlot* Slot = EditorOnlyData->GetSlotForEnabledMaterialProperty(InMaterialProperty);
+
+				const bool bIsActive = Slot && MaterialProperty->IsEnabled() && MaterialProperty->IsValidForModel(*EditorOnlyData);
+
+				if (bIsActive)
+				{
+					AddProperty(List, MaterialProperty);
+				}
+			}
+
+			return EDMIterationResult::Continue;
+		}
+	);
+
+	// Now inactive properties
+	UE::DynamicMaterial::ForEachMaterialPropertyType(
+		[this, &List, EditorOnlyData](EDMMaterialPropertyType InMaterialProperty)
+		{
+			if (UE::DynamicMaterialEditor::Private::IsCustomMaterialProperty(InMaterialProperty))
+			{
+				return EDMIterationResult::Continue;
+			}
+
+			if (UDMMaterialProperty* MaterialProperty = EditorOnlyData->GetMaterialProperty(InMaterialProperty))
+			{
+				UDMMaterialSlot* Slot = EditorOnlyData->GetSlotForEnabledMaterialProperty(InMaterialProperty);
+
+				const bool bIsActive = Slot && MaterialProperty->IsEnabled();
+				const bool bIsValid = MaterialProperty->IsValidForModel(*EditorOnlyData);
+
+				if (!bIsActive && bIsValid)
+				{
+					AddProperty(List, MaterialProperty);
+				}
+			}
+
+			return EDMIterationResult::Continue;
+		}
+	);
+
+	// Now invalid properties
+	UE::DynamicMaterial::ForEachMaterialPropertyType(
+		[this, &List, EditorOnlyData](EDMMaterialPropertyType InMaterialProperty)
+		{
+			if (UE::DynamicMaterialEditor::Private::IsCustomMaterialProperty(InMaterialProperty))
+			{
+				return EDMIterationResult::Continue;
+			}
+
+			if (UDMMaterialProperty* MaterialProperty = EditorOnlyData->GetMaterialProperty(InMaterialProperty))
+			{
+				UDMMaterialSlot* Slot = EditorOnlyData->GetSlotForEnabledMaterialProperty(InMaterialProperty);
+
+				const bool bIsActive = Slot && MaterialProperty->IsEnabled();
+				const bool bIsValid = MaterialProperty->IsValidForModel(*EditorOnlyData);
+
+				if (!bIsActive && !bIsValid)
+				{
+					AddProperty(List, MaterialProperty);
+				}
+			}
+
+			return EDMIterationResult::Continue;
+		}
+	);
+
+	return SNew(SBox)
+		.Padding(5.f)
+		[
+			List
+		];
+}
+
+void SDMMaterialProperties::AddProperty(const TSharedRef<SVerticalBox>& InList, UDMMaterialProperty* InProperty)
+{
+	TSharedPtr<SDMMaterialEditor> EditorWidget = EditorWidgetWeak.Pin();
+
+	if (!EditorWidget.IsValid())
+	{
+		return;
+	}
+
+	UDynamicMaterialModel* MaterialModel = EditorWidget->GetMaterialModel();
+
+	if (!MaterialModel)
+	{
+		return;
+	}
+
+	UDynamicMaterialModelEditorOnlyData* EditorOnlyData = UDynamicMaterialModelEditorOnlyData::Get(MaterialModel);
+
+	if (!EditorOnlyData)
+	{
+		return;
+	}
+
+	EDMMaterialPropertyType MaterialProperty = InProperty ? InProperty->GetMaterialProperty() : EDMMaterialPropertyType::None;
+
+	TSharedRef<SBox> PreviewWidgetContainer = SNew(SBox)
+		.WidthOverride(64.f)
+		.HeightOverride(23.f);
+
+	TSharedRef<SWidget> PropertyName = CreateSlot_PropertyName(MaterialProperty);
+
+	TSharedRef<SWidget> Slider = SNullWidget::NullWidget;
+
+	bool bPadSlider = false;
+
+	if (InProperty)
+	{
+		const bool bValidProperty = InProperty->IsValidForModel(*EditorOnlyData);
+		const bool bPropertyEnabled = bValidProperty && InProperty->IsEnabled() && EditorOnlyData->GetSlotForMaterialProperty(MaterialProperty);;
+
+		if (bPropertyEnabled)
+		{
+			MaterialProperty = InProperty->GetMaterialProperty();
+
+			PreviewWidgetContainer = SNew(SBox)
+				.WidthOverride(64.f)
+				.HeightOverride(64.f)
+				[
+					SNew(SBorder)
+					.BorderBackgroundColor(FLinearColor(1, 1, 1, 1.f))
+					.Padding(2.0f)
+					.BorderImage(FAppStyle::GetBrush(UE::DynamicMaterialEditor::Private::EditorDarkBackground))
+					[
+						SNew(SDMMaterialComponentPreview, EditorWidget.ToSharedRef(), InProperty)
+						.PreviewSize(FVector2D(60.f, 60.f))
+					]
+				];
+
+			PreviewWidgetContainer->SetCursor(EMouseCursor::Hand);
+			PreviewWidgetContainer->SetOnMouseButtonUp(FPointerEventHandler::CreateSP(this, &SDMMaterialProperties::OnPropertyClicked, MaterialProperty));
+
+			PropertyName->SetCursor(EMouseCursor::Hand);
+			PropertyName->SetOnMouseButtonUp(FPointerEventHandler::CreateSP(this, &SDMMaterialProperties::OnPropertyClicked, MaterialProperty));
+
+			Slider = CreateGlobalSlider(InProperty);
+		}
+		else if (!bValidProperty)
+		{
+			Slider = SNew(STextBlock)
+				.Font(IDetailLayoutBuilder::GetDetailFont())
+				.ColorAndOpacity(FStyleColors::Notifications)
+				.Text(INVTEXT("Not valid for material type."));
+
+			bPadSlider = true;
+		}
+
+		// If the property is just disabled, leave the slider widget blank.
+	}
+	else
+	{
+		Slider = SNew(STextBlock)
+			.Font(IDetailLayoutBuilder::GetDetailFont())
+			.ColorAndOpacity(FStyleColors::Error)
+			.Text(INVTEXT("Invalid property."));
+
+		bPadSlider = true;
+	}
+
+	InList->AddSlot()
+		.Padding(0.f, 0.f, 0.f, 10.f)
+		.AutoHeight()
+		[
+			SNew(SHorizontalBox)
+
+			+ SHorizontalBox::Slot()
+			.AutoWidth()
+			.VAlign(EVerticalAlignment::VAlign_Top)
+			[
+				PreviewWidgetContainer
+			]
+
+			+ SHorizontalBox::Slot()
+			.FillWidth(1.f)
+			.HAlign(EHorizontalAlignment::HAlign_Left)
+			.VAlign(EVerticalAlignment::VAlign_Fill)
+			.Padding(5.f, 0.f, 0.f, 0.f)
+			[
+				SNew(SVerticalBox)
+
+				+ SVerticalBox::Slot()
+				.AutoHeight()
+				[
+					SNew(SHorizontalBox)
+
+					+ SHorizontalBox::Slot()
+					.AutoWidth()
+					.HAlign(EHorizontalAlignment::HAlign_Left)
+					.VAlign(EVerticalAlignment::VAlign_Center)
+					[
+						CreateSlot_EnabledButton(MaterialProperty)
+					]
+
+					+ SHorizontalBox::Slot()
+					.AutoWidth()
+					.HAlign(EHorizontalAlignment::HAlign_Left)
+					.VAlign(EVerticalAlignment::VAlign_Center)
+					.Padding(5.f, 0.f, 0.f, 0.f)
+					[
+						PropertyName
+					]
+				]
+
+				+ SVerticalBox::Slot()
+				.FillHeight(1.f)
+				.HAlign(EHorizontalAlignment::HAlign_Fill)
+				.VAlign(EVerticalAlignment::VAlign_Center)
+				.Padding(0.f, bPadSlider ? 5.f : 0.f, 0.f, 0.f)
+				[
+					Slider
+				]
+			]
+		];
+}
+
+TSharedRef<SWidget> SDMMaterialProperties::CreateSlot_EnabledButton(EDMMaterialPropertyType InMaterialProperty)
+{
+	const FDMMaterialEditorPage Page = {EDMMaterialEditorMode::EditSlot, InMaterialProperty};
+	const FText Format = LOCTEXT("PropertyEnableFormat", "Toggle the {0} property.\n\nProperty must be valid for the Material Type.");
+	const FText ToolTip = FText::Format(Format, SDMMaterialPropertySelector::GetSelectButtonText(Page, /* Short Name */ false));
+
+	return SNew(SCheckBox)
+		.IsEnabled(this, &SDMMaterialProperties::GetPropertyEnabledEnabled, InMaterialProperty)
+		.IsChecked(this, &SDMMaterialProperties::GetPropertyEnabledState, InMaterialProperty)
+		.OnCheckStateChanged(this, &SDMMaterialProperties::OnPropertyEnabledStateChanged, InMaterialProperty)
+		.ToolTipText(ToolTip);
+}
+
+TSharedRef<SWidget> SDMMaterialProperties::CreateSlot_PropertyName(EDMMaterialPropertyType InMaterialProperty)
+{
+	const FDMMaterialEditorPage Page = {EDMMaterialEditorMode::EditSlot, InMaterialProperty};
+
+	return SNew(STextBlock)
+		.Font(IDetailLayoutBuilder::GetDetailFont())
+		.Text(SDMMaterialPropertySelector::GetSelectButtonText(Page, /* Short Name */ false))
+		.ToolTipText(SDMMaterialPropertySelector::GetSelectButtonText(Page, /* Short Name */ false));
+}
+
+bool SDMMaterialProperties::GetPropertyEnabledEnabled(EDMMaterialPropertyType InMaterialProperty) const
+{
+	TSharedPtr<SDMMaterialEditor> EditorWidget = EditorWidgetWeak.Pin();
+
+	if (!EditorWidget.IsValid())
+	{
+		return false;
+	}
+
+	UDynamicMaterialModel* MaterialModel = EditorWidget->GetMaterialModel();
+
+	if (!MaterialModel)
+	{
+		return false;
+	}
+
+	UDynamicMaterialModelEditorOnlyData* EditorOnlyData = UDynamicMaterialModelEditorOnlyData::Get(MaterialModel);
+
+	if (!EditorOnlyData)
+	{
+		return false;
+	}
+
+	UDMMaterialProperty* Property = EditorOnlyData->GetMaterialProperty(InMaterialProperty);
+
+	if (!Property)
+	{
+		return false;
+	}
+
+	return Property->IsValidForModel(*EditorOnlyData);
+}
+
+ECheckBoxState SDMMaterialProperties::GetPropertyEnabledState(EDMMaterialPropertyType InMaterialProperty) const
+{
+	TSharedPtr<SDMMaterialEditor> EditorWidget = EditorWidgetWeak.Pin();
+
+	if (!EditorWidget.IsValid())
+	{
+		return ECheckBoxState::Unchecked;
+	}
+
+	UDynamicMaterialModel* MaterialModel = EditorWidget->GetMaterialModel();
+
+	if (!MaterialModel)
+	{
+		return ECheckBoxState::Unchecked;
+	}
+
+	UDynamicMaterialModelEditorOnlyData* EditorOnlyData = UDynamicMaterialModelEditorOnlyData::Get(MaterialModel);
+
+	if (!EditorOnlyData)
+	{
+		return ECheckBoxState::Unchecked;
+	}
+
+	UDMMaterialProperty* Property = EditorOnlyData->GetMaterialProperty(InMaterialProperty);
+
+	if (!Property)
+	{
+		return ECheckBoxState::Unchecked;
+	}
+
+	return (Property->IsEnabled() && EditorOnlyData->GetSlotForMaterialProperty(InMaterialProperty))
+		? ECheckBoxState::Checked
+		: ECheckBoxState::Unchecked;
+}
+
+void SDMMaterialProperties::OnPropertyEnabledStateChanged(ECheckBoxState InState, EDMMaterialPropertyType InMaterialProperty)
+{
+	TSharedPtr<SDMMaterialEditor> EditorWidget = EditorWidgetWeak.Pin();
+
+	if (!EditorWidget.IsValid())
+	{
+		return;
+	}
+
+	UDynamicMaterialModel* MaterialModel = EditorWidget->GetMaterialModel();
+
+	if (!MaterialModel)
+	{
+		return;
+	}
+
+	UDynamicMaterialModelEditorOnlyData* EditorOnlyData = UDynamicMaterialModelEditorOnlyData::Get(MaterialModel);
+
+	if (!EditorOnlyData)
+	{
+		return;
+	}
+
+	UDMMaterialProperty* MaterialProperty = EditorOnlyData->GetMaterialProperty(InMaterialProperty);
+
+	if (!MaterialProperty)
+	{
+		return;
+	}
+
+	const bool bEnabled = InState == ECheckBoxState::Checked;
+
+	MaterialProperty->SetEnabled(bEnabled);
+
+	if (bEnabled)
+	{
+		UDMMaterialSlot* Slot = EditorOnlyData->GetSlotForMaterialProperty(InMaterialProperty);
+
+		if (!Slot)
+		{
+			EditorOnlyData->AddSlotForMaterialProperty(InMaterialProperty);
+		}
+	}
+
+	Content.Invalidate();
+
+	// Make sure we go back to the property previews
+	EditorWidget->EditProperties();
+}
+
+FReply SDMMaterialProperties::OnPropertyClicked(const FGeometry& InGeometry, const FPointerEvent& InPointerEvent, 
+	EDMMaterialPropertyType InMaterialProperty)
+{
+	TSharedPtr<SDMMaterialEditor> EditorWidget = EditorWidgetWeak.Pin();
+
+	if (!EditorWidget.IsValid())
+	{
+		return FReply::Handled();
+	}
+
+	UDynamicMaterialModel* MaterialModel = EditorWidget->GetMaterialModel();
+
+	if (!MaterialModel)
+	{
+		return FReply::Handled();
+	}
+
+	UDynamicMaterialModelEditorOnlyData* EditorOnlyData = UDynamicMaterialModelEditorOnlyData::Get(MaterialModel);
+
+	if (!EditorOnlyData)
+	{
+		return FReply::Handled();
+	}
+
+	UDMMaterialSlot* Slot = EditorOnlyData->GetSlotForMaterialProperty(InMaterialProperty);
+
+	if (!Slot)
+	{
+		return FReply::Handled();
+	}
+
+	EditorWidget->SelectProperty(InMaterialProperty);
+
+	return FReply::Handled();
+}
+
+TSharedRef<SWidget> SDMMaterialProperties::CreateGlobalSlider(UDMMaterialProperty* InProperty)
+{
+	if (!InProperty)
+	{
+		return SNullWidget::NullWidget;
+	}
+
+	TSharedPtr<SDMMaterialEditor> EditorWidget = EditorWidgetWeak.Pin();
+
+	if (!EditorWidget.IsValid())
+	{
+		return SNullWidget::NullWidget;
+	}
+
+	UDynamicMaterialModelBase* MaterialModelBase = EditorWidget->GetMaterialModelBase();
+
+	if (!MaterialModelBase)
+	{
+		return SNullWidget::NullWidget;
+	}
+
+	UObject* AlphaObject = InProperty->GetComponent(UDynamicMaterialModelEditorOnlyData::AlphaValueName);
+
+	if (!AlphaObject)
+	{
+		return SNullWidget::NullWidget;
+	}
+
+	if (UDynamicMaterialModelDynamic* MaterialModelDynamic = Cast<UDynamicMaterialModelDynamic>(MaterialModelBase))
+	{
+		AlphaObject = MaterialModelDynamic->GetComponentDynamic(AlphaObject->GetFName());
+
+		if (!AlphaObject)
+		{
+			return SNullWidget::NullWidget;
+		}
+	}
+
+	TSharedPtr<IDetailKeyframeHandler> KeyframeHandler = nullptr;
+
+	if (UWorld* World = MaterialModelBase->GetWorld())
+	{
+		if (const UDMWorldSubsystem* const WorldSubsystem = World->GetSubsystem<UDMWorldSubsystem>())
+		{
+			KeyframeHandler = WorldSubsystem->GetKeyframeHandler();
+		}
+	}
+
+	FCustomDetailsViewArgs Args;
+	Args.KeyframeHandler = KeyframeHandler;
+	Args.bAllowGlobalExtensions = true;
+	Args.bAllowResetToDefault = true;
+	Args.bShowCategories = false;
+
+	TSharedRef<ICustomDetailsView> DetailsView = ICustomDetailsViewModule::Get().CreateCustomDetailsView(Args);
+	FCustomDetailsViewItemId RootId = DetailsView->GetRootItem()->GetItemId();
+
+	FDMPropertyHandle PropertyHandle = FDMWidgetStatics::Get().GetPropertyHandle(this, AlphaObject, UDMMaterialValue::ValueName);
+
+	TSharedRef<ICustomDetailsViewItem> AlphaItem = DetailsView->CreateDetailTreeItem(PropertyHandle.DetailTreeNode.ToSharedRef());
+
+	if (UDMMaterialValue* AlphaValue = Cast<UDMMaterialValue>(AlphaObject))
+	{
+		AlphaItem->SetResetToDefaultOverride(FResetToDefaultOverride::Create(
+			FIsResetToDefaultVisible::CreateUObject(AlphaValue, &UDMMaterialValue::CanResetToDefault),
+			FResetToDefaultHandler::CreateUObject(AlphaValue, &UDMMaterialValue::ResetToDefault)
+		));
+	}
+	else if (UDMMaterialValueDynamic* AlphaValueDynamic = Cast<UDMMaterialValueDynamic>(AlphaObject))
+	{
+		AlphaItem->SetResetToDefaultOverride(FResetToDefaultOverride::Create(
+			FIsResetToDefaultVisible::CreateUObject(AlphaValueDynamic, &UDMMaterialValueDynamic::CanResetToDefault),
+			FResetToDefaultHandler::CreateUObject(AlphaValueDynamic, &UDMMaterialValueDynamic::ResetToDefault)
+		));
+	}
+
+	AlphaItem->MakeWidget(nullptr, SharedThis(this));
+
+	TSharedPtr<SWidget> ValueWidget = AlphaItem->GetWidget(ECustomDetailsViewWidgetType::Value);
+	TSharedPtr<SWidget> ExtensionWidget = AlphaItem->GetWidget(ECustomDetailsViewWidgetType::Extensions);
+
+	const FText PropertyGlobalSliderToolTipFormat = LOCTEXT("PropertyGlobalSliderToolTipFormat", "Change the global {0} value.");
+
+	const FText PropertyGlobalSliderToolTip = FText::Format(
+		PropertyGlobalSliderToolTipFormat,
+		UE::DynamicMaterialEditor::Private::GetMaterialPropertyLongDisplayName(InProperty->GetMaterialProperty())
+	);
+
+	return SNew(SHorizontalBox)
+		.ToolTipText(PropertyGlobalSliderToolTip)
+		
+		+ SHorizontalBox::Slot()
+		.AutoWidth()
+		[
+			SNew(SBox)
+			.HeightOverride(32.f)
+			.VAlign(VAlign_Center)
+			.VAlign(EVerticalAlignment::VAlign_Center)
+			[
+				ValueWidget.IsValid() ? ValueWidget.ToSharedRef() : SNullWidget::NullWidget
+			]
+		]
+
+		+ SHorizontalBox::Slot()
+		.FillWidth(1.f)
+		[
+			SNew(SBox)
+			.HeightOverride(32.f)
+			.HAlign(HAlign_Left)
+			.VAlign(VAlign_Center)
+			.VAlign(EVerticalAlignment::VAlign_Center)
+			[
+				ExtensionWidget.IsValid() ? ExtensionWidget.ToSharedRef() : SNullWidget::NullWidget
+			]
+		];
+}
+
+#undef LOCTEXT_NAMESPACE
diff --git a/Engine/Plugins/Experimental/DynamicMaterial/Source/DynamicMaterialEditor/Private/UI/Widgets/Editor/SDMMaterialProperties.h b/Engine/Plugins/Experimental/DynamicMaterial/Source/DynamicMaterialEditor/Private/UI/Widgets/Editor/SDMMaterialProperties.h
new file mode 100644
--- /dev/null
+++ b/Engine/Plugins/Experimental/DynamicMaterial/Source/DynamicMaterialEditor/Private/UI/Widgets/Editor/SDMMaterialProperties.h
@@ -0,0 +1,54 @@
+// Copyright Epic Games, Inc. All Rights Reserved.
+
+#pragma once
+
+#include "Widgets/SCompoundWidget.h"
+
+#include "Containers/Array.h"
+#include "UI/Utils/DMWidgetSlot.h"
+
+class ICustomDetailsViewItem;
+class SDMMaterialEditor;
+class SVerticalBox;
+class UDMMaterialProperty;
+class UDMMaterialSlot;
+enum class ECheckBoxState : uint8;
+enum class EDMMaterialPropertyType : uint8;
+
+class SDMMaterialProperties : public SCompoundWidget
+{
+	SLATE_DECLARE_WIDGET(SDMMaterialProperties, SCompoundWidget)
+
+	SLATE_BEGIN_ARGS(SDMMaterialProperties) {}
+	SLATE_END_ARGS()
+
+public:
+	virtual ~SDMMaterialProperties() override = default;
+
+	void Construct(const FArguments& InArgs, const TSharedRef<SDMMaterialEditor>& InEditorWidget);
+
+	void Validate();
+
+protected:
+	TWeakPtr<SDMMaterialEditor> EditorWidgetWeak;
+
+	TDMWidgetSlot<SWidget> Content;
+
+	TArray<TSharedRef<ICustomDetailsViewItem>> GlobalItems;
+
+	TSharedRef<SWidget> CreateSlot_Content();
+
+	void AddProperty(const TSharedRef<SVerticalBox>& InList, UDMMaterialProperty* InProperty);
+
+	TSharedRef<SWidget> CreateSlot_EnabledButton(EDMMaterialPropertyType InMaterialProperty);
+
+	TSharedRef<SWidget> CreateSlot_PropertyName(EDMMaterialPropertyType InMaterialProperty);
+
+	bool GetPropertyEnabledEnabled(EDMMaterialPropertyType InMaterialProperty) const;
+	ECheckBoxState GetPropertyEnabledState(EDMMaterialPropertyType InMaterialProperty) const;
+	void OnPropertyEnabledStateChanged(ECheckBoxState InState, EDMMaterialPropertyType InMaterialProperty);
+
+	FReply OnPropertyClicked(const FGeometry& InGeometry, const FPointerEvent& InPointerEvent, EDMMaterialPropertyType InMaterialProperty);
+
+	TSharedRef<SWidget> CreateGlobalSlider(UDMMaterialProperty* InProperty);
+};
diff --git a/Engine/Plugins/Experimental/DynamicMaterial/Source/DynamicMaterialEditor/Private/UI/Widgets/Editor/SDMMaterialPropertyPreviews.cpp b/Engine/Plugins/Experimental/DynamicMaterial/Source/DynamicMaterialEditor/Private/UI/Widgets/Editor/SDMMaterialPropertyPreviews.cpp
deleted file mode 100644
diff --git a/Engine/Plugins/Experimental/DynamicMaterial/Source/DynamicMaterialEditor/Private/UI/Widgets/Editor/SDMMaterialPropertyPreviews.h b/Engine/Plugins/Experimental/DynamicMaterial/Source/DynamicMaterialEditor/Private/UI/Widgets/Editor/SDMMaterialPropertyPreviews.h
deleted file mode 100644
diff --git a/Engine/Plugins/Experimental/DynamicMaterial/Source/DynamicMaterialEditor/Private/UI/Widgets/SDMMaterialEditor.cpp b/Engine/Plugins/Experimental/DynamicMaterial/Source/DynamicMaterialEditor/Private/UI/Widgets/SDMMaterialEditor.cpp
--- a/Engine/Plugins/Experimental/DynamicMaterial/Source/DynamicMaterialEditor/Private/UI/Widgets/SDMMaterialEditor.cpp
+++ b/Engine/Plugins/Experimental/DynamicMaterial/Source/DynamicMaterialEditor/Private/UI/Widgets/SDMMaterialEditor.cpp
@@ -23,7 +23,7 @@
 #include "UI/Widgets/Editor/SDMMaterialComponentEditor.h"
 #include "UI/Widgets/Editor/SDMMaterialGlobalSettingsEditor.h"
 #include "UI/Widgets/Editor/SDMMaterialPreview.h"
-#include "UI/Widgets/Editor/SDMMaterialPropertyPreviews.h"
+#include "UI/Widgets/Editor/SDMMaterialProperties.h"
 #include "UI/Widgets/Editor/SDMMaterialPropertySelector.h"
 #include "UI/Widgets/Editor/SDMMaterialSlotEditor.h"
 #include "UI/Widgets/Editor/SDMStatusBar.h"
@@ -320,7 +320,7 @@ void SDMMaterialEditor::EditComponent(UDMMaterialComponent* InComponent, bool bI
 		SlotEditorSlot.Invalidate();
 		SplitterSlot = nullptr;
 		GlobalSettingsEditorSlot.Invalidate();
-		MaterialPropertyPreviewsSlot.Invalidate();
+		MaterialPropertiesSlot.Invalidate();
 	}
 
 	EditMode = EDMMaterialEditorMode::EditSlot;
@@ -341,7 +341,7 @@ void SDMMaterialEditor::EditGlobalSettings(bool bInForceRefresh)
 		SlotEditorSlot.Invalidate();
 		SplitterSlot = nullptr;
 		ComponentEditorSlot.Invalidate();
-		MaterialPropertyPreviewsSlot.Invalidate();
+		MaterialPropertiesSlot.Invalidate();
 	}
 
 	EditMode = EDMMaterialEditorMode::GlobalSettings;
@@ -370,7 +370,7 @@ void SDMMaterialEditor::EditProperties(bool bInForceRefresh)
 	EditMode = EDMMaterialEditorMode::Properties;
 	SelectedMaterialProperty = EDMMaterialPropertyType::None;
 
-	MaterialPropertyPreviewsSlot.Invalidate();
+	MaterialPropertiesSlot.Invalidate();
 
 	PageHistoryAdd(FDMMaterialEditorPage::Properties);
 }
@@ -745,13 +745,13 @@ void SDMMaterialEditor::ValidateSlots()
 		}
 		else if (EditMode == EDMMaterialEditorMode::Properties)
 		{
-			if (MaterialPropertyPreviewsSlot.HasBeenInvalidated())
+			if (MaterialPropertiesSlot.HasBeenInvalidated())
 			{
-				MaterialPropertyPreviewsSlot << CreateSlot_MaterialPropertyPreviews();
+				MaterialPropertiesSlot << CreateSlot_MaterialProperties();
 			}
 			else
 			{
-				MaterialPropertyPreviewsSlot->Validate();
+				MaterialPropertiesSlot->Validate();
 			}
 		}
 		else
@@ -941,9 +941,9 @@ TSharedRef<SDMMaterialGlobalSettingsEditor> SDMMaterialEditor::CreateSlot_Global
 	return SNew(SDMMaterialGlobalSettingsEditor, SharedThis(this), GetMaterialModelBase());
 }
 
-TSharedRef<SDMMaterialPropertyPreviews> SDMMaterialEditor::CreateSlot_MaterialPropertyPreviews()
+TSharedRef<SDMMaterialProperties> SDMMaterialEditor::CreateSlot_MaterialProperties()
 {
-	return SNew(SDMMaterialPropertyPreviews, SharedThis(this));
+	return SNew(SDMMaterialProperties, SharedThis(this));
 }
 
 TSharedRef<SDMMaterialPreview> SDMMaterialEditor::CreateSlot_Preview()
diff --git a/Engine/Plugins/Experimental/DynamicMaterial/Source/DynamicMaterialEditor/Private/UI/Widgets/SDMMaterialEditor.h b/Engine/Plugins/Experimental/DynamicMaterial/Source/DynamicMaterialEditor/Private/UI/Widgets/SDMMaterialEditor.h
--- a/Engine/Plugins/Experimental/DynamicMaterial/Source/DynamicMaterialEditor/Private/UI/Widgets/SDMMaterialEditor.h
+++ b/Engine/Plugins/Experimental/DynamicMaterial/Source/DynamicMaterialEditor/Private/UI/Widgets/SDMMaterialEditor.h
@@ -19,7 +19,7 @@ class SDMMaterialComponentEditor;
 class SDMMaterialDesigner;
 class SDMMaterialGlobalSettingsEditor;
 class SDMMaterialPreview;
-class SDMMaterialPropertyPreviews;
+class SDMMaterialProperties;
 class SDMMaterialPropertySelector;
 class SDMMaterialSlotEditor;
 class SDMStatusBar;
@@ -163,7 +163,7 @@ protected:
 	TDMWidgetSlot<SDMMaterialPreview> MaterialPreviewSlot;
 	TDMWidgetSlot<SDMMaterialPropertySelector> PropertySelectorSlot;
 	TDMWidgetSlot<SDMMaterialGlobalSettingsEditor> GlobalSettingsEditorSlot;
-	TDMWidgetSlot<SDMMaterialPropertyPreviews> MaterialPropertyPreviewsSlot;
+	TDMWidgetSlot<SDMMaterialProperties> MaterialPropertiesSlot;
 	FSlotBase* SplitterSlot;
 	TDMWidgetSlot<SDMMaterialSlotEditor> SlotEditorSlot;
 	TDMWidgetSlot<SDMMaterialComponentEditor> ComponentEditorSlot;
@@ -227,7 +227,7 @@ protected:
 
 	TSharedRef<SDMMaterialGlobalSettingsEditor> CreateSlot_GlobalSettingsEditor();
 
-	TSharedRef<SDMMaterialPropertyPreviews> CreateSlot_MaterialPropertyPreviews();
+	TSharedRef<SDMMaterialProperties> CreateSlot_MaterialProperties();
 
 	TSharedRef<SDMMaterialPreview> CreateSlot_Preview();
 
