Material Designer
- A red X now appears over disabled stage previews.

diff --git a/Engine/Plugins/Experimental/DynamicMaterial/Source/DynamicMaterialEditor/Private/UI/Widgets/Editor/SlotEditor/SDMMaterialStage.cpp b/Engine/Plugins/Experimental/DynamicMaterial/Source/DynamicMaterialEditor/Private/UI/Widgets/Editor/SlotEditor/SDMMaterialStage.cpp
--- a/Engine/Plugins/Experimental/DynamicMaterial/Source/DynamicMaterialEditor/Private/UI/Widgets/Editor/SlotEditor/SDMMaterialStage.cpp
+++ b/Engine/Plugins/Experimental/DynamicMaterial/Source/DynamicMaterialEditor/Private/UI/Widgets/Editor/SlotEditor/SDMMaterialStage.cpp
@@ -16,14 +16,17 @@
 #include "Engine/Texture.h"
 #include "Framework/Application/SlateApplication.h"
 #include "SAssetDropTarget.h"
+#include "Styling/StyleColors.h"
 #include "UI/Menus/DMMaterialStageMenus.h"
 #include "UI/Widgets/Editor/SDMMaterialComponentEditor.h"
 #include "UI/Widgets/Editor/SlotEditor/SDMMaterialSlotLayerItem.h"
 #include "UI/Widgets/Editor/SlotEditor/SDMMaterialSlotLayerView.h"
 #include "UI/Widgets/SDMMaterialEditor.h"
 #include "UI/Widgets/Visualizers/SDMMaterialComponentPreview.h"
+#include "Widgets/Images/SImage.h"
 #include "Widgets/Layout/SBox.h"
 #include "Widgets/SNullWidget.h"
+#include "Widgets/SOverlay.h"
 #include "Widgets/SToolTip.h"
 #include "Widgets/Text/STextBlock.h"
 
@@ -91,8 +94,21 @@ void SDMMaterialStage::Construct(const FArguments& InArgs, const TSharedRef<SDMM
 					.Clipping(EWidgetClipping::ClipToBounds)
 					.BorderBackgroundColor(FLinearColor::Transparent)
 					[
-						SNew(SDMMaterialComponentPreview, EditorWidget.ToSharedRef(), InStage)
-						.PreviewSize(FVector2D(40.f))
+						SNew(SOverlay)
+						+ SOverlay::Slot()
+						[
+							SNew(SDMMaterialComponentPreview, EditorWidget.ToSharedRef(), InStage)
+							.PreviewSize(FVector2D(40.f))
+						]
+						+ SOverlay::Slot()
+						.Padding(5.f)
+						[
+							SNew(SImage)
+							.Image(FAppStyle::Get().GetBrush("Icons.X"))
+							.DesiredSizeOverride(FVector2D(30.f))
+							.ColorAndOpacity(FStyleColors::AccentRed)
+							.Visibility(this, &SDMMaterialStage::GetDisabledOverlayVisibility)
+						]
 					]
 				]
 				+ SOverlay::Slot()
@@ -232,6 +248,18 @@ const FSlateBrush* SDMMaterialStage::GetBorderBrush() const
 	return FDynamicMaterialEditorStyle::GetBrush(*BrushName);
 }
 
+EVisibility SDMMaterialStage::GetDisabledOverlayVisibility() const
+{
+	if (UDMMaterialStage* Stage = GetStage())
+	{
+		return Stage->IsEnabled()
+			? EVisibility::Collapsed
+			: EVisibility::HitTestInvisible;
+	}
+
+	return EVisibility::Collapsed;
+}
+
 bool SDMMaterialStage::OnAssetDraggedOver(TArrayView<FAssetData> InAssets)
 {
 	UDMMaterialStage* Stage = GetStage();
diff --git a/Engine/Plugins/Experimental/DynamicMaterial/Source/DynamicMaterialEditor/Private/UI/Widgets/Editor/SlotEditor/SDMMaterialStage.h b/Engine/Plugins/Experimental/DynamicMaterial/Source/DynamicMaterialEditor/Private/UI/Widgets/Editor/SlotEditor/SDMMaterialStage.h
--- a/Engine/Plugins/Experimental/DynamicMaterial/Source/DynamicMaterialEditor/Private/UI/Widgets/Editor/SlotEditor/SDMMaterialStage.h
+++ b/Engine/Plugins/Experimental/DynamicMaterial/Source/DynamicMaterialEditor/Private/UI/Widgets/Editor/SlotEditor/SDMMaterialStage.h
@@ -41,6 +41,8 @@ protected:
 
 	const FSlateBrush* GetBorderBrush() const;
 
+	EVisibility GetDisabledOverlayVisibility() const;
+
 	/** Drag and drop. */
 	bool OnAssetDraggedOver(TArrayView<FAssetData> InAssets);
 
