Material Designer
- Converted the style class for the editor to the subclass format.

diff --git a/Engine/Plugins/Experimental/DynamicMaterial/Source/DynamicMaterialEditor/Private/DynamicMaterialEditorModule.cpp b/Engine/Plugins/Experimental/DynamicMaterial/Source/DynamicMaterialEditor/Private/DynamicMaterialEditorModule.cpp
--- a/Engine/Plugins/Experimental/DynamicMaterial/Source/DynamicMaterialEditor/Private/DynamicMaterialEditorModule.cpp
+++ b/Engine/Plugins/Experimental/DynamicMaterial/Source/DynamicMaterialEditor/Private/DynamicMaterialEditorModule.cpp
@@ -211,7 +211,7 @@ FDynamicMaterialEditorModule::FDynamicMaterialEditorModule()
 
 void FDynamicMaterialEditorModule::StartupModule()
 {
-	FDynamicMaterialEditorStyle::Initialize();
+	FDynamicMaterialEditorStyle::Get();
 	FDynamicMaterialEditorCommands::Register();
 	FDMContentBrowserIntegration::Integrate();
 	MapCommands();
@@ -263,7 +263,6 @@ void FDynamicMaterialEditorModule::StartupModule()
 void FDynamicMaterialEditorModule::ShutdownModule()
 {
 	FDynamicMaterialEditorCommands::Unregister();
-	FDynamicMaterialEditorStyle::Shutdown();
 	FDMContentBrowserIntegration::Disintegrate();
 
 	if (FDynamicMaterialModule::AreUObjectsSafe() && FModuleManager::Get().IsModuleLoaded("PropertyEditor"))
diff --git a/Engine/Plugins/Experimental/DynamicMaterial/Source/DynamicMaterialEditor/Private/DynamicMaterialEditorStyle.cpp b/Engine/Plugins/Experimental/DynamicMaterial/Source/DynamicMaterialEditor/Private/DynamicMaterialEditorStyle.cpp
--- a/Engine/Plugins/Experimental/DynamicMaterial/Source/DynamicMaterialEditor/Private/DynamicMaterialEditorStyle.cpp
+++ b/Engine/Plugins/Experimental/DynamicMaterial/Source/DynamicMaterialEditor/Private/DynamicMaterialEditorStyle.cpp
@@ -2,61 +2,12 @@
 
 #include "DynamicMaterialEditorStyle.h"
 
-#include "Brushes/SlateImageBrush.h"
 #include "DetailLayoutBuilder.h"
-#include "DynamicMaterialEditorModule.h"
-#include "Engine/Texture2D.h"
+#include "DMDefs.h"
 #include "Interfaces/IPluginManager.h"
-#include "Misc/Paths.h"
-#include "Styling/SlateStyleRegistry.h"
+#include "Styling/SlateStyleMacros.h"
 #include "Styling/StyleColors.h"
 
-TSharedPtr<FSlateStyleSet> FDynamicMaterialEditorStyle::StyleInstance = nullptr;
-
-void FDynamicMaterialEditorStyle::Initialize()
-{
-	if (!StyleInstance.IsValid())
-	{
-		StyleInstance = Create();
-		FSlateStyleRegistry::RegisterSlateStyle(*StyleInstance);
-	}
-}
-
-void FDynamicMaterialEditorStyle::Shutdown()
-{
-	if (StyleInstance.IsValid())
-	{
-		FSlateStyleRegistry::UnRegisterSlateStyle(*StyleInstance);
-		ensure(StyleInstance.IsUnique());
-		StyleInstance.Reset();
-	}
-}
-
-FName FDynamicMaterialEditorStyle::GetStyleSetName()
-{
-	static FName StyleSetName(TEXT("DynamicMaterialEditorStyle"));
-	return StyleSetName;
-}
-
-const ISlateStyle& FDynamicMaterialEditorStyle::Get()
-{
-	if (!StyleInstance.IsValid())
-	{
-		Initialize();
-	}
-	return *StyleInstance;
-}
-
-FLinearColor FDynamicMaterialEditorStyle::GetColor(const FName& InName)
-{
-	return FDynamicMaterialEditorStyle::Get().GetColor(InName);
-}
-
-const FSlateBrush* FDynamicMaterialEditorStyle::GetBrush(const FName& InName)
-{
-	return FDynamicMaterialEditorStyle::Get().GetBrush(InName);
-}
-
 namespace UE::DynamicMaterialEditor::Private
 {
 	const TMap<EDMValueType, FName> TypeIcons = {
@@ -71,9 +22,27 @@ namespace UE::DynamicMaterialEditor::Private
 		{EDMValueType::VT_Float_Any,   TEXT("Icons.Type.Float_Any")},
 		{EDMValueType::VT_Texture,     TEXT("Icons.Type.Texture")},
 	};
+
+	FLinearColor ReplaceColorAlpha(const FLinearColor& InColor, const float InNewAlpha)
+	{
+		FLinearColor OutColor(InColor);
+		OutColor.A = InNewAlpha;
+		return OutColor;
+	}
+}
+
+FDynamicMaterialEditorStyle::FDynamicMaterialEditorStyle()
+	: FSlateStyleSet(TEXT("DynamicMaterialEditor"))
+{
+	SetupGenericStyles();
+	SetupStageStyles();
+	SetupLayerViewStyles();
+	SetupLayerViewItemHandleStyles();
+	SetupEffectsViewStyles();
+	SetupTextStyles();
 }
 
-FName FDynamicMaterialEditorStyle::GetBrushNameForType(EDMValueType InType)
+FName FDynamicMaterialEditorStyle::GetBrushNameForType(EDMValueType InType) const
 {
 	if (UE::DynamicMaterialEditor::Private::TypeIcons.Contains(InType))
 	{
@@ -86,42 +55,14 @@ FName FDynamicMaterialEditorStyle::GetBrushNameForType(EDMValueType InType)
 	return NoType;
 }
 
-
-#define IMAGE_BRUSH(RelativePath, ... ) FSlateImageBrush(Style->RootToContentDir(RelativePath, TEXT(".png")), __VA_ARGS__)
-#define BOX_BRUSH(RelativePath, ... ) FSlateBoxBrush(Style->RootToContentDir(RelativePath, TEXT(".png")), __VA_ARGS__)
-#define BORDER_BRUSH(RelativePath, ... ) FSlateBorderBrush(Style->RootToContentDir(RelativePath, TEXT(".png")), __VA_ARGS__)
-
-#define IMAGE_BRUSH_SVG(RelativePath, ...) FSlateVectorImageBrush(Style->RootToContentDir(RelativePath, TEXT(".svg")), __VA_ARGS__)
-#define BOX_BRUSH_SVG(RelativePath, ...) FSlateVectorBoxBrush(Style->RootToContentDir(RelativePath, TEXT(".svg")), __VA_ARGS__)
-#define BORDER_BRUSH_SVG(RelativePath, ...) FSlateVectorBorderBrush(Style->RootToContentDir(RelativePath, TEXT(".svg")), __VA_ARGS__)
-
-#define DEFAULT_FONT(...) FCoreStyle::GetDefaultFontStyle(__VA_ARGS__)
-
-#define IMAGE_PLUGIN_BRUSH(RelativePath, ...) FSlateImageBrush(FDynamicMaterialEditorStyle::InContent(RelativePath, ".png"), __VA_ARGS__)
-#define IMAGE_PLUGIN_BRUSH_SVG(RelativePath, ...) FSlateVectorImageBrush(FDynamicMaterialEditorStyle::InContent(RelativePath, ".svg"), __VA_ARGS__)
-
-const FVector2D Icon8x8(8.f, 8.f);
-const FVector2D Icon12x12(12.f, 12.f);
-const FVector2D Icon16x16(16.f, 16.f);
-const FVector2D Icon20x20(20.f, 20.f);
-const FVector2D Icon24x24(24.f, 24.f);
-const FVector2D Icon32x32(32.f, 32.f);
-const FVector2D Icon40x40(40.f, 40.f);
-
-FLinearColor ReplaceColorAlpha(const FLinearColor& InColor, const float InNewAlpha)
-{
-	FLinearColor OutColor(InColor);
-	OutColor.A = InNewAlpha;
-	return OutColor;
-}
-TSharedRef<FSlateStyleSet> FDynamicMaterialEditorStyle::Create()
+void FDynamicMaterialEditorStyle::SetupGenericStyles()
 {
-	TSharedRef<FSlateStyleSet> Style = MakeShared<FSlateStyleSet>(TEXT("DynamicMaterial"));
+	using namespace UE::DynamicMaterialEditor::Private;
 
 	TSharedPtr<IPlugin> Plugin = IPluginManager::Get().FindPlugin(TEXT("DynamicMaterial"));
 	check(Plugin.IsValid());
 
-	Style->SetContentRoot(FPaths::Combine(Plugin->GetBaseDir(), TEXT("Resources")));
+	SetContentRoot(FPaths::Combine(Plugin->GetBaseDir(), TEXT("Resources")));
 
 	////////////////////////////////////////////////////////////////////////////////////////////////////
 	// Color Styles
@@ -133,104 +74,90 @@ TSharedRef<FSlateStyleSet> FDynamicMaterialEditorStyle::Create()
 	const FLinearColor SelectHoverColor = FStyleColors::Select.GetSpecifiedColor();
 	const FLinearColor SelectPressColor = EngineSelectPressColor;
 
-	Style->Set("Color.Select", SelectColor);
-	Style->Set("Color.Select.Hover", SelectHoverColor);
-	Style->Set("Color.Select.Press", SelectPressColor);
+	Set("Color.Select", SelectColor);
+	Set("Color.Select.Hover", SelectHoverColor);
+	Set("Color.Select.Press", SelectPressColor);
 
 	////////////////////////////////////////////////////////////////////////////////////////////////////
 	// Brush Styles
-	Style->Set("Icons.Menu.Dropdown", new IMAGE_BRUSH_SVG("Icons/MenuDropdown", Icon16x16));
-
-	Style->Set("Icons.Type.None", new IMAGE_BRUSH("Icons/ValueTypes/None", Icon12x12));
-	Style->Set("Icons.Type.Bool", new IMAGE_BRUSH("Icons/ValueTypes/Bool", Icon12x12));
-	Style->Set("Icons.Type.Float1", new IMAGE_BRUSH("Icons/ValueTypes/Float1", Icon12x12));
-	Style->Set("Icons.Type.Float2", new IMAGE_BRUSH("Icons/ValueTypes/Float2", Icon12x12));
-	Style->Set("Icons.Type.Float3_RPY", new IMAGE_BRUSH("Icons/ValueTypes/Float3_RPY", Icon12x12));
-	Style->Set("Icons.Type.Float3_RGB", new IMAGE_BRUSH("Icons/ValueTypes/Float3_RGB", Icon12x12));
-	Style->Set("Icons.Type.Float3_XYZ", new IMAGE_BRUSH("Icons/ValueTypes/Float3_XYZ", Icon12x12));
-	Style->Set("Icons.Type.Float4_RGBA", new IMAGE_BRUSH("Icons/ValueTypes/Float4_RGBA", Icon12x12));
-	Style->Set("Icons.Type.Float_Any", new IMAGE_BRUSH("Icons/ValueTypes/Float_Any", Icon12x12));
-	Style->Set("Icons.Type.Texture", new IMAGE_BRUSH("Icons/ValueTypes/Texture", Icon12x12));
-
-	Style->Set("Icons.Material.DefaultLit", new IMAGE_BRUSH("Icons/EditorIcons/MaterialTypeDefaultLit", Icon32x32));
-	Style->Set("Icons.Material.Unlit", new IMAGE_BRUSH("Icons/EditorIcons/MaterialTypeUnlit", Icon32x32));
-
-	Style->Set("Icons.Lock", new IMAGE_BRUSH_SVG("Icons/EditorIcons/Lock", Icon16x16));
-	Style->Set("Icons.Unlock", new IMAGE_BRUSH_SVG("Icons/EditorIcons/Unlock", Icon16x16));
-
-	Style->Set("Icons.Remove", new IMAGE_BRUSH("Icons/EditorIcons/Remove_16px", Icon16x16));
-
-	Style->Set("Icons.Normalize", new IMAGE_BRUSH_SVG("Icons/EditorIcons/Normalize", Icon16x16));
-
-	Style->Set("Icons.Stage.EnabledButton", new IMAGE_BRUSH("Icons/EditorIcons/WhiteBall", Icon8x8));
-	Style->Set("Icons.Stage.BaseToggle", new IMAGE_BRUSH("Icons/EditorIcons/BaseToggle_16x", Icon16x16));
-	Style->Set("Icons.Stage.MaskToggle", new IMAGE_BRUSH("Icons/EditorIcons/MaskToggle_16x", Icon16x16));
-	Style->Set("Icons.Stage.Enabled", new IMAGE_BRUSH_SVG("Icons/EditorIcons/Enable", Icon24x24));
-	Style->Set("Icons.Stage.Disabled", new IMAGE_BRUSH_SVG("Icons/EditorIcons/Disable", Icon24x24));
-
-	Style->Set("Icons.Stage.ChainLinked", new IMAGE_BRUSH_SVG("Icons/EditorIcons/ChainLinked", Icon16x16));
-	Style->Set("Icons.Stage.ChainUnlinked", new IMAGE_BRUSH_SVG("Icons/EditorIcons/ChainUnlinked", Icon16x16));
-	Style->Set("Icons.Stage.ChainLinked.Horizontal", new IMAGE_BRUSH_SVG("Icons/EditorIcons/ChainLinked_Horizontal", Icon24x24));
-	Style->Set("Icons.Stage.ChainUnlinked.Horizontal", new IMAGE_BRUSH_SVG("Icons/EditorIcons/ChainUnlinked_Horizontal", Icon24x24));
-	Style->Set("Icons.Stage.ChainLinked.Vertical", new IMAGE_BRUSH_SVG("Icons/EditorIcons/ChainLinked_Vertical", Icon24x24));
-	Style->Set("Icons.Stage.ChainUnlinked.Vertical", new IMAGE_BRUSH_SVG("Icons/EditorIcons/ChainUnlinked_Vertical", Icon24x24));
-
-	Style->Set("ImageBorder", new FSlateRoundedBoxBrush(
-		FLinearColor::Transparent, 10.0f, 
+	Set("Icons.Menu.Dropdown", new IMAGE_BRUSH_SVG("Icons/MenuDropdown", CoreStyleConstants::Icon16x16));
+
+	Set("Icons.Type.None", new IMAGE_BRUSH("Icons/ValueTypes/None", CoreStyleConstants::Icon12x12));
+	Set("Icons.Type.Bool", new IMAGE_BRUSH("Icons/ValueTypes/Bool", CoreStyleConstants::Icon12x12));
+	Set("Icons.Type.Float1", new IMAGE_BRUSH("Icons/ValueTypes/Float1", CoreStyleConstants::Icon12x12));
+	Set("Icons.Type.Float2", new IMAGE_BRUSH("Icons/ValueTypes/Float2", CoreStyleConstants::Icon12x12));
+	Set("Icons.Type.Float3_RPY", new IMAGE_BRUSH("Icons/ValueTypes/Float3_RPY", CoreStyleConstants::Icon12x12));
+	Set("Icons.Type.Float3_RGB", new IMAGE_BRUSH("Icons/ValueTypes/Float3_RGB", CoreStyleConstants::Icon12x12));
+	Set("Icons.Type.Float3_XYZ", new IMAGE_BRUSH("Icons/ValueTypes/Float3_XYZ", CoreStyleConstants::Icon12x12));
+	Set("Icons.Type.Float4_RGBA", new IMAGE_BRUSH("Icons/ValueTypes/Float4_RGBA", CoreStyleConstants::Icon12x12));
+	Set("Icons.Type.Float_Any", new IMAGE_BRUSH("Icons/ValueTypes/Float_Any", CoreStyleConstants::Icon12x12));
+	Set("Icons.Type.Texture", new IMAGE_BRUSH("Icons/ValueTypes/Texture", CoreStyleConstants::Icon12x12));
+
+	Set("Icons.Material.DefaultLit", new IMAGE_BRUSH("Icons/EditorIcons/MaterialTypeDefaultLit", CoreStyleConstants::Icon32x32));
+	Set("Icons.Material.Unlit", new IMAGE_BRUSH("Icons/EditorIcons/MaterialTypeUnlit", CoreStyleConstants::Icon32x32));
+
+	Set("Icons.Lock", new IMAGE_BRUSH_SVG("Icons/EditorIcons/Lock", CoreStyleConstants::Icon16x16));
+	Set("Icons.Unlock", new IMAGE_BRUSH_SVG("Icons/EditorIcons/Unlock", CoreStyleConstants::Icon16x16));
+
+	Set("Icons.Remove", new IMAGE_BRUSH("Icons/EditorIcons/Remove_16px", CoreStyleConstants::Icon16x16));
+
+	Set("Icons.Normalize", new IMAGE_BRUSH_SVG("Icons/EditorIcons/Normalize", CoreStyleConstants::Icon16x16));
+
+	Set("Icons.Stage.EnabledButton", new IMAGE_BRUSH("Icons/EditorIcons/WhiteBall", CoreStyleConstants::Icon8x8));
+	Set("Icons.Stage.BaseToggle", new IMAGE_BRUSH("Icons/EditorIcons/BaseToggle_16x", CoreStyleConstants::Icon16x16));
+	Set("Icons.Stage.MaskToggle", new IMAGE_BRUSH("Icons/EditorIcons/MaskToggle_16x", CoreStyleConstants::Icon16x16));
+	Set("Icons.Stage.Enabled", new IMAGE_BRUSH_SVG("Icons/EditorIcons/Enable", CoreStyleConstants::Icon24x24));
+	Set("Icons.Stage.Disabled", new IMAGE_BRUSH_SVG("Icons/EditorIcons/Disable", CoreStyleConstants::Icon24x24));
+
+	Set("Icons.Stage.ChainLinked", new IMAGE_BRUSH_SVG("Icons/EditorIcons/ChainLinked", CoreStyleConstants::Icon16x16));
+	Set("Icons.Stage.ChainUnlinked", new IMAGE_BRUSH_SVG("Icons/EditorIcons/ChainUnlinked", CoreStyleConstants::Icon16x16));
+	Set("Icons.Stage.ChainLinked.Horizontal", new IMAGE_BRUSH_SVG("Icons/EditorIcons/ChainLinked_Horizontal", CoreStyleConstants::Icon24x24));
+	Set("Icons.Stage.ChainUnlinked.Horizontal", new IMAGE_BRUSH_SVG("Icons/EditorIcons/ChainUnlinked_Horizontal", CoreStyleConstants::Icon24x24));
+	Set("Icons.Stage.ChainLinked.Vertical", new IMAGE_BRUSH_SVG("Icons/EditorIcons/ChainLinked_Vertical", CoreStyleConstants::Icon24x24));
+	Set("Icons.Stage.ChainUnlinked.Vertical", new IMAGE_BRUSH_SVG("Icons/EditorIcons/ChainUnlinked_Vertical", CoreStyleConstants::Icon24x24));
+
+	Set("ImageBorder", new FSlateRoundedBoxBrush(
+		FLinearColor::Transparent, 10.0f,
 		FStyleColors::Panel.GetSpecifiedColor(), 2.0f));
 
-	Style->Set("Border.SinglePixel", new BORDER_BRUSH(TEXT("Images/Borders/Border_SinglePixel"), FMargin(1.0f / 4.0f)));
-	Style->Set("Border.LeftTopRight", new BORDER_BRUSH(TEXT("Images/Borders/Border_LeftTopRight"), FMargin(1.0f / 4.0f, 1.0f / 2.0f)));
-	Style->Set("Border.LeftBottomRight", new BORDER_BRUSH(TEXT("Images/Borders/Border_LeftBottomRight"), FMargin(1.0f / 4.0f, 1.0f / 2.0f)));
-	Style->Set("Border.TopLeftBottom", new BORDER_BRUSH(TEXT("Images/Borders/Border_TopLeftBottom"), FMargin(1.0f / 2.0f, 1.0f / 4.0f)));
-	Style->Set("Border.TopRightBottom", new BORDER_BRUSH(TEXT("Images/Borders/Border_TopRightBottom"), FMargin(1.0f / 2.0f, 1.0f / 4.0f)));
-	Style->Set("Border.Top", new BORDER_BRUSH(TEXT("Images/Borders/Border_Top"), FMargin(0.0f, 1.0f / 2.0f)));
-	Style->Set("Border.Bottom", new BORDER_BRUSH(TEXT("Images/Borders/Border_Bottom"), FMargin(0.0f, 1.0f / 2.0f)));
-	Style->Set("Border.Left", new BORDER_BRUSH(TEXT("Images/Borders/Border_Left"), FMargin(1.0f / 2.0f, 0.0f)));
-	Style->Set("Border.Right", new BORDER_BRUSH(TEXT("Images/Borders/Border_Right"), FMargin(1.0f / 2.0f, 0.0f)));
+	Set("Border.SinglePixel", new BORDER_BRUSH(TEXT("Images/Borders/Border_SinglePixel"), FMargin(1.0f / 4.0f)));
+	Set("Border.LeftTopRight", new BORDER_BRUSH(TEXT("Images/Borders/Border_LeftTopRight"), FMargin(1.0f / 4.0f, 1.0f / 2.0f)));
+	Set("Border.LeftBottomRight", new BORDER_BRUSH(TEXT("Images/Borders/Border_LeftBottomRight"), FMargin(1.0f / 4.0f, 1.0f / 2.0f)));
+	Set("Border.TopLeftBottom", new BORDER_BRUSH(TEXT("Images/Borders/Border_TopLeftBottom"), FMargin(1.0f / 2.0f, 1.0f / 4.0f)));
+	Set("Border.TopRightBottom", new BORDER_BRUSH(TEXT("Images/Borders/Border_TopRightBottom"), FMargin(1.0f / 2.0f, 1.0f / 4.0f)));
+	Set("Border.Top", new BORDER_BRUSH(TEXT("Images/Borders/Border_Top"), FMargin(0.0f, 1.0f / 2.0f)));
+	Set("Border.Bottom", new BORDER_BRUSH(TEXT("Images/Borders/Border_Bottom"), FMargin(0.0f, 1.0f / 2.0f)));
+	Set("Border.Left", new BORDER_BRUSH(TEXT("Images/Borders/Border_Left"), FMargin(1.0f / 2.0f, 0.0f)));
+	Set("Border.Right", new BORDER_BRUSH(TEXT("Images/Borders/Border_Right"), FMargin(1.0f / 2.0f, 0.0f)));
 
 	////////////////////////////////////////////////////////////////////////////////////////////////////
 	// Button Styles
-	Style->Set("HoverHintOnly", FButtonStyle()
+	Set("HoverHintOnly", FButtonStyle()
 		.SetNormal(FSlateNoResource())
 		.SetHovered(FSlateRoundedBoxBrush(FLinearColor(1, 1, 1, 0.15f), 4.0f))
 		.SetPressed(FSlateRoundedBoxBrush(FLinearColor(1, 1, 1, 0.25f), 4.0f))
 		.SetNormalPadding(FMargin(0, 0, 0, 1))
 		.SetPressedPadding(FMargin(0, 1, 0, 0)));
 
-	Style->Set("HoverHintOnly.Bordered", FButtonStyle()
+	Set("HoverHintOnly.Bordered", FButtonStyle()
 		.SetNormal(FSlateRoundedBoxBrush(FLinearColor::Transparent, 4.0f, FLinearColor(1, 1, 1, 0.25f), 1.0f))
 		.SetHovered(FSlateRoundedBoxBrush(FLinearColor(1, 1, 1, 0.15f), 4.0f, FLinearColor(1, 1, 1, 0.4f), 1.0f))
 		.SetPressed(FSlateRoundedBoxBrush(FLinearColor(1, 1, 1, 0.25f), 4.0f, FLinearColor(1, 1, 1, 0.5f), 1.0f))
 		.SetNormalPadding(FMargin(0, 0, 0, 1))
 		.SetPressedPadding(FMargin(0, 1, 0, 0)));
 
-	Style->Set("HoverHintOnly.Bordered.Dark", FButtonStyle()
+	Set("HoverHintOnly.Bordered.Dark", FButtonStyle()
 		.SetNormal(FSlateRoundedBoxBrush(FLinearColor::Transparent, 4.0f, FStyleColors::InputOutline.GetSpecifiedColor(), 1.0f))
 		.SetHovered(FSlateRoundedBoxBrush(FLinearColor(1, 1, 1, 0.15f), 4.0f, FLinearColor(1, 1, 1, 0.4f), 1.0f))
 		.SetPressed(FSlateRoundedBoxBrush(FLinearColor(1, 1, 1, 0.25f), 4.0f, FLinearColor(1, 1, 1, 0.5f), 1.0f))
 		.SetNormalPadding(FMargin(0, 0, 0, 1))
 		.SetPressedPadding(FMargin(0, 1, 0, 0)));
-
-	SetupStageStyles(Style);
-	SetupLayerViewStyles(Style);
-	SetupLayerViewItemHandleStyles(Style);
-	SetupEffectsViewStyles(Style);
-	SetupTextStyles(Style);
-
-	//////////////////////////////////////////////////////
-	/// Editable TextBox Style
-	FEditableTextBoxStyle InlineEditableTextBoxStyle;
-	InlineEditableTextBoxStyle.SetPadding(FMargin(0));
-	InlineEditableTextBoxStyle.SetBackgroundColor(FSlateColor(FLinearColor::Transparent));
-
-	Style->Set("InlineEditableTextBoxStyle", InlineEditableTextBoxStyle);
-
-	return Style;
 }
 
-void FDynamicMaterialEditorStyle::SetupStageStyles(const TSharedRef<FSlateStyleSet>& Style)
+void FDynamicMaterialEditorStyle::SetupStageStyles()
 {
+	using namespace UE::DynamicMaterialEditor::Private;
+
 	constexpr float StageCornerRadius = 6.0f;
 	constexpr float StageBorderWidth = 2.0f;
 	constexpr float NonHoverAlpha = 1.f;
@@ -241,55 +168,43 @@ void FDynamicMaterialEditorStyle::SetupStageStyles(const TSharedRef<FSlateStyleS
 	const FLinearColor DisabledColor = FStyleColors::AccentRed.GetSpecifiedColor() * FLinearColor(0.5f, 0.5f, 0.5f, 1.f);
 	const FLinearColor DisabledSelectedColor = FStyleColors::AccentRed.GetSpecifiedColor();
 
-	Style->Set("Stage.Inactive", new FSlateRoundedBoxBrush(
+	Set("Stage.Inactive", new FSlateRoundedBoxBrush(
 		FLinearColor::Transparent, StageCornerRadius,
 		FStyleColors::Panel.GetSpecifiedColor(), StageBorderWidth));
 
-	Style->Set("Stage.Enabled", new FSlateRoundedBoxBrush(
+	Set("Stage.Enabled", new FSlateRoundedBoxBrush(
 		FLinearColor::Transparent, StageCornerRadius,
 		ReplaceColorAlpha(EnabledColor, NonHoverAlpha), StageBorderWidth));
-	Style->Set("Stage.Enabled.Hover", new FSlateRoundedBoxBrush(
+	Set("Stage.Enabled.Hover", new FSlateRoundedBoxBrush(
 		FLinearColor::Transparent, StageCornerRadius,
 		ReplaceColorAlpha(EnabledColor, HoverAlpha), StageBorderWidth));
-	Style->Set("Stage.Enabled.Select", new FSlateRoundedBoxBrush(
+	Set("Stage.Enabled.Select", new FSlateRoundedBoxBrush(
 		FLinearColor::Transparent, StageCornerRadius,
 		ReplaceColorAlpha(EnabledSelectedColor, NonHoverAlpha), StageBorderWidth));
-	Style->Set("Stage.Enabled.Select.Hover", new FSlateRoundedBoxBrush(
+	Set("Stage.Enabled.Select.Hover", new FSlateRoundedBoxBrush(
 		FLinearColor::Transparent, StageCornerRadius,
 		ReplaceColorAlpha(EnabledSelectedColor, HoverAlpha), StageBorderWidth));
 
-	Style->Set("Stage.Disabled", new FSlateRoundedBoxBrush(
+	Set("Stage.Disabled", new FSlateRoundedBoxBrush(
 		FLinearColor::Transparent, StageCornerRadius,
 		ReplaceColorAlpha(DisabledColor, NonHoverAlpha), StageBorderWidth));
-	Style->Set("Stage.Disabled.Hover", new FSlateRoundedBoxBrush(
+	Set("Stage.Disabled.Hover", new FSlateRoundedBoxBrush(
 		FLinearColor::Transparent, StageCornerRadius,
 		ReplaceColorAlpha(DisabledColor, HoverAlpha), StageBorderWidth));
-	Style->Set("Stage.Disabled.Select", new FSlateRoundedBoxBrush(
+	Set("Stage.Disabled.Select", new FSlateRoundedBoxBrush(
 		FLinearColor::Transparent, StageCornerRadius,
 		ReplaceColorAlpha(DisabledSelectedColor, NonHoverAlpha), StageBorderWidth));
-	Style->Set("Stage.Disabled.Select.Hover", new FSlateRoundedBoxBrush(
+	Set("Stage.Disabled.Select.Hover", new FSlateRoundedBoxBrush(
 		FLinearColor::Transparent, StageCornerRadius,
 		ReplaceColorAlpha(DisabledSelectedColor, HoverAlpha), StageBorderWidth));
-
-	Style->Set("Stage.Outline", new FSlateRoundedBoxBrush(
-		FStyleColors::InputOutline.GetSpecifiedColor(), 1.0f));
-	Style->Set("Stage.Drag", new FSlateRoundedBoxBrush(
-		FLinearColor::Transparent, StageCornerRadius,
-		FStyleColors::AccentYellow.GetSpecifiedColor(), 2.0f));
-
-	Style->Set("Stage.TextBlock", new FSlateRoundedBoxBrush(
-		ReplaceColorAlpha(FStyleColors::Panel.GetSpecifiedColor(), 0.7f), 0.0f,
-		ReplaceColorAlpha(FStyleColors::InputOutline.GetSpecifiedColor(), 0.4f), 1.0f));
 }
 
-void FDynamicMaterialEditorStyle::SetupLayerViewStyles(const TSharedRef<FSlateStyleSet>& Style)
+void FDynamicMaterialEditorStyle::SetupLayerViewStyles()
 {
-	Style->Set("LayerView.Background", new FSlateRoundedBoxBrush(
-		FStyleColors::Panel.GetSpecifiedColor()/*FLinearColor(0, 0, 0, 0.25f)*/, 6.0f,
-		FStyleColors::Header.GetSpecifiedColor()/*FLinearColor(1, 1, 1, 0.2f)*/, 1.0f));
+	using namespace UE::DynamicMaterialEditor::Private;
 
-	Style->Set("LayerView.Details.Background", new FSlateRoundedBoxBrush(
-		FStyleColors::Recessed.GetSpecifiedColor()/*FLinearColor(0, 0, 0, 0.25f)*/, 6.0f,
+	Set("LayerView.Background", new FSlateRoundedBoxBrush(
+		FStyleColors::Panel.GetSpecifiedColor()/*FLinearColor(0, 0, 0, 0.25f)*/, 6.0f,
 		FStyleColors::Header.GetSpecifiedColor()/*FLinearColor(1, 1, 1, 0.2f)*/, 1.0f));
 
 	/**
@@ -297,12 +212,12 @@ void FDynamicMaterialEditorStyle::SetupLayerViewStyles(const TSharedRef<FSlateSt
 	 * and the SListView widget, so we are not using this style for the SDMMaterialSlotLayerView. Instead, we add a
 	 * SBorder around the SDMBLayerView and style that.
 	 */
-	Style->Set("LayerView", FTableViewStyle()
-		.SetBackgroundBrush(*Style->GetBrush("LayerView.Background"))
+	Set("LayerView", FTableViewStyle()
+		.SetBackgroundBrush(*GetBrush("LayerView.Background"))
 	);
 
-	const float LayerViewItemCornerRadius = 5.0f;
-	const float LayerViewItemBorderWidth = 1.0f;
+	constexpr float LayerViewItemCornerRadius = 5.0f;
+	constexpr float LayerViewItemBorderWidth = 1.0f;
 
 	const FLinearColor LayerViewItemFillColor = FLinearColor::Transparent;
 	constexpr FLinearColor LayerViewItemBorderColor = FLinearColor(1, 1, 1, 0.15f);
@@ -313,60 +228,51 @@ void FDynamicMaterialEditorStyle::SetupLayerViewStyles(const TSharedRef<FSlateSt
 	const FLinearColor LayerItemSelectFillColor = FStyleColors::Header.GetSpecifiedColor();
 	const FLinearColor LayerItemSelectBorderColor = ReplaceColorAlpha(FStyleColors::Select.GetSpecifiedColor(), 0.9f);
 
-	//constexpr FLinearColor LayerItemSelectHoverFillColor = FLinearColor(1, 1, 1, 0.25f);
-	//const FLinearColor LayerItemSelectHoverBorderColor = FStyleColors::Select.GetSpecifiedColor();
-
-	Style->Set("LayerView.Row.Item", new FSlateRoundedBoxBrush(
+	Set("LayerView.Row.Item", new FSlateRoundedBoxBrush(
 		LayerViewItemFillColor, LayerViewItemCornerRadius,
 		LayerViewItemBorderColor, LayerViewItemBorderWidth));
 
-	Style->Set("LayerView.Row.Hovered", new FSlateRoundedBoxBrush(
+	Set("LayerView.Row.Hovered", new FSlateRoundedBoxBrush(
 		LayerItemHoverFillColor, LayerViewItemCornerRadius,
 		LayerItemHoverBorderColor, LayerViewItemBorderWidth));
 
-	Style->Set("LayerView.Row.Selected", new FSlateRoundedBoxBrush(
+	Set("LayerView.Row.Selected", new FSlateRoundedBoxBrush(
 		LayerItemSelectFillColor, LayerViewItemCornerRadius,
 		LayerItemSelectBorderColor, LayerViewItemBorderWidth));
 
-	Style->Set("LayerView.Row.ActiveBrush", new FSlateRoundedBoxBrush(
+	Set("LayerView.Row.ActiveBrush", new FSlateRoundedBoxBrush(
 		LayerItemSelectFillColor, LayerViewItemCornerRadius,
 		LayerItemSelectBorderColor, LayerViewItemBorderWidth));
-	Style->Set("LayerView.Row.ActiveHoveredBrush", new FSlateRoundedBoxBrush(
+	Set("LayerView.Row.ActiveHoveredBrush", new FSlateRoundedBoxBrush(
 		LayerItemSelectFillColor, LayerViewItemCornerRadius,
 		LayerItemSelectBorderColor, LayerViewItemBorderWidth));
-	Style->Set("LayerView.Row.InactiveBrush", new FSlateRoundedBoxBrush(
+	Set("LayerView.Row.InactiveBrush", new FSlateRoundedBoxBrush(
 		LayerItemSelectFillColor, LayerViewItemCornerRadius,
 		LayerItemSelectBorderColor, LayerViewItemBorderWidth));
-	Style->Set("LayerView.Row.InactiveHoveredBrush", new FSlateRoundedBoxBrush(
+	Set("LayerView.Row.InactiveHoveredBrush", new FSlateRoundedBoxBrush(
 		LayerItemSelectFillColor, LayerViewItemCornerRadius,
 		LayerItemSelectBorderColor, LayerViewItemBorderWidth));
 
-	Style->Set("LayerView.Row", FTableRowStyle()
+	Set("LayerView.Row", FTableRowStyle()
 		.SetTextColor(FStyleColors::Foreground)
 		.SetSelectedTextColor(FStyleColors::ForegroundHover)
-		.SetEvenRowBackgroundBrush(*Style->GetBrush(TEXT("LayerView.Row.Item")))
-		.SetEvenRowBackgroundHoveredBrush(*Style->GetBrush(TEXT("LayerView.Row.Hovered")))
-		.SetOddRowBackgroundBrush(*Style->GetBrush(TEXT("LayerView.Row.Item")))
-		.SetOddRowBackgroundHoveredBrush(*Style->GetBrush(TEXT("LayerView.Row.Hovered")))
-		.SetSelectorFocusedBrush(*Style->GetBrush(TEXT("LayerView.Row.Selected")))
-		.SetActiveBrush(*Style->GetBrush(TEXT("LayerView.Row.ActiveBrush")))
-		.SetActiveHoveredBrush(*Style->GetBrush(TEXT("LayerView.Row.ActiveHoveredBrush")))
-		.SetInactiveBrush(*Style->GetBrush(TEXT("LayerView.Row.InactiveBrush")))
-		.SetInactiveHoveredBrush(*Style->GetBrush(TEXT("LayerView.Row.InactiveHoveredBrush")))
-		.SetSelectorFocusedBrush(BORDER_BRUSH("Images/DropIndicators/DropZoneIndicator_Onto", FMargin(4.f / 16.f), Style->GetColor(TEXT("Color.Select.Hover"))))
-		.SetDropIndicator_Onto(BOX_BRUSH("Images/DropIndicators/LayerView_DropIndicator_Onto", FMargin(4.0f / 16.0f), Style->GetColor(TEXT("Color.Select.Hover"))))
-		.SetDropIndicator_Above(BOX_BRUSH("Images/DropIndicators/LayerView_DropIndicator_Above", FMargin(4.0f / 16.0f, 4.0f / 16.0f, 0.f, 0.f), Style->GetColor(TEXT("Color.Select.Hover"))))
-		.SetDropIndicator_Below(BOX_BRUSH("Images/DropIndicators/LayerView_DropIndicator_Below", FMargin(4.0f / 16.0f, 0.f, 0.f, 4.0f / 16.0f), Style->GetColor(TEXT("Color.Select.Hover"))))
+		.SetEvenRowBackgroundBrush(*GetBrush(TEXT("LayerView.Row.Item")))
+		.SetEvenRowBackgroundHoveredBrush(*GetBrush(TEXT("LayerView.Row.Hovered")))
+		.SetOddRowBackgroundBrush(*GetBrush(TEXT("LayerView.Row.Item")))
+		.SetOddRowBackgroundHoveredBrush(*GetBrush(TEXT("LayerView.Row.Hovered")))
+		.SetSelectorFocusedBrush(*GetBrush(TEXT("LayerView.Row.Selected")))
+		.SetActiveBrush(*GetBrush(TEXT("LayerView.Row.ActiveBrush")))
+		.SetActiveHoveredBrush(*GetBrush(TEXT("LayerView.Row.ActiveHoveredBrush")))
+		.SetInactiveBrush(*GetBrush(TEXT("LayerView.Row.InactiveBrush")))
+		.SetInactiveHoveredBrush(*GetBrush(TEXT("LayerView.Row.InactiveHoveredBrush")))
+		.SetSelectorFocusedBrush(BORDER_BRUSH("Images/DropIndicators/DropZoneIndicator_Onto", FMargin(4.f / 16.f), GetColor(TEXT("Color.Select.Hover"))))
+		.SetDropIndicator_Onto(BOX_BRUSH("Images/DropIndicators/LayerView_DropIndicator_Onto", FMargin(4.0f / 16.0f), GetColor(TEXT("Color.Select.Hover"))))
+		.SetDropIndicator_Above(BOX_BRUSH("Images/DropIndicators/LayerView_DropIndicator_Above", FMargin(4.0f / 16.0f, 4.0f / 16.0f, 0.f, 0.f), GetColor(TEXT("Color.Select.Hover"))))
+		.SetDropIndicator_Below(BOX_BRUSH("Images/DropIndicators/LayerView_DropIndicator_Below", FMargin(4.0f / 16.0f, 0.f, 0.f, 4.0f / 16.0f), GetColor(TEXT("Color.Select.Hover"))))
 	);
-
-	Style->Set("LayerView.AddIcon", new IMAGE_BRUSH("Icons/EditorIcons/LayerAdd", Icon16x16));
-	Style->Set("LayerView.DuplicateIcon", new IMAGE_BRUSH("Icons/EditorIcons/Duplicate_40x", Icon40x40));
-	Style->Set("LayerView.RemoveIcon", new IMAGE_BRUSH("Icons/EditorIcons/LayerRemove", Icon16x16));
-
-	Style->Set("LayerView.Row.Handle", new IMAGE_BRUSH_SVG("Icons/DragHandle", Icon16x16));
 }
 
-void FDynamicMaterialEditorStyle::SetupLayerViewItemHandleStyles(const TSharedRef<FSlateStyleSet>& Style)
+void FDynamicMaterialEditorStyle::SetupLayerViewItemHandleStyles()
 {
 	constexpr FLinearColor RowHandleFillColor = FLinearColor(1, 1, 1, 0.3f);
 	const FLinearColor RowHandleHoverFillColor = FLinearColor(1, 1, 1, 0.4f);
@@ -374,72 +280,62 @@ void FDynamicMaterialEditorStyle::SetupLayerViewItemHandleStyles(const TSharedRe
 	constexpr float RowHandleCornerRadius = 6.0f;
 	constexpr float RowHandleBorderWidth = 1.0f;
 
-	Style->Set("LayerView.Row.Handle.Left", new FSlateRoundedBoxBrush(
+	Set("LayerView.Row.Handle.Left", new FSlateRoundedBoxBrush(
 		RowHandleFillColor, FVector4(RowHandleCornerRadius, 0.0f, 0.0f, RowHandleCornerRadius),
 		RowHandleBorderColor, RowHandleBorderWidth));
-	Style->Set("LayerView.Row.Handle.Top", new FSlateRoundedBoxBrush(
+	Set("LayerView.Row.Handle.Top", new FSlateRoundedBoxBrush(
 		RowHandleFillColor, FVector4(RowHandleCornerRadius, RowHandleCornerRadius, 0.0f, 0.0f),
 		RowHandleBorderColor, RowHandleBorderWidth));
-	Style->Set("LayerView.Row.Handle.Right", new FSlateRoundedBoxBrush(
+	Set("LayerView.Row.Handle.Right", new FSlateRoundedBoxBrush(
 		RowHandleFillColor, FVector4(0.0f, RowHandleCornerRadius, RowHandleCornerRadius, 0.0f),
 		RowHandleBorderColor, RowHandleBorderWidth));
-	Style->Set("LayerView.Row.Handle.Bottom", new FSlateRoundedBoxBrush(
+	Set("LayerView.Row.Handle.Bottom", new FSlateRoundedBoxBrush(
 		RowHandleFillColor, FVector4(0.0f, 0.0f, RowHandleCornerRadius, RowHandleCornerRadius),
 		RowHandleBorderColor, RowHandleBorderWidth));
 
-	Style->Set("LayerView.Row.Handle.Left.Hover", new FSlateRoundedBoxBrush(
+	Set("LayerView.Row.Handle.Left.Hover", new FSlateRoundedBoxBrush(
 		RowHandleHoverFillColor, FVector4(RowHandleCornerRadius, 0.0f, 0.0f, RowHandleCornerRadius),
 		RowHandleBorderColor, RowHandleBorderWidth));
-	Style->Set("LayerView.Row.Handle.Top.Hover", new FSlateRoundedBoxBrush(
+	Set("LayerView.Row.Handle.Top.Hover", new FSlateRoundedBoxBrush(
 		RowHandleHoverFillColor, FVector4(RowHandleCornerRadius, RowHandleCornerRadius, 0.0f, 0.0f),
 		RowHandleBorderColor, RowHandleBorderWidth));
-	Style->Set("LayerView.Row.Handle.Right.Hover", new FSlateRoundedBoxBrush(
+	Set("LayerView.Row.Handle.Right.Hover", new FSlateRoundedBoxBrush(
 		RowHandleHoverFillColor, FVector4(0.0f, RowHandleCornerRadius, RowHandleCornerRadius, 0.0f),
 		RowHandleBorderColor, RowHandleBorderWidth));
-	Style->Set("LayerView.Row.Handle.Bottom.Hover", new FSlateRoundedBoxBrush(
+	Set("LayerView.Row.Handle.Bottom.Hover", new FSlateRoundedBoxBrush(
 		RowHandleHoverFillColor, FVector4(0.0f, 0.0f, RowHandleCornerRadius, RowHandleCornerRadius),
 		RowHandleBorderColor, RowHandleBorderWidth));
 
-	Style->Set("LayerView.Row.Handle.Left.Select", new FSlateRoundedBoxBrush(
-		Style->GetColor(TEXT("Color.Select")), 0.0f,
+	Set("LayerView.Row.Handle.Left.Select", new FSlateRoundedBoxBrush(
+		GetColor(TEXT("Color.Select")), 0.0f,
 		RowHandleBorderColor, 0.0f));
-	Style->Set("LayerView.Row.Handle.Top.Select", new FSlateRoundedBoxBrush(
-		Style->GetColor(TEXT("Color.Select")), 0.0f,
+	Set("LayerView.Row.Handle.Top.Select", new FSlateRoundedBoxBrush(
+		GetColor(TEXT("Color.Select")), 0.0f,
 		RowHandleBorderColor, 0.0f));
-	Style->Set("LayerView.Row.Handle.Right.Select", new FSlateRoundedBoxBrush(
-		Style->GetColor(TEXT("Color.Select")), 0.0f,
+	Set("LayerView.Row.Handle.Right.Select", new FSlateRoundedBoxBrush(
+		GetColor(TEXT("Color.Select")), 0.0f,
 		RowHandleBorderColor, 0.0f));
-	Style->Set("LayerView.Row.Handle.Bottom.Select", new FSlateRoundedBoxBrush(
-		Style->GetColor(TEXT("Color.Select")), 0.0f,
+	Set("LayerView.Row.Handle.Bottom.Select", new FSlateRoundedBoxBrush(
+		GetColor(TEXT("Color.Select")), 0.0f,
 		RowHandleBorderColor, 0.0f));
 
-	Style->Set("LayerView.Row.Handle.Left.Select.Hover", new FSlateRoundedBoxBrush(
-		Style->GetColor(TEXT("Color.Select.Hover")), 0.0f,
+	Set("LayerView.Row.Handle.Left.Select.Hover", new FSlateRoundedBoxBrush(
+		GetColor(TEXT("Color.Select.Hover")), 0.0f,
 		RowHandleBorderColor, 0.0f));
-	Style->Set("LayerView.Row.Handle.Top.Select.Hover", new FSlateRoundedBoxBrush(
-		Style->GetColor(TEXT("Color.Select.Hover")), 0.0f,
+	Set("LayerView.Row.Handle.Top.Select.Hover", new FSlateRoundedBoxBrush(
+		GetColor(TEXT("Color.Select.Hover")), 0.0f,
 		RowHandleBorderColor, 0.0f));
-	Style->Set("LayerView.Row.Handle.Right.Select.Hover", new FSlateRoundedBoxBrush(
-		Style->GetColor(TEXT("Color.Select.Hover")), 0.0f,
+	Set("LayerView.Row.Handle.Right.Select.Hover", new FSlateRoundedBoxBrush(
+		GetColor(TEXT("Color.Select.Hover")), 0.0f,
 		RowHandleBorderColor, 0.0f));
-	Style->Set("LayerView.Row.Handle.Bottom.Select.Hover", new FSlateRoundedBoxBrush(
-		Style->GetColor(TEXT("Color.Select.Hover")), 0.0f,
+	Set("LayerView.Row.Handle.Bottom.Select.Hover", new FSlateRoundedBoxBrush(
+		GetColor(TEXT("Color.Select.Hover")), 0.0f,
 		RowHandleBorderColor, 0.0f));
 }
 
-void FDynamicMaterialEditorStyle::SetupEffectsViewStyles(const TSharedRef<FSlateStyleSet>& Style)
+void FDynamicMaterialEditorStyle::SetupEffectsViewStyles()
 {
-	Style->Set("EffectsView.Background", new FSlateRoundedBoxBrush(
-		FStyleColors::Panel.GetSpecifiedColor(), 0.0f,
-		FStyleColors::Header.GetSpecifiedColor(), 1.0f));
-
-	Style->Set("EffectsView.Details.Background", new FSlateRoundedBoxBrush(
-		FStyleColors::Recessed.GetSpecifiedColor(), 0.0f,
-		FStyleColors::Header.GetSpecifiedColor(), 1.0f));
-
-	Style->Set("EffectsView", FTableViewStyle()
-		.SetBackgroundBrush(*Style->GetBrush(TEXT("EffectsView.Background")))
-	);
+	using namespace UE::DynamicMaterialEditor::Private;
 
 	const float LayerViewItemCornerRadius = 0.0f;
 	const float LayerViewItemBorderWidth = 1.0f;
@@ -453,59 +349,59 @@ void FDynamicMaterialEditorStyle::SetupEffectsViewStyles(const TSharedRef<FSlate
 	const FLinearColor LayerItemSelectFillColor = FStyleColors::Header.GetSpecifiedColor();
 	const FLinearColor LayerItemSelectBorderColor = ReplaceColorAlpha(FStyleColors::Select.GetSpecifiedColor(), 0.9f);
 
-	Style->Set("EffectsView.Row.Item", new FSlateRoundedBoxBrush(
+	Set("EffectsView.Row.Item", new FSlateRoundedBoxBrush(
 		LayerViewItemFillColor, LayerViewItemCornerRadius,
 		LayerViewItemBorderColor, LayerViewItemBorderWidth));
 
-	Style->Set("EffectsView.Row.Hovered", new FSlateRoundedBoxBrush(
+	Set("EffectsView.Row.Hovered", new FSlateRoundedBoxBrush(
 		LayerItemHoverFillColor, LayerViewItemCornerRadius,
 		LayerItemHoverBorderColor, LayerViewItemBorderWidth));
 
-	Style->Set("EffectsView.Row.Selected", new FSlateRoundedBoxBrush(
+	Set("EffectsView.Row.Selected", new FSlateRoundedBoxBrush(
 		LayerItemSelectFillColor, LayerViewItemCornerRadius,
 		LayerItemSelectBorderColor, LayerViewItemBorderWidth));
 
-	Style->Set("EffectsView.Row.ActiveBrush", new FSlateRoundedBoxBrush(
+	Set("EffectsView.Row.ActiveBrush", new FSlateRoundedBoxBrush(
 		LayerItemSelectFillColor, LayerViewItemCornerRadius,
 		LayerItemSelectBorderColor, LayerViewItemBorderWidth));
-	Style->Set("EffectsView.Row.ActiveHoveredBrush", new FSlateRoundedBoxBrush(
+	Set("EffectsView.Row.ActiveHoveredBrush", new FSlateRoundedBoxBrush(
 		LayerItemSelectFillColor, LayerViewItemCornerRadius,
 		LayerItemSelectBorderColor, LayerViewItemBorderWidth));
-	Style->Set("EffectsView.Row.InactiveBrush", new FSlateRoundedBoxBrush(
+	Set("EffectsView.Row.InactiveBrush", new FSlateRoundedBoxBrush(
 		LayerItemSelectFillColor, LayerViewItemCornerRadius,
 		LayerItemSelectBorderColor, LayerViewItemBorderWidth));
-	Style->Set("EffectsView.Row.InactiveHoveredBrush", new FSlateRoundedBoxBrush(
+	Set("EffectsView.Row.InactiveHoveredBrush", new FSlateRoundedBoxBrush(
 		LayerItemSelectFillColor, LayerViewItemCornerRadius,
 		LayerItemSelectBorderColor, LayerViewItemBorderWidth));
 
-	Style->Set("EffectsView.Row", FTableRowStyle()
+	Set("EffectsView.Row", FTableRowStyle()
 		.SetTextColor(FStyleColors::Foreground)
 		.SetSelectedTextColor(FStyleColors::ForegroundHover)
-		.SetEvenRowBackgroundBrush(*Style->GetBrush(TEXT("EffectsView.Row.Item")))
-		.SetEvenRowBackgroundHoveredBrush(*Style->GetBrush(TEXT("EffectsView.Row.Hovered")))
-		.SetOddRowBackgroundBrush(*Style->GetBrush(TEXT("EffectsView.Row.Item")))
-		.SetOddRowBackgroundHoveredBrush(*Style->GetBrush(TEXT("EffectsView.Row.Hovered")))
-		.SetSelectorFocusedBrush(*Style->GetBrush(TEXT("EffectsView.Row.Selected")))
-		.SetActiveBrush(*Style->GetBrush(TEXT("EffectsView.Row.ActiveBrush")))
-		.SetActiveHoveredBrush(*Style->GetBrush(TEXT("EffectsView.Row.ActiveHoveredBrush")))
-		.SetInactiveBrush(*Style->GetBrush(TEXT("EffectsView.Row.InactiveBrush")))
-		.SetInactiveHoveredBrush(*Style->GetBrush(TEXT("EffectsView.Row.InactiveHoveredBrush")))
-		.SetSelectorFocusedBrush(BORDER_BRUSH("Images/DropIndicators/DropZoneIndicator_Onto", FMargin(4.f / 16.f), Style->GetColor(TEXT("Color.Select.Hover"))))
-		.SetDropIndicator_Onto(BOX_BRUSH("Images/DropIndicators/LayerView_DropIndicator_Onto", FMargin(4.0f / 16.0f), Style->GetColor(TEXT("Color.Select.Hover"))))
-		.SetDropIndicator_Above(BOX_BRUSH("Images/DropIndicators/LayerView_DropIndicator_Above", FMargin(4.0f / 16.0f, 4.0f / 16.0f, 0.f, 0.f), Style->GetColor(TEXT("Color.Select.Hover"))))
-		.SetDropIndicator_Below(BOX_BRUSH("Images/DropIndicators/LayerView_DropIndicator_Below", FMargin(4.0f / 16.0f, 0.f, 0.f, 4.0f / 16.0f), Style->GetColor(TEXT("Color.Select.Hover"))))
+		.SetEvenRowBackgroundBrush(*GetBrush(TEXT("EffectsView.Row.Item")))
+		.SetEvenRowBackgroundHoveredBrush(*GetBrush(TEXT("EffectsView.Row.Hovered")))
+		.SetOddRowBackgroundBrush(*GetBrush(TEXT("EffectsView.Row.Item")))
+		.SetOddRowBackgroundHoveredBrush(*GetBrush(TEXT("EffectsView.Row.Hovered")))
+		.SetSelectorFocusedBrush(*GetBrush(TEXT("EffectsView.Row.Selected")))
+		.SetActiveBrush(*GetBrush(TEXT("EffectsView.Row.ActiveBrush")))
+		.SetActiveHoveredBrush(*GetBrush(TEXT("EffectsView.Row.ActiveHoveredBrush")))
+		.SetInactiveBrush(*GetBrush(TEXT("EffectsView.Row.InactiveBrush")))
+		.SetInactiveHoveredBrush(*GetBrush(TEXT("EffectsView.Row.InactiveHoveredBrush")))
+		.SetSelectorFocusedBrush(BORDER_BRUSH("Images/DropIndicators/DropZoneIndicator_Onto", FMargin(4.f / 16.f), GetColor(TEXT("Color.Select.Hover"))))
+		.SetDropIndicator_Onto(BOX_BRUSH("Images/DropIndicators/LayerView_DropIndicator_Onto", FMargin(4.0f / 16.0f), GetColor(TEXT("Color.Select.Hover"))))
+		.SetDropIndicator_Above(BOX_BRUSH("Images/DropIndicators/LayerView_DropIndicator_Above", FMargin(4.0f / 16.0f, 4.0f / 16.0f, 0.f, 0.f), GetColor(TEXT("Color.Select.Hover"))))
+		.SetDropIndicator_Below(BOX_BRUSH("Images/DropIndicators/LayerView_DropIndicator_Below", FMargin(4.0f / 16.0f, 0.f, 0.f, 4.0f / 16.0f), GetColor(TEXT("Color.Select.Hover"))))
 	);
 
-	Style->Set("EffectsView.Row.Fx.Closed", new IMAGE_BRUSH_SVG("Icons/Fx_Closed", Icon24x24));
-	Style->Set("EffectsView.Row.Fx.Opened", new IMAGE_BRUSH_SVG("Icons/Fx_Opened", Icon24x24));
-	Style->Set("EffectsView.Row.Fx", new IMAGE_BRUSH_SVG("Icons/Fx", Icon24x24));
+	Set("EffectsView.Row.Fx.Closed", new IMAGE_BRUSH_SVG("Icons/Fx_Closed", CoreStyleConstants::Icon24x24));
+	Set("EffectsView.Row.Fx.Opened", new IMAGE_BRUSH_SVG("Icons/Fx_Opened", CoreStyleConstants::Icon24x24));
+	Set("EffectsView.Row.Fx", new IMAGE_BRUSH_SVG("Icons/Fx", CoreStyleConstants::Icon24x24));
 }
 
-void FDynamicMaterialEditorStyle::SetupTextStyles(const TSharedRef<FSlateStyleSet>& Style)
+void FDynamicMaterialEditorStyle::SetupTextStyles()
 {
 	const FTextBlockStyle NormalTextStyle(FAppStyle::GetWidgetStyle<FTextBlockStyle>(TEXT("NormalText")));
 
-	const FLinearColor LayerViewItemTextShadowColor = FLinearColor(0.0f, 0.0f, 0.0f, 0.5f);
+	constexpr FLinearColor LayerViewItemTextShadowColor = FLinearColor(0.0f, 0.0f, 0.0f, 0.5f);
 
 	FFontOutlineSettings HandleFontOutline;
 	HandleFontOutline.OutlineColor = LayerViewItemTextShadowColor;
@@ -513,77 +409,64 @@ void FDynamicMaterialEditorStyle::SetupTextStyles(const TSharedRef<FSlateStyleSe
 
 	FTextBlockStyle SmallTextStyle(NormalTextStyle);
 	SmallTextStyle.SetFont(DEFAULT_FONT("Regular", 8));
-	Style->Set("SmallFont", SmallTextStyle);
+	Set("SmallFont", SmallTextStyle);
 
 	FTextBlockStyle RegularTextStyle(NormalTextStyle);
 	RegularTextStyle.SetFont(DEFAULT_FONT("Regular", 10));
-	Style->Set("RegularFont", RegularTextStyle);
+	Set("RegularFont", RegularTextStyle);
 
 	FTextBlockStyle BoldTextStyle(NormalTextStyle);
 	BoldTextStyle.SetFont(DEFAULT_FONT("Bold", 10));
-	Style->Set("BoldFont", BoldTextStyle);
+	Set("BoldFont", BoldTextStyle);
 
-	Style->Set("ActorName", RegularTextStyle);
+	Set("ActorName", RegularTextStyle);
 
 	FTextBlockStyle ActorNameBigTextStyle(NormalTextStyle);
 	ActorNameBigTextStyle.SetFont(DEFAULT_FONT("Regular", 14));
-	Style->Set("ActorNameBig", ActorNameBigTextStyle);
+	Set("ActorNameBig", ActorNameBigTextStyle);
 
 	FTextBlockStyle ComponentNameBigTextStyle(NormalTextStyle);
 	ComponentNameBigTextStyle.SetFont(DEFAULT_FONT("Regular", 12));
-	Style->Set("ComponentNameBig", ComponentNameBigTextStyle);
+	Set("ComponentNameBig", ComponentNameBigTextStyle);
 
 	FTextBlockStyle SlotLayerInfoTextStyle(NormalTextStyle);
 	SlotLayerInfoTextStyle.SetFont(DEFAULT_FONT("Italic", 8));
-	Style->Set("SlotLayerInfo", SlotLayerInfoTextStyle);
+	Set("SlotLayerInfo", SlotLayerInfoTextStyle);
 
 	FSlateFontInfo LayerViewItemFont = DEFAULT_FONT("Bold", 12);
 	LayerViewItemFont.OutlineSettings = HandleFontOutline;
-	Style->Set("LayerView.Row.Font", LayerViewItemFont);
+	Set("LayerView.Row.Font", LayerViewItemFont);
 
-	Style->Set("LayerView.Row.HandleFont", RegularTextStyle);
+	Set("LayerView.Row.HandleFont", RegularTextStyle);
 
 	FTextBlockStyle LayerViewItemTextStyle(NormalTextStyle);
 	LayerViewItemTextStyle.SetShadowOffset(FVector2D(1.0f, 1.0f));
 	LayerViewItemTextStyle.SetColorAndOpacity(LayerViewItemTextShadowColor);
 
-	Style->Set("LayerView.Row.HeaderText",
+	Set("LayerView.Row.HeaderText",
 		FTextBlockStyle(LayerViewItemTextStyle)
 		.SetColorAndOpacity(FStyleColors::Foreground)
 		.SetFont(LayerViewItemFont));
 
-	Style->Set("LayerView.Row.HeaderText.Small",
+	Set("LayerView.Row.HeaderText.Small",
 		FTextBlockStyle(LayerViewItemTextStyle)
 		.SetColorAndOpacity(FStyleColors::Foreground)
 		.SetFont(RegularTextStyle.Font));
 
 	FTextBlockStyle StagePropertyDetailsTextStyle(NormalTextStyle);
 	StagePropertyDetailsTextStyle.SetFont(DEFAULT_FONT("Regular", 12));
-	Style->Set("Font.Stage.Details", StagePropertyDetailsTextStyle);
+	Set("Font.Stage.Details", StagePropertyDetailsTextStyle);
 
-	Style->Set("Font.Stage.Details.Bold", BoldTextStyle);
+	Set("Font.Stage.Details.Bold", BoldTextStyle);
 
 	FTextBlockStyle StagePropertyDetailsSmallTextStyle(NormalTextStyle);
 	StagePropertyDetailsSmallTextStyle.SetFont(IDetailLayoutBuilder::GetDetailFont());
-	Style->Set("Font.Stage.Details.Small", StagePropertyDetailsSmallTextStyle);
+	Set("Font.Stage.Details.Small", StagePropertyDetailsSmallTextStyle);
 
 	FTextBlockStyle StagePropertyDetailsSmallBoldTextStyle(NormalTextStyle);
 	StagePropertyDetailsSmallBoldTextStyle.SetFont(IDetailLayoutBuilder::GetDetailFontBold());
-	Style->Set("Font.Stage.Details.Small.Bold", StagePropertyDetailsSmallBoldTextStyle);
+	Set("Font.Stage.Details.Small.Bold", StagePropertyDetailsSmallBoldTextStyle);
 
-	Style->Set("Font.Stage.Details.Small.Bold", FTextBlockStyle(NormalTextStyle)
+	Set("Font.Stage.Details.Small.Bold", FTextBlockStyle(NormalTextStyle)
 		.SetFont(IDetailLayoutBuilder::GetDetailFontBold()));
 }
-
-#undef IMAGE_BRUSH
-#undef BOX_BRUSH
-#undef BORDER_BRUSH
-
-#undef IMAGE_BRUSH_SVG
-#undef BOX_BRUSH_SVG
-#undef BORDER_BRUSH_SVG
-
-#undef DEFAULT_FONT
-
-#undef IMAGE_PLUGIN_BRUSH
-#undef IMAGE_PLUGIN_BRUSH_SVG
diff --git a/Engine/Plugins/Experimental/DynamicMaterial/Source/DynamicMaterialEditor/Private/DynamicMaterialEditorStyle.h b/Engine/Plugins/Experimental/DynamicMaterial/Source/DynamicMaterialEditor/Private/DynamicMaterialEditorStyle.h
--- a/Engine/Plugins/Experimental/DynamicMaterial/Source/DynamicMaterialEditor/Private/DynamicMaterialEditorStyle.h
+++ b/Engine/Plugins/Experimental/DynamicMaterial/Source/DynamicMaterialEditor/Private/DynamicMaterialEditorStyle.h
@@ -2,33 +2,30 @@
 
 #pragma once
 
-#include "Styling/ISlateStyle.h"
 #include "Styling/SlateStyle.h"
 
+class FName;
 enum class EDMValueType : uint8;
 
-class FDynamicMaterialEditorStyle
+class FDynamicMaterialEditorStyle final : public FSlateStyleSet
 {
 public:
-	static void Initialize();
-	static void Shutdown();
+	static FDynamicMaterialEditorStyle& Get()
+	{
+		static FDynamicMaterialEditorStyle Instance;
+		return Instance;
+	}
 
-	static FName GetStyleSetName();
-	static const ISlateStyle& Get();
+	FDynamicMaterialEditorStyle();
+	virtual ~FDynamicMaterialEditorStyle() override = default;
 
-	static FLinearColor GetColor(const FName& InName);
-	static const FSlateBrush* GetBrush(const FName& InName);
-
-	static FName GetBrushNameForType(EDMValueType InType);
+	FName GetBrushNameForType(EDMValueType InType) const;
 
 private:
-	static TSharedPtr<FSlateStyleSet> StyleInstance;
-
-	static TSharedRef<FSlateStyleSet> Create();
-
-	static void SetupStageStyles(const TSharedRef<FSlateStyleSet>& Style);
-	static void SetupLayerViewStyles(const TSharedRef<FSlateStyleSet>& Style);
-	static void SetupLayerViewItemHandleStyles(const TSharedRef<FSlateStyleSet>& Style);
-	static void SetupEffectsViewStyles(const TSharedRef<FSlateStyleSet>& Style);
-	static void SetupTextStyles(const TSharedRef<FSlateStyleSet>& Style);
+	void SetupGenericStyles();
+	void SetupStageStyles();
+	void SetupLayerViewStyles();
+	void SetupLayerViewItemHandleStyles();
+	void SetupEffectsViewStyles();
+	void SetupTextStyles();
 };
diff --git a/Engine/Plugins/Experimental/DynamicMaterial/Source/DynamicMaterialEditor/Private/UI/Widgets/Editor/SDMMaterialSlotEditor.cpp b/Engine/Plugins/Experimental/DynamicMaterial/Source/DynamicMaterialEditor/Private/UI/Widgets/Editor/SDMMaterialSlotEditor.cpp
--- a/Engine/Plugins/Experimental/DynamicMaterial/Source/DynamicMaterialEditor/Private/UI/Widgets/Editor/SDMMaterialSlotEditor.cpp
+++ b/Engine/Plugins/Experimental/DynamicMaterial/Source/DynamicMaterialEditor/Private/UI/Widgets/Editor/SDMMaterialSlotEditor.cpp
@@ -425,7 +425,7 @@ TSharedRef<SWidget> SDMMaterialSlotEditor::CreateSlot_Container()
 		[
 			SNew(SBorder)
 			.Padding(2.0f)
-			.BorderImage(FDynamicMaterialEditorStyle::GetBrush("LayerView.Background"))
+			.BorderImage(FDynamicMaterialEditorStyle::Get().GetBrush("LayerView.Background"))
 			.HAlign(HAlign_Fill)
 			.VAlign(VAlign_Fill)
 			[
@@ -794,7 +794,7 @@ TSharedRef<SWidget> SDMMaterialSlotEditor::CreateSlot_LayerSettings()
 			.ButtonContent()
 			[
 				SNew(SImage)
-				.Image(FDynamicMaterialEditorStyle::GetBrush("EffectsView.Row.Fx"))
+				.Image(FDynamicMaterialEditorStyle::Get().GetBrush("EffectsView.Row.Fx"))
 				.DesiredSizeOverride(FVector2D(16.0f))
 			]
 		]
diff --git a/Engine/Plugins/Experimental/DynamicMaterial/Source/DynamicMaterialEditor/Private/UI/Widgets/Editor/SDMStatusBar.cpp b/Engine/Plugins/Experimental/DynamicMaterial/Source/DynamicMaterialEditor/Private/UI/Widgets/Editor/SDMStatusBar.cpp
--- a/Engine/Plugins/Experimental/DynamicMaterial/Source/DynamicMaterialEditor/Private/UI/Widgets/Editor/SDMStatusBar.cpp
+++ b/Engine/Plugins/Experimental/DynamicMaterial/Source/DynamicMaterialEditor/Private/UI/Widgets/Editor/SDMStatusBar.cpp
@@ -83,7 +83,7 @@ TSharedRef<SWidget> SDMStatusBar::CreateContent()
 	return SNew(SBorder)
 		.HAlign(HAlign_Fill)
 		.VAlign(VAlign_Center)
-		.BorderImage(FDynamicMaterialEditorStyle::GetBrush("Border.Top"))
+		.BorderImage(FDynamicMaterialEditorStyle::Get().GetBrush("Border.Top"))
 		.BorderBackgroundColor(FLinearColor(1, 1, 1, 0.05f))
 		[
 			SNew(SWrapBox)
@@ -123,7 +123,7 @@ SWrapBox::FSlot::FSlotArguments SDMStatusBar::CreateStatsWrapBoxEntry(const FTex
 			[
 				SNew(SBorder)
 				.BorderBackgroundColor(SeparatorColor)
-				.BorderImage(FDynamicMaterialEditorStyle::GetBrush("Border.Left"))
+				.BorderImage(FDynamicMaterialEditorStyle::Get().GetBrush("Border.Left"))
 			]
 			+ SHorizontalBox::Slot()
 			.AutoWidth()
diff --git a/Engine/Plugins/Experimental/DynamicMaterial/Source/DynamicMaterialEditor/Private/UI/Widgets/Editor/SDMToolBar.cpp b/Engine/Plugins/Experimental/DynamicMaterial/Source/DynamicMaterialEditor/Private/UI/Widgets/Editor/SDMToolBar.cpp
--- a/Engine/Plugins/Experimental/DynamicMaterial/Source/DynamicMaterialEditor/Private/UI/Widgets/Editor/SDMToolBar.cpp
+++ b/Engine/Plugins/Experimental/DynamicMaterial/Source/DynamicMaterialEditor/Private/UI/Widgets/Editor/SDMToolBar.cpp
@@ -57,7 +57,7 @@ void SDMToolBar::Construct(const FArguments& InArgs, const TSharedRef<SDMMateria
 			SNew(SBorder)
 			.HAlign(HAlign_Fill)
 			.VAlign(VAlign_Center)
-			.BorderImage(FDynamicMaterialEditorStyle::GetBrush("Border.Bottom"))
+			.BorderImage(FDynamicMaterialEditorStyle::Get().GetBrush("Border.Bottom"))
 			.BorderBackgroundColor(FLinearColor(1, 1, 1, 0.05f))
 			.Padding(0.f, 3.f, 0.f, 3.f)
 			[
@@ -281,7 +281,7 @@ TSharedRef<SWidget> SDMToolBar::CreateToolBarEntries()
 			.ButtonContent()
 			[
 				SNew(SImage)
-				.Image(FDynamicMaterialEditorStyle::GetBrush("Icons.Menu.Dropdown"))
+				.Image(FDynamicMaterialEditorStyle::Get().GetBrush("Icons.Menu.Dropdown"))
 				.DesiredSizeOverride(DefaultToolBarButtonSize)
 			]
 		];
diff --git a/Engine/Plugins/Experimental/DynamicMaterial/Source/DynamicMaterialEditor/Private/UI/Widgets/Editor/SlotEditor/SDMMaterialSlotLayerItem.cpp b/Engine/Plugins/Experimental/DynamicMaterial/Source/DynamicMaterialEditor/Private/UI/Widgets/Editor/SlotEditor/SDMMaterialSlotLayerItem.cpp
--- a/Engine/Plugins/Experimental/DynamicMaterial/Source/DynamicMaterialEditor/Private/UI/Widgets/Editor/SlotEditor/SDMMaterialSlotLayerItem.cpp
+++ b/Engine/Plugins/Experimental/DynamicMaterial/Source/DynamicMaterialEditor/Private/UI/Widgets/Editor/SlotEditor/SDMMaterialSlotLayerItem.cpp
@@ -305,7 +305,7 @@ TSharedRef<SWidget> SDMMaterialSlotLayerItem::CreateEffectsRowContent()
 			[
 				SNew(SBorder)
 				.BorderBackgroundColor(FLinearColor(1, 1, 1, 0.2f))
-				.BorderImage(FDynamicMaterialEditorStyle::GetBrush("Border.Right"))
+				.BorderImage(FDynamicMaterialEditorStyle::Get().GetBrush("Border.Right"))
 			]
 		]
 		+ SHorizontalBox::Slot()
@@ -316,7 +316,7 @@ TSharedRef<SWidget> SDMMaterialSlotLayerItem::CreateEffectsRowContent()
 			[
 				SNew(SBorder)
 				.BorderBackgroundColor(FLinearColor(1, 1, 1, 0.2f))
-				.BorderImage(FDynamicMaterialEditorStyle::GetBrush("Border.Left"))
+				.BorderImage(FDynamicMaterialEditorStyle::Get().GetBrush("Border.Left"))
 			]
 		]
 		+ SHorizontalBox::Slot()
@@ -636,8 +636,8 @@ FReply SDMMaterialSlotLayerItem::OnEffectsToggleButtonClicked()
 
 const FSlateBrush* SDMMaterialSlotLayerItem::GetEffectsToggleButtonImage() const
 {
-	static const FSlateBrush* Displayed = FDynamicMaterialEditorStyle::GetBrush("EffectsView.Row.Fx.Opened");
-	static const FSlateBrush* Hidden = FDynamicMaterialEditorStyle::GetBrush("EffectsView.Row.Fx.Closed");
+	static const FSlateBrush* Displayed = FDynamicMaterialEditorStyle::Get().GetBrush("EffectsView.Row.Fx.Opened");
+	static const FSlateBrush* Hidden = FDynamicMaterialEditorStyle::Get().GetBrush("EffectsView.Row.Fx.Closed");
 
 	if (AreEffectsExpanded())
 	{
@@ -649,10 +649,10 @@ const FSlateBrush* SDMMaterialSlotLayerItem::GetEffectsToggleButtonImage() const
 
 const FSlateBrush* SDMMaterialSlotLayerItem::GetRowHandleBrush() const
 {
-	static const FSlateBrush* const Default = FDynamicMaterialEditorStyle::GetBrush("LayerView.Row.Handle.Left");
-	static const FSlateBrush* const Selected = FDynamicMaterialEditorStyle::GetBrush("LayerView.Row.Handle.Left.Select");
-	static const FSlateBrush* const Hovered = FDynamicMaterialEditorStyle::GetBrush("LayerView.Row.Handle.Left.Hover");
-	static const FSlateBrush* const SelectedHovered = FDynamicMaterialEditorStyle::GetBrush("LayerView.Row.Handle.Left.Select.Hover");
+	static const FSlateBrush* const Default = FDynamicMaterialEditorStyle::Get().GetBrush("LayerView.Row.Handle.Left");
+	static const FSlateBrush* const Selected = FDynamicMaterialEditorStyle::Get().GetBrush("LayerView.Row.Handle.Left.Select");
+	static const FSlateBrush* const Hovered = FDynamicMaterialEditorStyle::Get().GetBrush("LayerView.Row.Handle.Left.Hover");
+	static const FSlateBrush* const SelectedHovered = FDynamicMaterialEditorStyle::Get().GetBrush("LayerView.Row.Handle.Left.Select.Hover");
 
 	const bool bSelected = IsSelected();
 	const bool bHovered = IsHovered();
@@ -734,8 +734,8 @@ EVisibility SDMMaterialSlotLayerItem::GetLayerLinkToggleButtonVisibility() const
 
 const FSlateBrush* SDMMaterialSlotLayerItem::GetLayerLinkToggleButtonImage() const
 {
-	static const FSlateBrush* const Unlinked = FDynamicMaterialEditorStyle::GetBrush("Icons.Stage.ChainUnlinked.Vertical");
-	static const FSlateBrush* const Linked = FDynamicMaterialEditorStyle::GetBrush("Icons.Stage.ChainLinked.Vertical");
+	static const FSlateBrush* const Unlinked = FDynamicMaterialEditorStyle::Get().GetBrush("Icons.Stage.ChainUnlinked.Vertical");
+	static const FSlateBrush* const Linked = FDynamicMaterialEditorStyle::Get().GetBrush("Icons.Stage.ChainLinked.Vertical");
 
 	return AreStagesLinked() ? Linked : Unlinked;
 }
@@ -777,8 +777,8 @@ FReply SDMMaterialSlotLayerItem::OnLayerLinkToggleButton()
 
 const FSlateBrush* SDMMaterialSlotLayerItem::GetStageToggleButtonImage(EDMMaterialLayerStage InLayerStage) const
 {
-	static const FSlateBrush* Disabled = FDynamicMaterialEditorStyle::GetBrush("Icons.Stage.Disabled");
-	static const FSlateBrush* Enabled = FDynamicMaterialEditorStyle::GetBrush("Icons.Stage.Enabled");
+	static const FSlateBrush* Disabled = FDynamicMaterialEditorStyle::Get().GetBrush("Icons.Stage.Disabled");
+	static const FSlateBrush* Enabled = FDynamicMaterialEditorStyle::Get().GetBrush("Icons.Stage.Enabled");
 
 	if (const UDMMaterialLayerObject* Layer = GetLayer())
 	{
diff --git a/Engine/Plugins/Experimental/DynamicMaterial/Source/DynamicMaterialEditor/Private/UI/Widgets/Editor/SlotEditor/SDMMaterialStage.cpp b/Engine/Plugins/Experimental/DynamicMaterial/Source/DynamicMaterialEditor/Private/UI/Widgets/Editor/SlotEditor/SDMMaterialStage.cpp
--- a/Engine/Plugins/Experimental/DynamicMaterial/Source/DynamicMaterialEditor/Private/UI/Widgets/Editor/SlotEditor/SDMMaterialStage.cpp
+++ b/Engine/Plugins/Experimental/DynamicMaterial/Source/DynamicMaterialEditor/Private/UI/Widgets/Editor/SlotEditor/SDMMaterialStage.cpp
@@ -228,7 +228,7 @@ const FSlateBrush* SDMMaterialStage::GetBorderBrush() const
 
 	if (!Stage)
 	{
-		return FDynamicMaterialEditorStyle::GetBrush(TEXT("Stage.Inactive"));
+		return FDynamicMaterialEditorStyle::Get().GetBrush(TEXT("Stage.Inactive"));
 	}
 
 	FString BrushName = Stage->IsEnabled()
@@ -245,7 +245,7 @@ const FSlateBrush* SDMMaterialStage::GetBorderBrush() const
 		BrushName.Append(".Hover");
 	}
 
-	return FDynamicMaterialEditorStyle::GetBrush(*BrushName);
+	return FDynamicMaterialEditorStyle::Get().GetBrush(*BrushName);
 }
 
 EVisibility SDMMaterialStage::GetDisabledOverlayVisibility() const
