Material Designer
- Added type icon system.
- Added type icons to layer list.
- Added type icons to source menus.
- Added toggle layer to right click on layer (not stage).
- Fixed add new layer button to work when no layer is selected.
- Removed some extraneous wording from tooltips.

diff --git a/Engine/Plugins/Experimental/DynamicMaterial/Source/DynamicMaterial/Private/Components/DMMaterialComponent.cpp b/Engine/Plugins/Experimental/DynamicMaterial/Source/DynamicMaterial/Private/Components/DMMaterialComponent.cpp
--- a/Engine/Plugins/Experimental/DynamicMaterial/Source/DynamicMaterial/Private/Components/DMMaterialComponent.cpp
+++ b/Engine/Plugins/Experimental/DynamicMaterial/Source/DynamicMaterial/Private/Components/DMMaterialComponent.cpp
@@ -5,6 +5,11 @@
 #include "DynamicMaterialModule.h"
 #include "HAL/PlatformTime.h"
 
+#if WITH_EDITOR
+#include "Styling/SlateIconFinder.h"
+#include "Textures/SlateIcon.h"
+#endif
+
 #define LOCTEXT_NAMESPACE "DMMaterialComponent"
 
 double UDMMaterialComponent::MinCleanTime = FPlatformTime::Seconds();
@@ -155,6 +160,19 @@ FText UDMMaterialComponent::GetComponentDescription() const
 	return GetClass()->GetDisplayNameText();
 }
 
+FSlateIcon UDMMaterialComponent::GetComponentIcon() const
+{
+	FSlateIcon Icon = FSlateIconFinder::FindIconForClass(GetClass());
+
+	if (Icon.IsSet())
+	{
+		return Icon;
+	}
+
+	// Fall back to a default icon.
+	return FSlateIconFinder::FindIconForClass(UDMMaterialComponent::StaticClass());
+}
+
 bool UDMMaterialComponent::CanClean()
 {
 	return (FPlatformTime::Seconds() >= MinCleanTime);
diff --git a/Engine/Plugins/Experimental/DynamicMaterial/Source/DynamicMaterial/Private/DMValueDefinition.cpp b/Engine/Plugins/Experimental/DynamicMaterial/Source/DynamicMaterial/Private/DMValueDefinition.cpp
--- a/Engine/Plugins/Experimental/DynamicMaterial/Source/DynamicMaterial/Private/DMValueDefinition.cpp
+++ b/Engine/Plugins/Experimental/DynamicMaterial/Source/DynamicMaterial/Private/DMValueDefinition.cpp
@@ -13,6 +13,10 @@
 #include "Components/MaterialValues/DMMaterialValueTexture.h"
 #include "Containers/Map.h"
 
+#if WITH_EDITOR
+#include "Textures/SlateIcon.h"
+#endif
+
 #define LOCTEXT_NAMESPACE "DMValueDefinition"
 
 namespace UE::MaterialDesigner::Private
@@ -273,4 +277,21 @@ const FDMValueDefinition& UDMValueDefinitionLibrary::GetValueDefinition(EDMValue
 	return TypeDefinitions[InValueType]; //-V558
 }
 
+#if WITH_EDITOR
+FSlateIcon UDMValueDefinitionLibrary::GetValueIcon(EDMValueType InType)
+{
+	using namespace UE::MaterialDesigner::Private;
+
+	if (UClass* ValueClass = TypeDefinitions[InType].GetValueClass())
+	{
+		if (UDMMaterialValue* ValueCDO = ValueClass->GetDefaultObject<UDMMaterialValue>())
+		{
+			return ValueCDO->GetComponentIcon();
+		}
+	}
+
+	return GetDefault<UDMMaterialComponent>()->GetComponentIcon();
+}
+#endif
+
 #undef LOCTEXT_NAMESPACE
diff --git a/Engine/Plugins/Experimental/DynamicMaterial/Source/DynamicMaterial/Public/Components/DMMaterialComponent.h b/Engine/Plugins/Experimental/DynamicMaterial/Source/DynamicMaterial/Public/Components/DMMaterialComponent.h
--- a/Engine/Plugins/Experimental/DynamicMaterial/Source/DynamicMaterial/Public/Components/DMMaterialComponent.h
+++ b/Engine/Plugins/Experimental/DynamicMaterial/Source/DynamicMaterial/Public/Components/DMMaterialComponent.h
@@ -9,8 +9,9 @@
 
 class UDynamicMaterialModel;
 enum class EDMUpdateType : uint8;
-struct FDMComponentPathSegment;
 struct FDMComponentPath;
+struct FDMComponentPathSegment;
+struct FSlateIcon;
 
 UENUM(BlueprintType)
 enum class EDMComponentLifetimeState : uint8
@@ -85,6 +86,9 @@ public:
 	UFUNCTION(BlueprintPure, Category = "Material Designer")
 	DYNAMICMATERIAL_API virtual FText GetComponentDescription() const;
 
+	/** Returns a brush which indicates this component (type) */
+	DYNAMICMATERIAL_API virtual FSlateIcon GetComponentIcon() const;
+
 	/** Returns true if this component has been marked dirty. */
 	DYNAMICMATERIAL_API bool NeedsClean();
 
diff --git a/Engine/Plugins/Experimental/DynamicMaterial/Source/DynamicMaterial/Public/DMValueDefinition.h b/Engine/Plugins/Experimental/DynamicMaterial/Source/DynamicMaterial/Public/DMValueDefinition.h
--- a/Engine/Plugins/Experimental/DynamicMaterial/Source/DynamicMaterial/Public/DMValueDefinition.h
+++ b/Engine/Plugins/Experimental/DynamicMaterial/Source/DynamicMaterial/Public/DMValueDefinition.h
@@ -10,6 +10,10 @@
 
 class UDMMaterialValue;
 
+#if WITH_EDITOR
+struct FSlateIcon;
+#endif
+
 /**
  * Stores information about basic value types, such as EDMValueType::Float1.
  */
@@ -111,4 +115,8 @@ public:
 	/** Converts a number of floats into the value type */
 	UFUNCTION(BlueprintPure, CallInEditor, Category = "Material Designer")
 	static DYNAMICMATERIAL_API const FDMValueDefinition& GetTypeForFloatCount(int32 InFloatCount);
+
+#if WITH_EDITOR
+	static DYNAMICMATERIAL_API FSlateIcon GetValueIcon(EDMValueType InType);
+#endif
 };
diff --git a/Engine/Plugins/Experimental/DynamicMaterial/Source/DynamicMaterialEditor/Private/Components/DMMaterialStage.cpp b/Engine/Plugins/Experimental/DynamicMaterial/Source/DynamicMaterialEditor/Private/Components/DMMaterialStage.cpp
--- a/Engine/Plugins/Experimental/DynamicMaterial/Source/DynamicMaterialEditor/Private/Components/DMMaterialStage.cpp
+++ b/Engine/Plugins/Experimental/DynamicMaterial/Source/DynamicMaterialEditor/Private/Components/DMMaterialStage.cpp
@@ -298,6 +298,16 @@ FText UDMMaterialStage::GetComponentDescription() const
 	return LOCTEXT("StageDescription", "Material Stage");
 }
 
+FSlateIcon UDMMaterialStage::GetComponentIcon() const
+{
+	if (Source && Source->IsComponentValid())
+	{
+		return Source->GetComponentIcon();
+	}
+
+	return Super::GetComponentIcon();
+}
+
 EDMValueType UDMMaterialStage::GetSourceType(const FDMMaterialStageConnectorChannel& InChannel) const
 {
 	if (InChannel.SourceIndex == FDMMaterialStageConnectorChannel::PREVIOUS_STAGE)
diff --git a/Engine/Plugins/Experimental/DynamicMaterial/Source/DynamicMaterialEditor/Private/Components/DMMaterialStageBlend.cpp b/Engine/Plugins/Experimental/DynamicMaterial/Source/DynamicMaterialEditor/Private/Components/DMMaterialStageBlend.cpp
--- a/Engine/Plugins/Experimental/DynamicMaterial/Source/DynamicMaterialEditor/Private/Components/DMMaterialStageBlend.cpp
+++ b/Engine/Plugins/Experimental/DynamicMaterial/Source/DynamicMaterialEditor/Private/Components/DMMaterialStageBlend.cpp
@@ -395,6 +395,16 @@ bool UDMMaterialStageBlend::GenerateStagePreviewMaterial(UDMMaterialStage* InSta
 	return true;
 }
 
+FSlateIcon UDMMaterialStageBlend::GetComponentIcon() const
+{
+	if (UDMMaterialStageInput* InputBValue = GetInputB())
+	{
+		return InputBValue->GetComponentIcon();
+	}
+
+	return Super::GetComponentIcon();
+}
+
 void UDMMaterialStageBlend::Update(UDMMaterialComponent* InSource, EDMUpdateType InUpdateType)
 {
 	Super::Update(InSource, InUpdateType);
diff --git a/Engine/Plugins/Experimental/DynamicMaterial/Source/DynamicMaterialEditor/Private/Components/DMMaterialStageThroughputLayerBlend.cpp b/Engine/Plugins/Experimental/DynamicMaterial/Source/DynamicMaterialEditor/Private/Components/DMMaterialStageThroughputLayerBlend.cpp
--- a/Engine/Plugins/Experimental/DynamicMaterial/Source/DynamicMaterialEditor/Private/Components/DMMaterialStageThroughputLayerBlend.cpp
+++ b/Engine/Plugins/Experimental/DynamicMaterial/Source/DynamicMaterialEditor/Private/Components/DMMaterialStageThroughputLayerBlend.cpp
@@ -1134,6 +1134,16 @@ void UDMMaterialStageThroughputLayerBlend::SetMaskChannelOverride(EAvaColorChann
 	Update(this, EDMUpdateType::Structure);
 }
 
+FSlateIcon UDMMaterialStageThroughputLayerBlend::GetComponentIcon() const
+{
+	if (UDMMaterialStageInput* InputMaskValue = GetInputMask())
+	{
+		return InputMaskValue->GetComponentIcon();
+	}
+
+	return Super::GetComponentIcon();
+}
+
 bool UDMMaterialStageThroughputLayerBlend::CanUseMaskChannelOverride() const
 {
 	return GetDefaultMaskChannelOverrideOutputIndex() != INDEX_NONE;
diff --git a/Engine/Plugins/Experimental/DynamicMaterial/Source/DynamicMaterialEditor/Private/Components/MaterialStageInputs/DMMSISlot.cpp b/Engine/Plugins/Experimental/DynamicMaterial/Source/DynamicMaterialEditor/Private/Components/MaterialStageInputs/DMMSISlot.cpp
--- a/Engine/Plugins/Experimental/DynamicMaterial/Source/DynamicMaterialEditor/Private/Components/MaterialStageInputs/DMMSISlot.cpp
+++ b/Engine/Plugins/Experimental/DynamicMaterial/Source/DynamicMaterialEditor/Private/Components/MaterialStageInputs/DMMSISlot.cpp
@@ -149,6 +149,19 @@ FText UDMMaterialStageInputSlot::GetComponentDescription() const
 	return Super::GetComponentDescription();
 }
 
+FSlateIcon UDMMaterialStageInputSlot::GetComponentIcon() const
+{
+	if (Slot)
+	{
+		if (Slot->IsComponentValid())
+		{
+			return Slot->GetComponentIcon();
+		}
+	}
+
+	return Super::GetComponentIcon();
+}
+
 FText UDMMaterialStageInputSlot::GetChannelDescription(const FDMMaterialStageConnectorChannel& Channel)
 {
 	return GetComponentDescription();
diff --git a/Engine/Plugins/Experimental/DynamicMaterial/Source/DynamicMaterialEditor/Private/Components/MaterialStageInputs/DMMSITextureUV.cpp b/Engine/Plugins/Experimental/DynamicMaterial/Source/DynamicMaterialEditor/Private/Components/MaterialStageInputs/DMMSITextureUV.cpp
--- a/Engine/Plugins/Experimental/DynamicMaterial/Source/DynamicMaterialEditor/Private/Components/MaterialStageInputs/DMMSITextureUV.cpp
+++ b/Engine/Plugins/Experimental/DynamicMaterial/Source/DynamicMaterialEditor/Private/Components/MaterialStageInputs/DMMSITextureUV.cpp
@@ -110,7 +110,12 @@ UDMMaterialStageInputTextureUV* UDMMaterialStageInputTextureUV::ChangeStageInput
 
 FText UDMMaterialStageInputTextureUV::GetComponentDescription() const
 {
-	return LOCTEXT("TexUV", "Tex UV");
+	return LOCTEXT("TexureUV", "Texture UV");
+}
+
+FSlateIcon UDMMaterialStageInputTextureUV::GetComponentIcon() const
+{
+	return GetDefault<UDMTextureUV>()->GetComponentIcon();
 }
 
 FText UDMMaterialStageInputTextureUV::GetChannelDescription(const FDMMaterialStageConnectorChannel& Channel)
diff --git a/Engine/Plugins/Experimental/DynamicMaterial/Source/DynamicMaterialEditor/Private/Components/MaterialStageInputs/DMMSIThroughput.cpp b/Engine/Plugins/Experimental/DynamicMaterial/Source/DynamicMaterialEditor/Private/Components/MaterialStageInputs/DMMSIThroughput.cpp
--- a/Engine/Plugins/Experimental/DynamicMaterial/Source/DynamicMaterialEditor/Private/Components/MaterialStageInputs/DMMSIThroughput.cpp
+++ b/Engine/Plugins/Experimental/DynamicMaterial/Source/DynamicMaterialEditor/Private/Components/MaterialStageInputs/DMMSIThroughput.cpp
@@ -248,7 +248,7 @@ FText UDMMaterialStageInputThroughput::GetComponentDescription() const
 	if (TSubclassOf<UDMMaterialStageThroughput> MaterialStageThroughputClass = GetMaterialStageThroughputClass())
 	{
 		check(MaterialStageThroughputClass.Get());
-		UDMMaterialStageThroughput* ThroughputCDO = Cast<UDMMaterialStageThroughput>(MaterialStageThroughputClass->GetDefaultObject(true));
+		UDMMaterialStageThroughput* ThroughputCDO = Cast<UDMMaterialStageThroughput>(MaterialStageThroughputClass->GetDefaultObject(/* Create if needed */ true));
 
 		return ThroughputCDO->GetComponentDescription();
 	}
@@ -256,6 +256,24 @@ FText UDMMaterialStageInputThroughput::GetComponentDescription() const
 	return Super::GetComponentDescription();
 }
 
+FSlateIcon UDMMaterialStageInputThroughput::GetComponentIcon() const
+{
+	if (UDMMaterialStageThroughput* Throughput = GetMaterialStageThroughput())
+	{
+		return Throughput->GetComponentIcon();
+	}
+
+	if (TSubclassOf<UDMMaterialStageThroughput> MaterialStageThroughputClass = GetMaterialStageThroughputClass())
+	{
+		check(MaterialStageThroughputClass.Get());
+		UDMMaterialStageThroughput* ThroughputCDO = Cast<UDMMaterialStageThroughput>(MaterialStageThroughputClass->GetDefaultObject(/* Create if needed */ true));
+
+		return ThroughputCDO->GetComponentIcon();
+	}
+
+	return Super::GetComponentIcon();
+}
+
 FText UDMMaterialStageInputThroughput::GetChannelDescription(const FDMMaterialStageConnectorChannel& Channel)
 {
 	if (TSubclassOf<UDMMaterialStageThroughput> MaterialStageThroughputClass = GetMaterialStageThroughputClass())
diff --git a/Engine/Plugins/Experimental/DynamicMaterial/Source/DynamicMaterialEditor/Private/Components/MaterialStageInputs/DMMSIValue.cpp b/Engine/Plugins/Experimental/DynamicMaterial/Source/DynamicMaterialEditor/Private/Components/MaterialStageInputs/DMMSIValue.cpp
--- a/Engine/Plugins/Experimental/DynamicMaterial/Source/DynamicMaterialEditor/Private/Components/MaterialStageInputs/DMMSIValue.cpp
+++ b/Engine/Plugins/Experimental/DynamicMaterial/Source/DynamicMaterialEditor/Private/Components/MaterialStageInputs/DMMSIValue.cpp
@@ -365,6 +365,16 @@ FText UDMMaterialStageInputValue::GetComponentDescription() const
 	}
 }
 
+FSlateIcon UDMMaterialStageInputValue::GetComponentIcon() const
+{
+	if (Value && Value->IsComponentValid())
+	{
+		return Value->GetComponentIcon();
+	}
+
+	return Super::GetComponentIcon();
+}
+
 FText UDMMaterialStageInputValue::GetChannelDescription(const FDMMaterialStageConnectorChannel& Channel)
 {
 	if (!Value->IsComponentValid())
diff --git a/Engine/Plugins/Experimental/DynamicMaterial/Source/DynamicMaterialEditor/Private/DynamicMaterialEditorStyle.cpp b/Engine/Plugins/Experimental/DynamicMaterial/Source/DynamicMaterialEditor/Private/DynamicMaterialEditorStyle.cpp
--- a/Engine/Plugins/Experimental/DynamicMaterial/Source/DynamicMaterialEditor/Private/DynamicMaterialEditorStyle.cpp
+++ b/Engine/Plugins/Experimental/DynamicMaterial/Source/DynamicMaterialEditor/Private/DynamicMaterialEditorStyle.cpp
@@ -6,23 +6,11 @@
 #include "DMDefs.h"
 #include "Interfaces/IPluginManager.h"
 #include "Styling/SlateStyleMacros.h"
+#include "Styling/SlateStyleRegistry.h"
 #include "Styling/StyleColors.h"
 
 namespace UE::DynamicMaterialEditor::Private
 {
-	const TMap<EDMValueType, FName> TypeIcons = {
-		{EDMValueType::VT_None,        TEXT("Icons.Type.None")},
-		{EDMValueType::VT_Bool,        TEXT("Icons.Type.Bool")},
-		{EDMValueType::VT_Float1,      TEXT("Icons.Type.Float1")},
-		{EDMValueType::VT_Float2,      TEXT("Icons.Type.Float2")},
-		{EDMValueType::VT_Float3_RPY,  TEXT("Icons.Type.Float3_RPY")},
-		{EDMValueType::VT_Float3_RGB,  TEXT("Icons.Type.Float3_RGB")},
-		{EDMValueType::VT_Float3_XYZ,  TEXT("Icons.Type.Float3_XYZ")},
-		{EDMValueType::VT_Float4_RGBA, TEXT("Icons.Type.Float4_RGBA")},
-		{EDMValueType::VT_Float_Any,   TEXT("Icons.Type.Float_Any")},
-		{EDMValueType::VT_Texture,     TEXT("Icons.Type.Texture")},
-	};
-
 	FLinearColor ReplaceColorAlpha(const FLinearColor& InColor, const float InNewAlpha)
 	{
 		FLinearColor OutColor(InColor);
@@ -32,42 +20,36 @@ namespace UE::DynamicMaterialEditor::Private
 }
 
 FDynamicMaterialEditorStyle::FDynamicMaterialEditorStyle()
-	: FSlateStyleSet(TEXT("DynamicMaterialEditor"))
+	: FSlateStyleSet(UE_MODULE_NAME)
 {
-	SetupGenericStyles();
+	SetupGeneralStyles();
 	SetupStageStyles();
 	SetupLayerViewStyles();
 	SetupLayerViewItemHandleStyles();
 	SetupEffectsViewStyles();
 	SetupTextStyles();
+	SetupComponentIcons();
+
+	FSlateStyleRegistry::RegisterSlateStyle(*this);
 }
 
-FName FDynamicMaterialEditorStyle::GetBrushNameForType(EDMValueType InType) const
+FDynamicMaterialEditorStyle::~FDynamicMaterialEditorStyle()
 {
-	if (UE::DynamicMaterialEditor::Private::TypeIcons.Contains(InType))
-	{
-		return UE::DynamicMaterialEditor::Private::TypeIcons[InType];
-	}
-
-	checkNoEntry();
-
-	static FName NoType = TEXT("Icons.Type.None");
-	return NoType;
+	FSlateStyleRegistry::UnRegisterSlateStyle(*this);
 }
 
-void FDynamicMaterialEditorStyle::SetupGenericStyles()
+void FDynamicMaterialEditorStyle::SetupGeneralStyles()
 {
 	using namespace UE::DynamicMaterialEditor::Private;
 
-	TSharedPtr<IPlugin> Plugin = IPluginManager::Get().FindPlugin(TEXT("DynamicMaterial"));
+	TSharedPtr<IPlugin> Plugin = IPluginManager::Get().FindPlugin(UE_PLUGIN_NAME);
 	check(Plugin.IsValid());
 
 	SetContentRoot(FPaths::Combine(Plugin->GetBaseDir(), TEXT("Resources")));
+	SetCoreContentRoot(FPaths::EngineContentDir() / TEXT("Editor/Slate")); // This is the engine's content root.
 
 	////////////////////////////////////////////////////////////////////////////////////////////////////
 	// Color Styles
-	const FLinearColor EngineSelectColor = FStyleColors::Primary.GetSpecifiedColor();
-	const FLinearColor EngineSelectHoverColor = FStyleColors::PrimaryHover.GetSpecifiedColor();
 	const FLinearColor EngineSelectPressColor = FStyleColors::PrimaryPress.GetSpecifiedColor();
 
 	const FLinearColor SelectColor = ReplaceColorAlpha(FStyleColors::Select.GetSpecifiedColor(), 0.9f);
@@ -82,17 +64,6 @@ void FDynamicMaterialEditorStyle::SetupGenericStyles()
 	// Brush Styles
 	Set("Icons.Menu.Dropdown", new IMAGE_BRUSH_SVG("Icons/MenuDropdown", CoreStyleConstants::Icon16x16));
 
-	Set("Icons.Type.None", new IMAGE_BRUSH("Icons/ValueTypes/None", CoreStyleConstants::Icon12x12));
-	Set("Icons.Type.Bool", new IMAGE_BRUSH("Icons/ValueTypes/Bool", CoreStyleConstants::Icon12x12));
-	Set("Icons.Type.Float1", new IMAGE_BRUSH("Icons/ValueTypes/Float1", CoreStyleConstants::Icon12x12));
-	Set("Icons.Type.Float2", new IMAGE_BRUSH("Icons/ValueTypes/Float2", CoreStyleConstants::Icon12x12));
-	Set("Icons.Type.Float3_RPY", new IMAGE_BRUSH("Icons/ValueTypes/Float3_RPY", CoreStyleConstants::Icon12x12));
-	Set("Icons.Type.Float3_RGB", new IMAGE_BRUSH("Icons/ValueTypes/Float3_RGB", CoreStyleConstants::Icon12x12));
-	Set("Icons.Type.Float3_XYZ", new IMAGE_BRUSH("Icons/ValueTypes/Float3_XYZ", CoreStyleConstants::Icon12x12));
-	Set("Icons.Type.Float4_RGBA", new IMAGE_BRUSH("Icons/ValueTypes/Float4_RGBA", CoreStyleConstants::Icon12x12));
-	Set("Icons.Type.Float_Any", new IMAGE_BRUSH("Icons/ValueTypes/Float_Any", CoreStyleConstants::Icon12x12));
-	Set("Icons.Type.Texture", new IMAGE_BRUSH("Icons/ValueTypes/Texture", CoreStyleConstants::Icon12x12));
-
 	Set("Icons.Material.DefaultLit", new IMAGE_BRUSH("Icons/EditorIcons/MaterialTypeDefaultLit", CoreStyleConstants::Icon32x32));
 	Set("Icons.Material.Unlit", new IMAGE_BRUSH("Icons/EditorIcons/MaterialTypeUnlit", CoreStyleConstants::Icon32x32));
 
@@ -470,3 +441,35 @@ void FDynamicMaterialEditorStyle::SetupTextStyles()
 	Set("Font.Stage.Details.Small.Bold", FTextBlockStyle(NormalTextStyle)
 		.SetFont(IDetailLayoutBuilder::GetDetailFontBold()));
 }
+
+void FDynamicMaterialEditorStyle::SetupComponentIcons()
+{
+	Set(TEXT("ClassIcon.DMMaterialComponent"), new CORE_IMAGE_BRUSH("Icons/icon_help_16x", CoreStyleConstants::Icon16x16));
+
+	Set(TEXT("ClassIcon.DMMaterialValueBool"),         new IMAGE_BRUSH("Icons/ValueTypes/Bool", CoreStyleConstants::Icon16x16));
+	Set(TEXT("ClassIcon.DMMaterialValueColorAtlas"),   new CORE_IMAGE_BRUSH("Icons/icon_help_16x", CoreStyleConstants::Icon16x16));
+	Set(TEXT("ClassIcon.DMMaterialValueFloat1"),       new IMAGE_BRUSH("Icons/ValueTypes/Float1", CoreStyleConstants::Icon16x16));
+	Set(TEXT("ClassIcon.DMMaterialValueFloat2"),       new IMAGE_BRUSH("Icons/ValueTypes/Float2", CoreStyleConstants::Icon16x16));
+	Set(TEXT("ClassIcon.DMMaterialValueFloat3RGB"),    new FSlateImageBrush(FPaths::EngineContentDir() + TEXT("Slate/Common/ColorPicker_Mode_16x.png"), CoreStyleConstants::Icon16x16));
+	Set(TEXT("ClassIcon.DMMaterialValueFloat3RPY"),    new CORE_IMAGE_BRUSH("Icons/icon_ClockwiseRotation_16x", CoreStyleConstants::Icon16x16));
+	Set(TEXT("ClassIcon.DMMaterialValueFloat3XYZ"),    new CORE_IMAGE_BRUSH("Icons/Movable_16x", CoreStyleConstants::Icon16x16));
+	Set(TEXT("ClassIcon.DMMaterialValueFloat4"),       new FSlateImageBrush(FPaths::EngineContentDir() + TEXT("Slate/Common/ColorPicker_Mode_16x.png"), CoreStyleConstants::Icon16x16));
+	Set(TEXT("ClassIcon.DMMaterialValueRenderTarget"), new CORE_IMAGE_BRUSH("Icons/AssetIcons/TextureRenderTarget2D_16x", CoreStyleConstants::Icon16x16));
+	Set(TEXT("ClassIcon.DMMaterialValueTexture"),      new IMAGE_BRUSH("Icons/ValueTypes/Texture", CoreStyleConstants::Icon16x16));
+
+	Set(TEXT("ClassIcon.DMRenderTargetRenderer"),           new CORE_IMAGE_BRUSH("Icons/icon_help_16x", CoreStyleConstants::Icon16x16));
+	Set(TEXT("ClassIcon.DMRenderTargetWidgetRendererBase"), new CORE_IMAGE_BRUSH("Icons/AssetIcons/WidgetBlueprint_16x", CoreStyleConstants::Icon16x16));
+	Set(TEXT("ClassIcon.DMRenderTargetTextRenderer"),       new CORE_IMAGE_BRUSH("Icons/icon_help_16x", CoreStyleConstants::Icon16x16));
+	Set(TEXT("ClassIcon.DMRenderTargetUMGWidgetRenderer"),  new CORE_IMAGE_BRUSH("Icons/AssetIcons/WidgetBlueprint_16x", CoreStyleConstants::Icon16x16));
+
+	Set(TEXT("ClassIcon.DMMaterialStageGradientLinear"), new CORE_IMAGE_BRUSH("Icons/icon_help_16x", CoreStyleConstants::Icon16x16));
+	Set(TEXT("ClassIcon.DMMaterialStageGradientRadial"), new CORE_IMAGE_BRUSH("Icons/icon_help_16x", CoreStyleConstants::Icon16x16));
+
+	Set(TEXT("ClassIcon.DMMaterialStageExpressionSceneTexture"),           new CORE_IMAGE_BRUSH("Icons/AssetIcons/PostProcessVolume_16x", CoreStyleConstants::Icon16x16));
+	Set(TEXT("ClassIcon.DMMaterialStageExpressionTextureSampleBase"),      new IMAGE_BRUSH("Icons/ValueTypes/Texture", CoreStyleConstants::Icon16x16));
+	Set(TEXT("ClassIcon.DMMaterialStageExpressionTextureSampleEdgeColor"), new IMAGE_BRUSH("Icons/ClassIcons/TextureEdgeColor", CoreStyleConstants::Icon16x16));
+	Set(TEXT("ClassIcon.DMMaterialStageExpressionWorldPositionNoise"),     new CORE_IMAGE_BRUSH("Icons/icon_help_16x", CoreStyleConstants::Icon16x16));
+
+	Set(TEXT("ClassIcon.DMMaterialStageFunction"),      new CORE_IMAGE_BRUSH("Icons/AssetIcons/MaterialFunction_16x", CoreStyleConstants::Icon16x16));
+	Set(TEXT("ClassIcon.DMMaterialStageInputFunction"), new CORE_IMAGE_BRUSH("Icons/AssetIcons/MaterialFunction_16x", CoreStyleConstants::Icon16x16));
+}
diff --git a/Engine/Plugins/Experimental/DynamicMaterial/Source/DynamicMaterialEditor/Private/DynamicMaterialEditorStyle.h b/Engine/Plugins/Experimental/DynamicMaterial/Source/DynamicMaterialEditor/Private/DynamicMaterialEditorStyle.h
--- a/Engine/Plugins/Experimental/DynamicMaterial/Source/DynamicMaterialEditor/Private/DynamicMaterialEditorStyle.h
+++ b/Engine/Plugins/Experimental/DynamicMaterial/Source/DynamicMaterialEditor/Private/DynamicMaterialEditorStyle.h
@@ -17,15 +17,14 @@ public:
 	}
 
 	FDynamicMaterialEditorStyle();
-	virtual ~FDynamicMaterialEditorStyle() override = default;
-
-	FName GetBrushNameForType(EDMValueType InType) const;
+	virtual ~FDynamicMaterialEditorStyle() override;
 
 private:
-	void SetupGenericStyles();
+	void SetupGeneralStyles();
 	void SetupStageStyles();
 	void SetupLayerViewStyles();
 	void SetupLayerViewItemHandleStyles();
 	void SetupEffectsViewStyles();
 	void SetupTextStyles();
+	void SetupComponentIcons();
 };
diff --git a/Engine/Plugins/Experimental/DynamicMaterial/Source/DynamicMaterialEditor/Private/UI/Menus/DMMaterialSlotLayerMenus.cpp b/Engine/Plugins/Experimental/DynamicMaterial/Source/DynamicMaterialEditor/Private/UI/Menus/DMMaterialSlotLayerMenus.cpp
--- a/Engine/Plugins/Experimental/DynamicMaterial/Source/DynamicMaterialEditor/Private/UI/Menus/DMMaterialSlotLayerMenus.cpp
+++ b/Engine/Plugins/Experimental/DynamicMaterial/Source/DynamicMaterialEditor/Private/UI/Menus/DMMaterialSlotLayerMenus.cpp
@@ -7,11 +7,15 @@
 #include "Components/DMMaterialSlot.h"
 #include "Components/DMMaterialStageBlend.h"
 #include "Components/DMMaterialStageExpression.h"
+#include "Components/DMMaterialStageFunction.h"
 #include "Components/DMMaterialStageGradient.h"
 #include "Components/DMMaterialValue.h"
+#include "Components/MaterialStageExpressions/DMMSESceneTexture.h"
 #include "Components/MaterialStageExpressions/DMMSETextureSample.h"
 #include "Components/MaterialStageExpressions/DMMSETextureSampleEdgeColor.h"
 #include "Components/MaterialStageExpressions/DMMSEWorldPositionNoise.h"
+#include "Components/MaterialValues/DMMaterialValueColorAtlas.h"
+#include "Components/MaterialValues/DMMaterialValueFloat3RGB.h"
 #include "Components/RenderTargetRenderers/DMRenderTargetTextRenderer.h"
 #include "Components/RenderTargetRenderers/DMRenderTargetUMGWidgetRenderer.h"
 #include "DMDefs.h"
@@ -19,9 +23,11 @@
 #include "DMValueDefinition.h"
 #include "DynamicMaterialEditorCommands.h"
 #include "DynamicMaterialEditorModule.h"
+#include "DynamicMaterialEditorStyle.h"
 #include "Framework/Commands/GenericCommands.h"
 #include "Model/DynamicMaterialModel.h"
 #include "Model/DynamicMaterialModelEditorOnlyData.h"
+#include "ScopedTransaction.h"
 #include "Styling/SlateIconFinder.h"
 #include "Templates/SharedPointer.h"
 #include "ToolMenus.h"
@@ -40,7 +46,7 @@ namespace UE::DynamicMaterialEditor::Private
 	static const FName GlobalValuesSectionName = TEXT("GlobalValues");
 }
 
-TSharedRef<SWidget> FDMMaterialSlotLayerMenus::GenerateSlotLayerMenu(const TSharedPtr<SDMMaterialSlotEditor>& InSlotWidget, UDMMaterialLayerObject* InLayerObject)
+TSharedRef<SWidget> FDMMaterialSlotLayerMenus::GenerateSlotLayerMenu(const TSharedPtr<SDMMaterialSlotEditor>& InSlotWidget, UDMMaterialLayerObject* InLayer)
 {
 	using namespace UE::DynamicMaterialEditor::Private;
 
@@ -63,7 +69,7 @@ TSharedRef<SWidget> FDMMaterialSlotLayerMenus::GenerateSlotLayerMenu(const TShar
 		NewToolMenu->AddDynamicSection(NAME_None, FNewToolMenuDelegate::CreateStatic(&FDMMaterialSlotLayerMenus::AddLayerModifySection));
 	}
 
-	FToolMenuContext MenuContext(UDMMenuContext::CreateEditor(InSlotWidget->GetEditorWidget()));
+	FToolMenuContext MenuContext(UDMMenuContext::CreateLayer(InSlotWidget->GetEditorWidget(), InLayer));
 
 	return UToolMenus::Get()->GenerateWidget(SlotLayerMenuName, MenuContext);
 }
@@ -72,7 +78,7 @@ void FDMMaterialSlotLayerMenus::AddAddLayerSection(UToolMenu* InMenu)
 {
 	using namespace UE::DynamicMaterialEditor::Private;
 
-	if (!IsValid(InMenu))
+	if (!IsValid(InMenu) || InMenu->ContainsSection(SlotLayerAddSectionName))
 	{
 		return;
 	}
@@ -118,7 +124,7 @@ void FDMMaterialSlotLayerMenus::AddAddLayerSection(UToolMenu* InMenu)
 		NAME_None,
 		LOCTEXT("AddTextureSample", "Texture"),
 		LOCTEXT("AddTextureSampleTooltip", "Add a Material Stage based on a Texture."),
-		FSlateIcon(),
+		GetDefault<UDMMaterialStageExpressionTextureSample>()->GetComponentIcon(),
 		FUIAction(FExecuteAction::CreateWeakLambda(
 			Slot,
 			[Slot, ExpressionClass = TSubclassOf<UDMMaterialStageExpression>(UDMMaterialStageExpressionTextureSample::StaticClass())]
@@ -132,7 +138,7 @@ void FDMMaterialSlotLayerMenus::AddAddLayerSection(UToolMenu* InMenu)
 		NAME_None,
 		LOCTEXT("AddColor", "Solid Color"),
 		LOCTEXT("AddColorTooltip", "Add a new Material Layer with a solid RGB color."),
-		FSlateIcon(),
+		GetDefault<UDMMaterialValueFloat3RGB>()->GetComponentIcon(),
 		FUIAction(FExecuteAction::CreateWeakLambda(
 			Slot,
 			[Slot]
@@ -146,7 +152,7 @@ void FDMMaterialSlotLayerMenus::AddAddLayerSection(UToolMenu* InMenu)
 		NAME_None,
 		LOCTEXT("AddEdgeColor", "Texture Edge Color"),
 		LOCTEXT("AddEdgeColorTooltip", "Add a new Material Layer with a solid color based on the edge color on a texture."),
-		FSlateIcon(),
+		GetDefault<UDMMaterialStageExpressionTextureSampleEdgeColor>()->GetComponentIcon(),
 		FUIAction(FExecuteAction::CreateWeakLambda(
 			Slot,
 			[Slot]
@@ -160,7 +166,7 @@ void FDMMaterialSlotLayerMenus::AddAddLayerSection(UToolMenu* InMenu)
 		NAME_None,
 		LOCTEXT("AddNoise", "Noise"),
 		LOCTEXT("AddNoiseTooltip", "Add a new Material Layer with a noise pattern."),
-		FSlateIcon(),
+		GetDefault<UDMMaterialStageExpressionWorldPositionNoise>()->GetComponentIcon(),
 		FUIAction(FExecuteAction::CreateWeakLambda(
 			Slot,
 			[Slot]
@@ -176,7 +182,7 @@ void FDMMaterialSlotLayerMenus::AddAddLayerSection(UToolMenu* InMenu)
 			NAME_None,
 			LOCTEXT("AddSceneTexture", "Post Process"),
 			LOCTEXT("AddSceneTextureTooltip", "Add a new Material Layer that represents the Scene Texture for a post process material."),
-			FSlateIcon(),
+			GetDefault<UDMMaterialStageExpressionSceneTexture>()->GetComponentIcon(),
 			FUIAction(FExecuteAction::CreateWeakLambda(
 				Slot,
 				[Slot]
@@ -216,12 +222,56 @@ void FDMMaterialSlotLayerMenus::AddLayerModifySection(UToolMenu* InMenu)
 {
 	using namespace UE::DynamicMaterialEditor::Private;
 
-	if (!IsValid(InMenu))
+	if (!IsValid(InMenu) || InMenu->ContainsSection(SlotLayerModifySectionName))
+	{
+		return;
+	}
+
+	const UDMMenuContext* const MenuContext = InMenu->FindContext<UDMMenuContext>();
+	if (!MenuContext)
+	{
+		return;
+	}
+
+	UDMMaterialLayerObject* Layer = MenuContext->GetLayer();
+	if (!Layer)
+	{
+		return;
+	}
+
+	UDMMaterialSlot* Slot = Layer->GetSlot();
+	if (!Slot)
 	{
 		return;
 	}
 
-	FToolMenuSection& NewSection = InMenu->AddSection(SlotLayerModifySectionName, LOCTEXT("LayerActions", "LayerActions"));
+	FToolMenuSection& NewSection = InMenu->AddSection(SlotLayerModifySectionName, LOCTEXT("LayerActions", "Layer Actions"));
+
+	if (Slot->CanRemoveLayer(Layer))
+	{
+		FSlateIcon ToggleLayerIcon = Layer->IsEnabled()
+			? FSlateIcon(FAppStyle::Get().GetStyleSetName(), TEXT("Kismet.VariableList.HideForInstance"))
+			: FSlateIcon(FAppStyle::Get().GetStyleSetName(), TEXT("Kismet.VariableList.ExposeForInstance"));
+
+		NewSection.AddMenuEntry(NAME_None,
+			LOCTEXT("ToggleLayer", "Toggle Layer"),
+			LOCTEXT("ToggleLayerTooltip", "Toggle the Layer.\n\nAlt+Left Click"),
+			ToggleLayerIcon,
+			FUIAction(FExecuteAction::CreateWeakLambda(
+				Layer,
+				[Layer]()
+				{
+					FScopedTransaction Transaction(LOCTEXT("ToggleAllStageEnabled", "Toggle All Stage Enabled"));
+
+					for (UDMMaterialStage* Stage : Layer->GetStages(EDMMaterialLayerStage::All))
+					{
+						Stage->Modify();
+						Stage->SetEnabled(!Stage->IsEnabled());
+					}
+				}
+			))
+		);
+	}
 
 	NewSection.AddMenuEntry(
 		FDynamicMaterialEditorCommands::Get().InsertDefaultLayerAbove,
@@ -306,7 +356,7 @@ void FDMMaterialSlotLayerMenus::AddGlobalValueSection(UToolMenu* InMenu)
 							NAME_None,
 							Value->GetDescription(),
 							LOCTEXT("AddValueStageSpecificTooltip", "Add a Material Stage based on this Material Value."),
-							FSlateIcon(),
+							Value->GetComponentIcon(),
 							FUIAction(FExecuteAction::CreateWeakLambda(
 								Value,
 								[ValueWeak = TWeakObjectPtr<UDMMaterialValue>(Value)]
@@ -348,8 +398,10 @@ void FDMMaterialSlotLayerMenus::AddGlobalValueSection(UToolMenu* InMenu)
 
 				for (EDMValueType ValueType : UDMValueDefinitionLibrary::GetValueTypes())
 				{
-					FText Name = UDMValueDefinitionLibrary::GetValueDefinition(ValueType).GetDisplayName();
-					FText FormattedTooltip = FText::Format(LOCTEXT("AddTypeTooltipTemplate", "Add a new {0} Value and use it as a Material Stage."), Name);
+					const FText Name = UDMValueDefinitionLibrary::GetValueDefinition(ValueType).GetDisplayName();
+					const FText FormattedTooltip = FText::Format(LOCTEXT("AddTypeTooltipTemplate", "Add a new {0} Value and use it as a Material Stage."), Name);
+
+					const FSlateIcon ValueIcon = UDMValueDefinitionLibrary::GetValueIcon(ValueType);
 
 					InMenu->AddMenuEntry(
 						NAME_None,
@@ -357,7 +409,7 @@ void FDMMaterialSlotLayerMenus::AddGlobalValueSection(UToolMenu* InMenu)
 							NAME_None,
 							Name,
 							FormattedTooltip,
-							FSlateIcon(),
+							ValueIcon,
 							FUIAction(FExecuteAction::CreateWeakLambda(
 								Slot,
 								[Slot, ValueType]
@@ -387,7 +439,7 @@ void FDMMaterialSlotLayerMenus::AddSlotMenuEntry(const TSharedPtr<SDMMaterialSlo
 		FToolMenuEntry::InitMenuEntry(
 			NAME_None,
 			InName,
-			LOCTEXT("AddValueStageSpecificTooltip", "Add a Material Stage based on this Material Value."),
+			LOCTEXT("AddValueStageSpecificTooltip", "Add a Material Stage based on this Material Slot."),
 			FSlateIcon(),
 			FUIAction(FExecuteAction::CreateWeakLambda(
 				TargetSlot,
@@ -592,7 +644,7 @@ void FDMMaterialSlotLayerMenus::AddLayerMenu_Gradients(UToolMenu* InMenu)
 				NAME_None,
 				MenuName,
 				LOCTEXT("ChangeGradientSourceTooltip", "Change the source of this stage to a Material Gradient."),
-				FSlateIcon(),
+				GradientCDO->GetComponentIcon(),
 				FUIAction(FExecuteAction::CreateWeakLambda(
 					Slot,
 					[Slot, GradientClass = TSubclassOf<UDMMaterialStageGradient>(Gradient.Get())]
@@ -610,7 +662,7 @@ void FDMMaterialSlotLayerMenus::AddLayerMenu_Gradients(UToolMenu* InMenu)
 			NAME_None,
 			LOCTEXT("AddColorAtlas", "Color Atlas"),
 			LOCTEXT("AddColorAtlasTooltip", "Add a new Material Layer with a Color Atlas."),
-			FSlateIcon(),
+			GetDefault<UDMMaterialValueColorAtlas>()->GetComponentIcon(),
 			FUIAction(FExecuteAction::CreateWeakLambda(
 				Slot,
 				[Slot]
@@ -675,7 +727,7 @@ void FDMMaterialSlotLayerMenus::AddLayerMenu_Advanced(UToolMenu* InMenu)
 		NAME_None,
 		LOCTEXT("AddText", "Text"),
 		LOCTEXT("AddTextTooltip", "Add a Material Stage based on a Text Renderer."),
-		FSlateIcon(),
+		GetDefault<UDMRenderTargetTextRenderer>()->GetComponentIcon(),
 		FUIAction(FExecuteAction::CreateWeakLambda(
 			Slot,
 			[Slot]
@@ -689,7 +741,7 @@ void FDMMaterialSlotLayerMenus::AddLayerMenu_Advanced(UToolMenu* InMenu)
 		NAME_None,
 		LOCTEXT("AddWidget", "Widget"),
 		LOCTEXT("AddWidgetTooltip", "Add a Material Stage based on a Widget Renderer."),
-		FSlateIcon(),
+		GetDefault<UDMRenderTargetUMGWidgetRenderer>()->GetComponentIcon(),
 		FUIAction(FExecuteAction::CreateWeakLambda(
 			Slot,
 			[Slot]
@@ -703,7 +755,7 @@ void FDMMaterialSlotLayerMenus::AddLayerMenu_Advanced(UToolMenu* InMenu)
 		NAME_None,
 		LOCTEXT("AddMaterialFunction", "Material Function"),
 		LOCTEXT("AddMaterialFunctionTooltip", "Add a new Material Layer based on a Material Function."),
-		FSlateIcon(),
+		GetDefault<UDMMaterialStageFunction>()->GetComponentIcon(),
 		FUIAction(FExecuteAction::CreateWeakLambda(
 			Slot,
 			[Slot]
diff --git a/Engine/Plugins/Experimental/DynamicMaterial/Source/DynamicMaterialEditor/Private/UI/Menus/DMMaterialSlotLayerMenus.h b/Engine/Plugins/Experimental/DynamicMaterial/Source/DynamicMaterialEditor/Private/UI/Menus/DMMaterialSlotLayerMenus.h
--- a/Engine/Plugins/Experimental/DynamicMaterial/Source/DynamicMaterialEditor/Private/UI/Menus/DMMaterialSlotLayerMenus.h
+++ b/Engine/Plugins/Experimental/DynamicMaterial/Source/DynamicMaterialEditor/Private/UI/Menus/DMMaterialSlotLayerMenus.h
@@ -15,7 +15,7 @@ enum class EDMMaterialPropertyType : uint8;
 class FDMMaterialSlotLayerMenus final
 {
 public:
-	static TSharedRef<SWidget> GenerateSlotLayerMenu(const TSharedPtr<SDMMaterialSlotEditor>& InSlotWidget, UDMMaterialLayerObject* InLayerObject);
+	static TSharedRef<SWidget> GenerateSlotLayerMenu(const TSharedPtr<SDMMaterialSlotEditor>& InSlotWidget, UDMMaterialLayerObject* InLayer);
 
 	static void AddAddLayerSection(UToolMenu* InMenu);
 
diff --git a/Engine/Plugins/Experimental/DynamicMaterial/Source/DynamicMaterialEditor/Private/UI/Menus/DMMaterialStageMenus.cpp b/Engine/Plugins/Experimental/DynamicMaterial/Source/DynamicMaterialEditor/Private/UI/Menus/DMMaterialStageMenus.cpp
--- a/Engine/Plugins/Experimental/DynamicMaterial/Source/DynamicMaterialEditor/Private/UI/Menus/DMMaterialStageMenus.cpp
+++ b/Engine/Plugins/Experimental/DynamicMaterial/Source/DynamicMaterialEditor/Private/UI/Menus/DMMaterialStageMenus.cpp
@@ -6,6 +6,7 @@
 #include "Components/DMMaterialSlot.h"
 #include "Components/DMMaterialStage.h"
 #include "DMMaterialStageSourceMenus.h"
+#include "DynamicMaterialEditorStyle.h"
 #include "ScopedTransaction.h"
 #include "Templates/SharedPointer.h"
 #include "ToolMenu.h"
@@ -83,31 +84,11 @@ void FDMMaterialStageMenus::AddStageSettingsSection(UToolMenu* InMenu)
 	}
 
 	const EDMMaterialLayerStage StageType = Layer->GetStageType(Stage);
-	const bool bAllowRemoveLayer = Slot->CanRemoveLayer(Layer);
-
-	FToolMenuSection& NewSection = InMenu->AddSection(StageMenuToggleName, LOCTEXT("MaterialStageMenu", "Material Stage"));
 
-	if (StageType == EDMMaterialLayerStage::Mask)
-	{
-		NewSection.AddMenuEntry(NAME_None,
-			LOCTEXT("ToggleLayerMask", "Toggle Mask"),
-			LOCTEXT("ToggleLayerMaskTooltip", "Toggle the layer mask on and off.\n\nAlt+Shift+Left Click"),
-			FSlateIcon(),
-			FUIAction(FExecuteAction::CreateWeakLambda(
-				Layer,
-				[Layer]()
-				{
-					FScopedTransaction Transaction(LOCTEXT("ToggleBaseStageEnabled", "Toggle Base Stage Enabled"));
+	// Only if we can remove a layer can we toggle the base stage.
+	const bool bAllowRemoveLayer = Slot->CanRemoveLayer(Layer);
 
-					if (UDMMaterialStage* Stage = Layer->GetStage(EDMMaterialLayerStage::Mask))
-					{
-						Stage->Modify();
-						Stage->SetEnabled(!Stage->IsEnabled());
-					}
-				}
-			))
-		);
-	}
+	FToolMenuSection& NewSection = InMenu->AddSection(StageMenuToggleName, LOCTEXT("MaterialStageMenu", "Stage Actions"));
 
 	if (bAllowRemoveLayer)
 	{
@@ -115,9 +96,8 @@ void FDMMaterialStageMenus::AddStageSettingsSection(UToolMenu* InMenu)
 		{
 			NewSection.AddMenuEntry(NAME_None,
 				LOCTEXT("ToggleLayerBase", "Toggle Base"),
-				LOCTEXT("ToggleLayerBaseTooltip", "Toggle the layer base on and off.\n\n"
-					"Warning: Toggling a layer base off may result in inputs being reset where incompatibilities are found.\n\nAlt+Shift+Left Click"),
-				FSlateIcon(),
+				LOCTEXT("ToggleLayerBaseTooltip", "Toggle the Layer Base.\n\nAlt+Shift+Left Click"),
+				FSlateIcon(FDynamicMaterialEditorStyle::Get().GetStyleSetName(), TEXT("Icons.Stage.Enabled")),
 				FUIAction(FExecuteAction::CreateWeakLambda(
 					Layer,
 					[Layer]()
@@ -133,51 +113,28 @@ void FDMMaterialStageMenus::AddStageSettingsSection(UToolMenu* InMenu)
 				))
 			);
 		}
+	}
 
+	if (StageType == EDMMaterialLayerStage::Mask)
+	{
 		NewSection.AddMenuEntry(NAME_None,
-			LOCTEXT("ToggleLayer", "Toggle Layer"),
-			LOCTEXT("ToggleLayerTooltip", "Toggle the entire layer on and off.\n\n"
-				"Warning: Toggling a layer off may result in inputs being reset where incompatibilities are found.\n\nAlt+Left Click"),
-			FSlateIcon(),
+			LOCTEXT("ToggleLayerMask", "Toggle Mask"),
+			LOCTEXT("ToggleLayerMaskTooltip", "Toggle the Layer Mask.\n\nAlt+Shift+Left Click"),
+			FSlateIcon(FDynamicMaterialEditorStyle::Get().GetStyleSetName(), TEXT("Icons.Stage.Enabled")),
 			FUIAction(FExecuteAction::CreateWeakLambda(
 				Layer,
 				[Layer]()
 				{
-					FScopedTransaction Transaction(LOCTEXT("ToggleAllStageEnabled", "Toggle All Stage Enabled"));
+					FScopedTransaction Transaction(LOCTEXT("ToggleMaskStageEnabled", "Toggle Mask Stage Enabled"));
 
-					for (UDMMaterialStage* Stage : Layer->GetStages(EDMMaterialLayerStage::All))
+					if (UDMMaterialStage* Stage = Layer->GetStage(EDMMaterialLayerStage::Mask))
 					{
 						Stage->Modify();
 						Stage->SetEnabled(!Stage->IsEnabled());
 					}
 				}
 			))
-		);			
-	}
-
-	if (bAllowRemoveLayer)
-	{
-		if (TSharedPtr<SDMMaterialEditor> EditorWidget = MenuContext->GetEditorWidget())
-		{
-			NewSection.AddMenuEntry(NAME_None,
-				LOCTEXT("RemoveLayer", "Remove Layer"),
-				LOCTEXT("RemoveLayerTooltip", "Remove this layer from its Material Slot.\n\n"
-					"Warning: Removing a stage off may result in inputs being reset where incompatibilities are found."),
-				FSlateIcon(),
-				FUIAction(FExecuteAction::CreateWeakLambda(
-					Layer,
-					[Layer]()
-					{
-						FScopedTransaction Transaction(LOCTEXT("RemoverLayer", "Remover Layer"));
-
-						if (UDMMaterialSlot* Slot = Layer->GetSlot())
-						{
-							Slot->RemoveLayer(Layer);
-						}
-					}
-				))
-			);
-		}
+		);
 	}
 }
 
@@ -190,7 +147,7 @@ void FDMMaterialStageMenus::AddStageSourceSection(UToolMenu* InMenu)
 		return;
 	}
 
-	FToolMenuSection& NewSection = InMenu->AddSection(StageSourceMenuName, LOCTEXT("MaterialStageSource", "Stage Source"));
+	FToolMenuSection& NewSection = InMenu->AddSection(StageSourceMenuName, LOCTEXT("MaterialStageSource", "Change Stage Source"));
 	NewSection.Context.AddObject(InMenu->Context.FindContext<UDMMenuContext>());
 
 	FDMMaterialStageSourceMenus::CreateChangeMaterialStageSource(NewSection);
diff --git a/Engine/Plugins/Experimental/DynamicMaterial/Source/DynamicMaterialEditor/Private/UI/Menus/DMMaterialStageSourceMenus.cpp b/Engine/Plugins/Experimental/DynamicMaterial/Source/DynamicMaterialEditor/Private/UI/Menus/DMMaterialStageSourceMenus.cpp
--- a/Engine/Plugins/Experimental/DynamicMaterial/Source/DynamicMaterialEditor/Private/UI/Menus/DMMaterialStageSourceMenus.cpp
+++ b/Engine/Plugins/Experimental/DynamicMaterial/Source/DynamicMaterialEditor/Private/UI/Menus/DMMaterialStageSourceMenus.cpp
@@ -23,6 +23,8 @@
 #include "Components/MaterialStageInputs/DMMSIGradient.h"
 #include "Components/MaterialStageInputs/DMMSISlot.h"
 #include "Components/MaterialStageInputs/DMMSIValue.h"
+#include "Components/MaterialValues/DMMaterialValueColorAtlas.h"
+#include "Components/MaterialValues/DMMaterialValueFloat3RGB.h"
 #include "Components/RenderTargetRenderers/DMRenderTargetTextRenderer.h"
 #include "Components/RenderTargetRenderers/DMRenderTargetUMGWidgetRenderer.h"
 #include "DMDefs.h"
@@ -149,10 +151,12 @@ void FDMMaterialStageSourceMenus::GenerateChangeSourceMenu_NewLocalValues(UToolM
 		const FText Name = UDMValueDefinitionLibrary::GetValueDefinition(ValueType).GetDisplayName();
 		const FText FormattedTooltip = FText::Format(NameTooltipFormat, Name);
 
+		const FSlateIcon ValueIcon = UDMValueDefinitionLibrary::GetValueIcon(ValueType);
+
 		InMenu->AddMenuEntry(NAME_None, FToolMenuEntry::InitMenuEntry(NAME_None,
 			Name,
 			FormattedTooltip,
-			FSlateIcon(),
+			ValueIcon,
 			FUIAction(
 				FExecuteAction::CreateWeakLambda(
 					MenuContext,
@@ -250,7 +254,7 @@ void FDMMaterialStageSourceMenus::GenerateChangeSourceMenu_GlobalValues(UToolMen
 		InMenu->AddMenuEntry(NAME_None, FToolMenuEntry::InitMenuEntry(NAME_None,
 			Value->GetDescription(),
 			LOCTEXT("ChangeSourceValueSourceTooltip2", "Change the source of this stage to this Material Value."),
-			FSlateIcon(),
+			Value->GetComponentIcon(),
 			FUIAction(FExecuteAction::CreateWeakLambda(
 				MenuContext,
 				[MenuContext, Value]()
@@ -337,10 +341,12 @@ void FDMMaterialStageSourceMenus::GenerateChangeSourceMenu_NewGlobalValues(UTool
 		const FText Name = UDMValueDefinitionLibrary::GetValueDefinition(ValueType).GetDisplayName();
 		const FText FormattedTooltip = FText::Format(NameTooltipFormat, Name);
 
+		const FSlateIcon ValueIcon = UDMValueDefinitionLibrary::GetValueIcon(ValueType);
+
 		InMenu->AddMenuEntry(NAME_None, FToolMenuEntry::InitMenuEntry(NAME_None,
 			Name,
 			FormattedTooltip,
-			FSlateIcon(),
+			ValueIcon,
 			FUIAction(
 				FExecuteAction::CreateWeakLambda(
 					MenuContext,
@@ -714,7 +720,7 @@ void FDMMaterialStageSourceMenus::GenerateChangeSourceMenu_Gradients(UToolMenu*
 					NAME_None,
 					MenuName,
 					LOCTEXT("ChangeSourceGradientTooltip", "Change the source of this stage to a Material Gradient."),
-					FSlateIcon(),
+					GradientCDO->GetComponentIcon(),
 					FUIAction(FExecuteAction::CreateWeakLambda(
 						MenuContext,
 						[MenuContext, GradientClass]()
@@ -771,7 +777,7 @@ void FDMMaterialStageSourceMenus::GenerateChangeSourceMenu_Gradients(UToolMenu*
 			NAME_None,
 			LOCTEXT("ChangeSourceColorAtlas", "Color Atlas"),
 			LOCTEXT("ChangeSourceColorAtlasTooltip", "Change the source of this stage to a Color Atlas."),
-			FSlateIcon(),
+			GetDefault<UDMMaterialValueColorAtlas>()->GetComponentIcon(),
 			FUIAction(
 				FExecuteAction::CreateStatic(
 					&ChangeSourceToColorAtlasFromContext,
@@ -819,7 +825,7 @@ void FDMMaterialStageSourceMenus::GenerateChangeSourceMenu_Advanced(UToolMenu* c
 	NewSection.AddMenuEntry("Text",
 		LOCTEXT("ChangeSourceText", "Text"),
 		LOCTEXT("ChangeSourceTextTooltip", "Change the source of this stage to a Text Renderer."),
-		FSlateIcon(),
+		GetDefault<UDMRenderTargetTextRenderer>()->GetComponentIcon(),
 		FUIAction(
 			FExecuteAction::CreateStatic(
 				&ChangeSourceToTextFromContext,
@@ -831,7 +837,7 @@ void FDMMaterialStageSourceMenus::GenerateChangeSourceMenu_Advanced(UToolMenu* c
 	NewSection.AddMenuEntry("Widget",
 		LOCTEXT("ChangeSourceWidget", "Widget"),
 		LOCTEXT("ChangeSourceWidgetTooltip", "Change the source of this stage to a Widget Renderer."),
-		FSlateIcon(),
+		GetDefault<UDMRenderTargetUMGWidgetRenderer>()->GetComponentIcon(),
 		FUIAction(
 			FExecuteAction::CreateStatic(
 				&ChangeSourceToWidgetFromContext,
@@ -843,7 +849,7 @@ void FDMMaterialStageSourceMenus::GenerateChangeSourceMenu_Advanced(UToolMenu* c
 	NewSection.AddMenuEntry("MaterialFunction",
 		LOCTEXT("ChangeSourceMaterialFunction", "Material Function"),
 		LOCTEXT("ChangeSourceMaterialFunctionTooltip", "Change the source of this stage to a Material Function."),
-		FSlateIcon(),
+		GetDefault<UDMMaterialStageFunction>()->GetComponentIcon(),
 		FUIAction(
 			FExecuteAction::CreateStatic(
 				&ChangeSourceToMaterialFunctionFromContext,
@@ -888,7 +894,7 @@ void FDMMaterialStageSourceMenus::GenerateChangeSourceMenu_Advanced(UToolMenu* c
 							LOCTEXT("ChangeSourceSlotOuptut", "Slot Output"),
 							LOCTEXT("ChangeSourceSlotOutputTooltip", "Change the source of this stage to the output from another Material Slot."),
 							FNewToolMenuDelegate::CreateStatic(&GenerateChangeSourceMenu_Slots)
-							);
+						);
 					}));
 		}
 	}
@@ -1462,7 +1468,7 @@ void FDMMaterialStageSourceMenus::CreateChangeMaterialStageSource(FToolMenuSecti
 	InSection.AddMenuEntry("TextureSample",
 		LOCTEXT("TextureSample", "Texture"),
 		LOCTEXT("ChangeSourceTextureSampleTooltip", "Change the source of this stage to a texture."),
-		FSlateIcon(),
+		GetDefault<UDMMaterialStageExpressionTextureSample>()->GetComponentIcon(),
 		FUIAction(
 			FExecuteAction::CreateStatic(
 				&ChangeSourceToTextureSampleFromContext,
@@ -1474,7 +1480,7 @@ void FDMMaterialStageSourceMenus::CreateChangeMaterialStageSource(FToolMenuSecti
 	InSection.AddMenuEntry("SolidColor",
 		LOCTEXT("ChangeSourceColorRGB", "Solid Color"),
 		LOCTEXT("ChangeSourceColorRGBTooltip", "Change the source of this stage to a Solid Color."),
-		FSlateIcon(),
+		GetDefault<UDMMaterialValueFloat3RGB>()->GetComponentIcon(),
 		FUIAction(
 			FExecuteAction::CreateStatic(
 				&ChangeSourceToSolidColorRGBFromContext,
@@ -1486,7 +1492,7 @@ void FDMMaterialStageSourceMenus::CreateChangeMaterialStageSource(FToolMenuSecti
 	InSection.AddMenuEntry("TextureSample_EdgeColor",
 		LOCTEXT("AddTextureSampleEgdeColor", "Texture Edge Color"),
 		LOCTEXT("ChangeSourceTextureSampleEdgeColorTooltip", "Change the source of this stage to the edge color of a texture."),
-		FSlateIcon(),
+		GetDefault<UDMMaterialStageExpressionTextureSampleEdgeColor>()->GetComponentIcon(),
 		FUIAction(
 			FExecuteAction::CreateStatic(
 				&ChangeSourceToTextureSampleEdgeColorFromContext,
@@ -1498,7 +1504,7 @@ void FDMMaterialStageSourceMenus::CreateChangeMaterialStageSource(FToolMenuSecti
 	InSection.AddMenuEntry("Noise",
 		LOCTEXT("ChangeSourceNoise", "Noise"),
 		LOCTEXT("ChangeSourceNoiseTooltip", "Change the source of this stage to a Noise Renderer."),
-		FSlateIcon(),
+		GetDefault<UDMMaterialStageExpressionWorldPositionNoise>()->GetComponentIcon(),
 		FUIAction(
 			FExecuteAction::CreateStatic(
 				&ChangeSourceToNoiseFromContext,
@@ -1507,21 +1513,24 @@ void FDMMaterialStageSourceMenus::CreateChangeMaterialStageSource(FToolMenuSecti
 		)
 	);
 
-	InSection.AddMenuEntry("SceneTexture",
-		LOCTEXT("AddSceneTexture", "Post Process"),
-		LOCTEXT("ChangeSourceSceneTextureTooltip", "Change the source of this stage to Scene Texture in post process materials."),
-		FSlateIcon(),
-		FUIAction(
-			FExecuteAction::CreateStatic(
-				&ChangeSourceToSolidColorRGBFromContext,
-				MenuContext
-			),
-			FCanExecuteAction::CreateStatic(
-				&CanChangeSourceToSceneTextureFromContext,
-				MenuContext
+	if (ModelEditorOnlyData->GetDomain() == EMaterialDomain::MD_PostProcess)
+	{
+		InSection.AddMenuEntry("SceneTexture",
+			LOCTEXT("AddSceneTexture", "Post Process"),
+			LOCTEXT("ChangeSourceSceneTextureTooltip", "Change the source of this stage to Scene Texture in post process materials."),
+			GetDefault<UDMMaterialStageExpressionSceneTexture>()->GetComponentIcon(),
+			FUIAction(
+				FExecuteAction::CreateStatic(
+					&ChangeSourceToSolidColorRGBFromContext,
+					MenuContext
+				),
+				FCanExecuteAction::CreateStatic(
+					&CanChangeSourceToSceneTextureFromContext,
+					MenuContext
+				)
 			)
-		)
-	);
+		);
+	}
 
 	if constexpr (UE::DynamicMaterialEditor::bGlobalValuesEnabled)
 	{
diff --git a/Engine/Plugins/Experimental/DynamicMaterial/Source/DynamicMaterialEditor/Private/UI/Widgets/Editor/SDMMaterialSlotEditor.cpp b/Engine/Plugins/Experimental/DynamicMaterial/Source/DynamicMaterialEditor/Private/UI/Widgets/Editor/SDMMaterialSlotEditor.cpp
--- a/Engine/Plugins/Experimental/DynamicMaterial/Source/DynamicMaterialEditor/Private/UI/Widgets/Editor/SDMMaterialSlotEditor.cpp
+++ b/Engine/Plugins/Experimental/DynamicMaterial/Source/DynamicMaterialEditor/Private/UI/Widgets/Editor/SDMMaterialSlotEditor.cpp
@@ -724,38 +724,7 @@ TSharedRef<SDMMaterialSlotLayerView> SDMMaterialSlotEditor::CreateSlot_LayerView
 {
 	TSharedRef<SDMMaterialSlotLayerView> NewLayerView = SNew(SDMMaterialSlotLayerView, SharedThis(this));
 	NewLayerView->GetOnSelectionChanged().AddSP(this, &SDMMaterialSlotEditor::OnLayerSelected);
-
-	if (UDMMaterialSlot* Slot = GetSlot())
-	{
-		if (TSharedPtr<SDMMaterialEditor> EditorWidget = GetEditorWidget())
-		{
-			bool bHasValidComponentToEdit = false;
-
-			if (UDMMaterialComponent* ComponentToEdit = EditorWidget->GetComponentToEdit())
-			{
-				if (ComponentToEdit->GetTypedParent<UDMMaterialSlot>(/* Allow Subclasses */ true) == Slot)
-				{
-					bHasValidComponentToEdit = true;
-				}
-			}
-
-			if (!bHasValidComponentToEdit)
-			{
-				const TArray<UDMMaterialLayerObject*>& Layers = Slot->GetLayers();
-
-				if (!Layers.IsEmpty())
-				{
-					UDMMaterialLayerObject* LastLayer = Layers.Last();
-					NewLayerView->SetSelectedLayer(LastLayer);
-
-					if (UDMMaterialStage* Stage = LastLayer->GetFirstEnabledStage(EDMMaterialLayerStage::All))
-					{
-						EditorWidget->EditComponent(Stage);
-					}
-				}
-			}
-		}
-	}
+	NewLayerView->EnsureSelectedStage();
 
 	return NewLayerView;
 }
@@ -888,12 +857,7 @@ FText SDMMaterialSlotEditor::GetLayerButtonsDescription() const
 
 TSharedRef<SWidget> SDMMaterialSlotEditor::GetLayerButtonsMenuContent()
 {
-	if (UDMMaterialLayerObject* LayerObject = LayerViewSlot->GetSelectedLayer())
-	{
-		return FDMMaterialSlotLayerMenus::GenerateSlotLayerMenu(SharedThis(this), LayerObject);
-	}
-
-	return SNullWidget::NullWidget;
+	return FDMMaterialSlotLayerMenus::GenerateSlotLayerMenu(SharedThis(this), nullptr);
 }
 
 bool SDMMaterialSlotEditor::GetLayerCanAddEffect() const
diff --git a/Engine/Plugins/Experimental/DynamicMaterial/Source/DynamicMaterialEditor/Private/UI/Widgets/Editor/SlotEditor/SDMMaterialSlotLayerItem.cpp b/Engine/Plugins/Experimental/DynamicMaterial/Source/DynamicMaterialEditor/Private/UI/Widgets/Editor/SlotEditor/SDMMaterialSlotLayerItem.cpp
--- a/Engine/Plugins/Experimental/DynamicMaterial/Source/DynamicMaterialEditor/Private/UI/Widgets/Editor/SlotEditor/SDMMaterialSlotLayerItem.cpp
+++ b/Engine/Plugins/Experimental/DynamicMaterial/Source/DynamicMaterialEditor/Private/UI/Widgets/Editor/SlotEditor/SDMMaterialSlotLayerItem.cpp
@@ -16,6 +16,7 @@
 #include "Model/DynamicMaterialModelEditorOnlyData.h"
 #include "Styling/StyleColors.h"
 #include "UI/DragDrop/DMSlotLayerDragDropOperation.h"
+#include "UI/Menus/DMMaterialStageSourceMenus.h"
 #include "UI/Utils/DMWidgetStatics.h"
 #include "UI/Widgets/Editor/SDMMaterialComponentEditor.h"
 #include "UI/Widgets/Editor/SDMMaterialSlotEditor.h"
@@ -24,7 +25,9 @@
 #include "UI/Widgets/Editor/SlotEditor/SDMMaterialStage.h"
 #include "UI/Widgets/SDMMaterialEditor.h"
 #include "Widgets/Input/SButton.h"
+#include "Widgets/Input/SComboButton.h"
 #include "Widgets/Input/SEditableTextBox.h"
+#include "Widgets/SNullWidget.h"
 
 #define LOCTEXT_NAMESPACE "SDMMaterialSlotLayerItem"
 
@@ -221,9 +224,22 @@ TSharedRef<SWidget> SDMMaterialSlotLayerItem::CreateHeaderRowContent()
 			+ SHorizontalBox::Slot()
 			.AutoWidth()
 			.VAlign(VAlign_Center)
-			.Padding(0.0f, VerticalSpacing, HorizontalSpacing, VerticalSpacing)
 			[
-				CreateLayerBaseToggleButton()
+				SNew(SVerticalBox)
+				+ SVerticalBox::Slot()
+				.AutoHeight()
+				.HAlign(HAlign_Left)
+				.VAlign(VAlign_Center)
+				[
+					CreateLayerBaseToggleButton()
+				]
+				+ SVerticalBox::Slot()
+				.AutoHeight()
+				.HAlign(HAlign_Left)
+				.VAlign(VAlign_Center)
+				[
+					CreateStageSourceButton(EDMMaterialLayerStage::Base)
+				]
 			]
 			+ SHorizontalBox::Slot()
 			.AutoWidth()
@@ -242,9 +258,22 @@ TSharedRef<SWidget> SDMMaterialSlotLayerItem::CreateHeaderRowContent()
 			+ SHorizontalBox::Slot()
 			.AutoWidth()
 			.VAlign(VAlign_Center)
-			.Padding(0.0f, VerticalSpacing, HorizontalSpacing, VerticalSpacing)
 			[
-				CreateLayerMaskToggleButton()
+				SNew(SVerticalBox)
+				+ SVerticalBox::Slot()
+				.AutoHeight()
+				.HAlign(HAlign_Left)
+				.VAlign(VAlign_Center)
+				[
+					CreateLayerMaskToggleButton()
+				]
+				+ SVerticalBox::Slot()
+				.AutoHeight()
+				.HAlign(HAlign_Left)
+				.VAlign(VAlign_Center)
+				[
+					CreateStageSourceButton(EDMMaterialLayerStage::Mask)
+				]
 			]
 			+ SHorizontalBox::Slot()
 			.AutoWidth()
@@ -279,7 +308,7 @@ TSharedRef<SWidget> SDMMaterialSlotLayerItem::CreateHeaderRowContent()
 				[
 					SNew(STextBlock)
 					.TextStyle(FDynamicMaterialEditorStyle::Get(), "SmallFont")
-					.Text(this, &SDMMaterialSlotLayerItem::GetStageDescription)
+					.Text(this, &SDMMaterialSlotLayerItem::GetStageDescription, EDMMaterialLayerStage::All)
 				]
 			]
 			+ SHorizontalBox::Slot()
@@ -438,9 +467,7 @@ TSharedRef<SWidget> SDMMaterialSlotLayerItem::CreateLayerBaseToggleButton()
 		.ContentPadding(0.0f)
 		.ButtonStyle(FDynamicMaterialEditorStyle::Get(), "HoverHintOnly")
 		.Cursor(EMouseCursor::Default)
-		.ToolTipText(LOCTEXT("MaterialLayerBaseToggleTooltip", "Toggle the layer base on and off.\n\n"
-			"Warning: Toggling a layer base off may result in inputs being reset where incompatibilities are found.\n\n"
-			"The base of the first layer cannot be toggled off."))
+		.ToolTipText(LOCTEXT("MaterialLayerBaseToggleTooltip", "Toggle the Layer Base."))
 		.OnClicked(this, &SDMMaterialSlotLayerItem::OnStageToggleButtonClicked, EDMMaterialLayerStage::Base)
 		[
 			SNew(SImage)
@@ -454,7 +481,7 @@ TSharedRef<SWidget> SDMMaterialSlotLayerItem::CreateLayerMaskToggleButton()
 		.IsEnabled(!bIsDynamic)
 		.ContentPadding(0.0f)
 		.ButtonStyle(FDynamicMaterialEditorStyle::Get(), "HoverHintOnly")
-		.ToolTipText(LOCTEXT("MaterialLayerMaskToggleTooltip", "Toggle the layer mask on and off."))
+		.ToolTipText(LOCTEXT("MaterialLayerMaskToggleTooltip", "Toggle the Layer Mask."))
 		.Cursor(EMouseCursor::Default)
 		.OnClicked(this, &SDMMaterialSlotLayerItem::OnStageToggleButtonClicked, EDMMaterialLayerStage::Mask)
 		[
@@ -469,7 +496,7 @@ TSharedRef<SWidget> SDMMaterialSlotLayerItem::CreateLayerLinkToggleButton()
 		.IsEnabled(!bIsDynamic)
 		.ContentPadding(0.0f)
 		.ButtonStyle(FDynamicMaterialEditorStyle::Get(), "HoverHintOnly")
-		.ToolTipText(LOCTEXT("MaterialStageLinkTooltip", "Click to toggle UV Link on and off."))
+		.ToolTipText(LOCTEXT("MaterialStageLinkTooltip", "Toggle Layer UV Link."))
 		.Cursor(EMouseCursor::Default)
 		.OnClicked(this, &SDMMaterialSlotLayerItem::OnLayerLinkToggleButton)
 		.Visibility(this, &SDMMaterialSlotLayerItem::GetLayerLinkToggleButtonVisibility)
@@ -551,6 +578,22 @@ TSharedRef<SWidget> SDMMaterialSlotLayerItem::CreateEffectsToggleButton()
 		];
 }
 
+TSharedRef<SWidget> SDMMaterialSlotLayerItem::CreateStageSourceButton(EDMMaterialLayerStage InStage)
+{
+	return SNew(SButton)
+		.IsEnabled(!bIsDynamic)
+		.ContentPadding(5.f)
+		.ButtonStyle(FDynamicMaterialEditorStyle::Get(), "HoverHintOnly")
+		.ToolTipText(this, &SDMMaterialSlotLayerItem::GetStageSourceButtonToolTip, InStage)
+		.Cursor(EMouseCursor::Default)
+		.OnClicked(this, &SDMMaterialSlotLayerItem::OnStageSourceButtonClicked, InStage)
+		[
+			SNew(SImage)
+			.DesiredSizeOverride(FVector2D(16.f))
+			.Image(this, &SDMMaterialSlotLayerItem::GetStageSourceButtonImage, InStage)
+		];
+}
+
 EVisibility SDMMaterialSlotLayerItem::GetEffectsListVisibility() const
 {
 	if (!LayerItem.IsValid())
@@ -813,6 +856,22 @@ FReply SDMMaterialSlotLayerItem::OnStageToggleButtonClicked(EDMMaterialLayerStag
 	return FReply::Handled();
 }
 
+FText SDMMaterialSlotLayerItem::GetStageSourceButtonToolTip(EDMMaterialLayerStage InLayerStage) const
+{
+	if (UDMMaterialLayerObject* Layer = GetLayer())
+	{
+		if (UDMMaterialStage* Stage = Layer->GetFirstEnabledStage(InLayerStage))
+		{
+			return FText::Format(
+				LOCTEXT("StageSourceToolTipFormat", "Click to change the Material Stage Source.\n\nSource: {0}."),
+				Stage->GetComponentDescription()
+			);
+		}
+	}
+
+	return GetDefault<UDMMaterialStage>()->GetComponentDescription();
+}
+
 FText SDMMaterialSlotLayerItem::GetToolTipText() const
 {
 	if (UDMMaterialLayerObject* Layer = GetLayer())
@@ -857,11 +916,11 @@ FText SDMMaterialSlotLayerItem::GetBlendModeText() const
 	return LOCTEXT("Error", "Error");
 }
 
-FText SDMMaterialSlotLayerItem::GetStageDescription() const
+FText SDMMaterialSlotLayerItem::GetStageDescription(EDMMaterialLayerStage InLayerStage) const
 {
 	if (UDMMaterialLayerObject* Layer = GetLayer())
 	{
-		if (UDMMaterialStage* Stage = Layer->GetFirstEnabledStage(EDMMaterialLayerStage::All))
+		if (UDMMaterialStage* Stage = Layer->GetFirstEnabledStage(InLayerStage))
 		{
 			return Stage->GetComponentDescription();
 		}
@@ -870,6 +929,63 @@ FText SDMMaterialSlotLayerItem::GetStageDescription() const
 	return GetDefault<UDMMaterialStage>()->GetComponentDescription();
 }
 
+const FSlateBrush* SDMMaterialSlotLayerItem::GetStageSourceButtonImage(EDMMaterialLayerStage InLayerStage) const
+{
+	if (UDMMaterialLayerObject* Layer = GetLayer())
+	{
+		if (UDMMaterialStage* Stage = Layer->GetFirstValidStage(InLayerStage))
+		{
+			return Stage->GetComponentIcon().GetIcon();
+		}
+	}
+
+	return FAppStyle::Get().GetBrush("Icons.ErrorWithColor");
+}
+
+FReply SDMMaterialSlotLayerItem::OnStageSourceButtonClicked(EDMMaterialLayerStage InLayerStage)
+{
+	FSlateApplication& SlateApplication = FSlateApplication::Get();
+
+	SlateApplication.PushMenu(
+		SharedThis(this),
+		FWidgetPath(),
+		GetStageSourceMenuContent(InLayerStage),
+		SlateApplication.GetCursorPos(),
+		FPopupTransitionEffect::ContextMenu
+	);
+
+	return FReply::Handled();
+}
+
+TSharedRef<SWidget> SDMMaterialSlotLayerItem::GetStageSourceMenuContent(EDMMaterialLayerStage InLayerStage)
+{
+	TSharedPtr<SDMMaterialSlotLayerView> LayerView = GetSlotLayerView();
+
+	if (!LayerView.IsValid())
+	{
+		return SNullWidget::NullWidget;
+	}
+
+	TSharedPtr<SDMMaterialSlotEditor> SlotEditorWidget = LayerView->GetSlotEditorWidget();
+
+	if (!SlotEditorWidget.IsValid())
+	{
+		return SNullWidget::NullWidget;
+	}
+
+	switch (InLayerStage)
+	{
+		case EDMMaterialLayerStage::Base:
+			return FDMMaterialStageSourceMenus::MakeChangeSourceMenu(SlotEditorWidget, BaseStageWidget);
+
+		case EDMMaterialLayerStage::Mask:
+			return FDMMaterialStageSourceMenus::MakeChangeSourceMenu(SlotEditorWidget, MaskStageWidget);
+
+		default:
+			return SNullWidget::NullWidget;
+	}
+}
+
 int32 SDMMaterialSlotLayerItem::OnLayerItemPaintDropIndicator(EItemDropZone InItemDropZone, const FPaintArgs& InArgs, 
 	const FGeometry& InAllottedGeometry, const FSlateRect& InMyCullingRect, FSlateWindowElementList& OutDrawElements, int32 InLayerId, 
 	const FWidgetStyle& InWidgetStyle, bool bInParentEnabled) const
diff --git a/Engine/Plugins/Experimental/DynamicMaterial/Source/DynamicMaterialEditor/Private/UI/Widgets/Editor/SlotEditor/SDMMaterialSlotLayerItem.h b/Engine/Plugins/Experimental/DynamicMaterial/Source/DynamicMaterialEditor/Private/UI/Widgets/Editor/SlotEditor/SDMMaterialSlotLayerItem.h
--- a/Engine/Plugins/Experimental/DynamicMaterial/Source/DynamicMaterialEditor/Private/UI/Widgets/Editor/SlotEditor/SDMMaterialSlotLayerItem.h
+++ b/Engine/Plugins/Experimental/DynamicMaterial/Source/DynamicMaterialEditor/Private/UI/Widgets/Editor/SlotEditor/SDMMaterialSlotLayerItem.h
@@ -15,6 +15,7 @@ class SDMMaterialSlotLayerView;
 class SDMMaterialStage;
 class UDMMaterialEffectStack;
 enum class EDMMaterialLayerStage : uint8;
+struct FSlateBrush;
 
 /**
  * Material Slot Layer
@@ -80,6 +81,7 @@ protected:
 	TSharedRef<SWidget> CreateLayerHeaderText();
 	TSharedPtr<SWidget> CreateLayerHeaderEditableText();
 	TSharedRef<SWidget> CreateEffectsToggleButton();
+	TSharedRef<SWidget> CreateStageSourceButton(EDMMaterialLayerStage InStage);
 
 	EVisibility GetEffectsListVisibility() const;
 	EVisibility GetEffectsToggleButtonVisibility() const;
@@ -101,12 +103,17 @@ protected:
 	const FSlateBrush* GetStageToggleButtonImage(EDMMaterialLayerStage InLayerStage) const;
 	FReply OnStageToggleButtonClicked(EDMMaterialLayerStage InLayerStage);
 
+	FText GetStageSourceButtonToolTip(EDMMaterialLayerStage InLayerStage) const;
+	const FSlateBrush* GetStageSourceButtonImage(EDMMaterialLayerStage InLayerStage) const;
+	FReply OnStageSourceButtonClicked(EDMMaterialLayerStage InLayerStage);
+	TSharedRef<SWidget> GetStageSourceMenuContent(EDMMaterialLayerStage InLayerStage);
+
 	/** Text */
 	FText GetToolTipText() const;
 	FText GetLayerHeaderText() const;
 	FText GetLayerIndexText() const;
 	FText GetBlendModeText() const;
-	FText GetStageDescription() const;
+	FText GetStageDescription(EDMMaterialLayerStage InLayerStage) const;
 
 	/** Drag and Drop */
 	int32 OnLayerItemPaintDropIndicator(EItemDropZone InItemDropZone, const FPaintArgs& InArgs, const FGeometry& InAllottedGeometry,
diff --git a/Engine/Plugins/Experimental/DynamicMaterial/Source/DynamicMaterialEditor/Private/UI/Widgets/Editor/SlotEditor/SDMMaterialSlotLayerView.cpp b/Engine/Plugins/Experimental/DynamicMaterial/Source/DynamicMaterialEditor/Private/UI/Widgets/Editor/SlotEditor/SDMMaterialSlotLayerView.cpp
--- a/Engine/Plugins/Experimental/DynamicMaterial/Source/DynamicMaterialEditor/Private/UI/Widgets/Editor/SlotEditor/SDMMaterialSlotLayerView.cpp
+++ b/Engine/Plugins/Experimental/DynamicMaterial/Source/DynamicMaterialEditor/Private/UI/Widgets/Editor/SlotEditor/SDMMaterialSlotLayerView.cpp
@@ -128,6 +128,58 @@ SDMMaterialSlotLayerView::FOnSelectionChanged::RegistrationType& SDMMaterialSlot
 	return OnSelectionChanged;
 }
 
+void SDMMaterialSlotLayerView::EnsureSelectedStage()
+{
+	if (GetSelectedLayer())
+	{
+		return;
+	}
+
+	TSharedPtr<SDMMaterialSlotEditor> SlotEditorWidget = GetSlotEditorWidget();
+
+	if (!SlotEditorWidget.IsValid())
+	{
+		return;
+	}
+
+	UDMMaterialSlot* Slot = SlotEditorWidget->GetSlot();
+
+	if (!Slot)
+	{
+		return;
+	}
+
+	TSharedPtr<SDMMaterialEditor> EditorWidget = SlotEditorWidget->GetEditorWidget();
+
+	if (!EditorWidget.IsValid())
+	{
+		return;
+	}
+
+	if (UDMMaterialComponent* ComponentToEdit = EditorWidget->GetComponentToEdit())
+	{
+		if (ComponentToEdit->GetTypedParent<UDMMaterialSlot>(/* Allow Subclasses */ true) == Slot)
+		{
+			return;
+		}
+	}
+
+	const TArray<UDMMaterialLayerObject*>& Layers = Slot->GetLayers();
+
+	if (Layers.IsEmpty())
+	{
+		return;
+	}
+
+	UDMMaterialLayerObject* LastLayer = Layers.Last();
+	SetSelectedLayer(LastLayer);
+
+	if (UDMMaterialStage* Stage = LastLayer->GetFirstEnabledStage(EDMMaterialLayerStage::All))
+	{
+		EditorWidget->EditComponent(Stage);
+	}
+}
+
 void SDMMaterialSlotLayerView::PostUndo(bool bSuccess)
 {
 	OnUndo();
@@ -264,28 +316,7 @@ TSharedPtr<SWidget> SDMMaterialSlotLayerView::CreateLayerItemContextMenu()
 		return nullptr;
 	}
 
-	TSharedPtr<SDMMaterialEditor> EditorWidget = SlotEditor->GetEditorWidget();
-
-	if (!EditorWidget.IsValid())
-	{
-		return nullptr;
-	}
-
-	UDMMaterialLayerObject* LayerObject = nullptr;
-	const FVector2f MousePosition = FSlateApplication::Get().GetCursorPos();
-
-	for (const TSharedPtr<FDMMaterialLayerReference>& LayerItem : LayerItems)
-	{
-		TSharedPtr<SDMMaterialSlotLayerItem> LayerItemWidget = WidgetFromLayerItem(LayerItem);
-
-		if (LayerItemWidget.IsValid() && LayerItemWidget->GetTickSpaceGeometry().IsUnderLocation(MousePosition))
-		{
-			LayerObject = LayerItem->GetLayer();
-			break;
-		}
-	}
-
-	return FDMMaterialSlotLayerMenus::GenerateSlotLayerMenu(SlotEditor, LayerObject);
+	return FDMMaterialSlotLayerMenus::GenerateSlotLayerMenu(SlotEditor, GetSelectedLayer());
 }
 
 TSharedPtr<SDMMaterialSlotLayerItem> SDMMaterialSlotLayerView::WidgetFromLayerItem(const TSharedPtr<FDMMaterialLayerReference>& InItem) const
@@ -411,6 +442,7 @@ void SDMMaterialSlotLayerView::OnLayersUpdated(UDMMaterialSlot* InSlot)
 {
 	RegenerateItems();
 	RequestListRefresh();
+	EnsureSelectedStage();
 }
 
 #undef LOCTEXT_NAMESPACE
diff --git a/Engine/Plugins/Experimental/DynamicMaterial/Source/DynamicMaterialEditor/Private/UI/Widgets/Editor/SlotEditor/SDMMaterialSlotLayerView.h b/Engine/Plugins/Experimental/DynamicMaterial/Source/DynamicMaterialEditor/Private/UI/Widgets/Editor/SlotEditor/SDMMaterialSlotLayerView.h
--- a/Engine/Plugins/Experimental/DynamicMaterial/Source/DynamicMaterialEditor/Private/UI/Widgets/Editor/SlotEditor/SDMMaterialSlotLayerView.h
+++ b/Engine/Plugins/Experimental/DynamicMaterial/Source/DynamicMaterialEditor/Private/UI/Widgets/Editor/SlotEditor/SDMMaterialSlotLayerView.h
@@ -38,6 +38,8 @@ public:
 
 	FOnSelectionChanged::RegistrationType& GetOnSelectionChanged();
 
+	void EnsureSelectedStage();
+
 	//~ Begin FUndoClient
 	virtual void PostUndo(bool bSuccess) override;
 	virtual void PostRedo(bool bSuccess) override;
diff --git a/Engine/Plugins/Experimental/DynamicMaterial/Source/DynamicMaterialEditor/Public/Components/DMMaterialStage.h b/Engine/Plugins/Experimental/DynamicMaterial/Source/DynamicMaterialEditor/Public/Components/DMMaterialStage.h
--- a/Engine/Plugins/Experimental/DynamicMaterial/Source/DynamicMaterialEditor/Public/Components/DMMaterialStage.h
+++ b/Engine/Plugins/Experimental/DynamicMaterial/Source/DynamicMaterialEditor/Public/Components/DMMaterialStage.h
@@ -233,6 +233,7 @@ public:
 
 	//~ Begin UDMMaterialComponent
 	DYNAMICMATERIALEDITOR_API virtual FText GetComponentDescription() const override;
+	DYNAMICMATERIALEDITOR_API virtual FSlateIcon GetComponentIcon() const override;
 	DYNAMICMATERIALEDITOR_API virtual void Update(UDMMaterialComponent* InSource, EDMUpdateType InUpdateType) override;
 	DYNAMICMATERIALEDITOR_API virtual FString GetComponentPathComponent() const override;
 	DYNAMICMATERIALEDITOR_API virtual UDMMaterialComponent* GetParentComponent() const override;
diff --git a/Engine/Plugins/Experimental/DynamicMaterial/Source/DynamicMaterialEditor/Public/Components/DMMaterialStageBlend.h b/Engine/Plugins/Experimental/DynamicMaterial/Source/DynamicMaterialEditor/Public/Components/DMMaterialStageBlend.h
--- a/Engine/Plugins/Experimental/DynamicMaterial/Source/DynamicMaterialEditor/Public/Components/DMMaterialStageBlend.h
+++ b/Engine/Plugins/Experimental/DynamicMaterial/Source/DynamicMaterialEditor/Public/Components/DMMaterialStageBlend.h
@@ -64,6 +64,7 @@ public:
 	//~ End UDMMaterialStageSource
 
 	//~ Begin UDMMaterialComponent
+	DYNAMICMATERIALEDITOR_API virtual FSlateIcon GetComponentIcon() const override;
 	DYNAMICMATERIALEDITOR_API virtual void Update(UDMMaterialComponent* InSource, EDMUpdateType InUpdateType) override;
 	//~ End UDMMaterialComponent
 
diff --git a/Engine/Plugins/Experimental/DynamicMaterial/Source/DynamicMaterialEditor/Public/Components/DMMaterialStageThroughputLayerBlend.h b/Engine/Plugins/Experimental/DynamicMaterial/Source/DynamicMaterialEditor/Public/Components/DMMaterialStageThroughputLayerBlend.h
--- a/Engine/Plugins/Experimental/DynamicMaterial/Source/DynamicMaterialEditor/Public/Components/DMMaterialStageThroughputLayerBlend.h
+++ b/Engine/Plugins/Experimental/DynamicMaterial/Source/DynamicMaterialEditor/Public/Components/DMMaterialStageThroughputLayerBlend.h
@@ -39,6 +39,7 @@ public:
 	DYNAMICMATERIALEDITOR_API void SetMaskChannelOverride(EAvaColorChannel InMaskChannel);
 
 	//~ Begin UDMMaterialComponent
+	DYNAMICMATERIALEDITOR_API virtual FSlateIcon GetComponentIcon() const override;
 	DYNAMICMATERIALEDITOR_API virtual void Update(UDMMaterialComponent* InSource, EDMUpdateType InUpdateType) override;
 	DYNAMICMATERIALEDITOR_API virtual void GenerateExpressions(const TSharedRef<FDMMaterialBuildState>& InBuildState) const override;
 	DYNAMICMATERIALEDITOR_API virtual void PostEditorDuplicate(UDynamicMaterialModel* InMaterialModel, UDMMaterialComponent* InParent) override;
diff --git a/Engine/Plugins/Experimental/DynamicMaterial/Source/DynamicMaterialEditor/Public/Components/MaterialStageInputs/DMMSISlot.h b/Engine/Plugins/Experimental/DynamicMaterial/Source/DynamicMaterialEditor/Public/Components/MaterialStageInputs/DMMSISlot.h
--- a/Engine/Plugins/Experimental/DynamicMaterial/Source/DynamicMaterialEditor/Public/Components/MaterialStageInputs/DMMSISlot.h
+++ b/Engine/Plugins/Experimental/DynamicMaterial/Source/DynamicMaterialEditor/Public/Components/MaterialStageInputs/DMMSISlot.h
@@ -53,12 +53,13 @@ public:
 	UFUNCTION(BlueprintCallable, Category = "Material Designer")
 	DYNAMICMATERIALEDITOR_API void SetMaterialProperty(EDMMaterialPropertyType InMaterialProperty);
 
-	//~ Start UDMMaterialStageInput
-	DYNAMICMATERIALEDITOR_API virtual FText GetComponentDescription() const override;
+	//~ Begin UDMMaterialStageInput
 	DYNAMICMATERIALEDITOR_API virtual FText GetChannelDescription(const FDMMaterialStageConnectorChannel& Channel) override;
 	//~ End UDMMaterialStageInput
 
-	//~ Start UDMMaterialComponent
+	//~ Begin UDMMaterialComponent
+	DYNAMICMATERIALEDITOR_API virtual FText GetComponentDescription() const override;
+	DYNAMICMATERIALEDITOR_API virtual FSlateIcon GetComponentIcon() const override;
 	DYNAMICMATERIALEDITOR_API virtual void GenerateExpressions(const TSharedRef<FDMMaterialBuildState>& InBuildState) const override;
 	DYNAMICMATERIALEDITOR_API virtual void PostEditorDuplicate(UDynamicMaterialModel* InMaterialModel, UDMMaterialComponent* InParent) override;
 	//~ End UDMMaterialComponent
@@ -85,7 +86,7 @@ protected:
 	void InitSlot();
 	void DeinitSlot();
 
-	//~ Start UDMMaterialStageInput
+	//~ Begin UDMMaterialStageInput
 	DYNAMICMATERIALEDITOR_API virtual void UpdateOutputConnectors() override;
 	//~ End UDMMaterialStageInput
 
diff --git a/Engine/Plugins/Experimental/DynamicMaterial/Source/DynamicMaterialEditor/Public/Components/MaterialStageInputs/DMMSITextureUV.h b/Engine/Plugins/Experimental/DynamicMaterial/Source/DynamicMaterialEditor/Public/Components/MaterialStageInputs/DMMSITextureUV.h
--- a/Engine/Plugins/Experimental/DynamicMaterial/Source/DynamicMaterialEditor/Public/Components/MaterialStageInputs/DMMSITextureUV.h
+++ b/Engine/Plugins/Experimental/DynamicMaterial/Source/DynamicMaterialEditor/Public/Components/MaterialStageInputs/DMMSITextureUV.h
@@ -45,12 +45,13 @@ public:
 	UFUNCTION(BlueprintPure, Category = "Material Designer")
 	UDMTextureUV* GetTextureUV() { return TextureUV; }
 
-	//~ Start UDMMaterialStageInput
-	DYNAMICMATERIALEDITOR_API virtual FText GetComponentDescription() const override;
+	//~ Begin UDMMaterialStageInput
 	DYNAMICMATERIALEDITOR_API virtual FText GetChannelDescription(const FDMMaterialStageConnectorChannel& Channel) override;
 	//~ End UDMMaterialStageInput
 
-	//~ Start UDMMaterialComponent
+	//~ Begin UDMMaterialComponent
+	DYNAMICMATERIALEDITOR_API virtual FText GetComponentDescription() const override;
+	DYNAMICMATERIALEDITOR_API virtual FSlateIcon GetComponentIcon() const override;
 	DYNAMICMATERIALEDITOR_API virtual void GenerateExpressions(const TSharedRef<FDMMaterialBuildState>& InBuildState) const override;
 	DYNAMICMATERIALEDITOR_API virtual void PostEditorDuplicate(UDynamicMaterialModel* InMaterialModel, UDMMaterialComponent* InParent) override;
 	//~ End UDMMaterialComponent
@@ -76,7 +77,7 @@ protected:
 
 	void AddEffects(const TSharedRef<FDMMaterialBuildState>& InBuildState, TArray<UMaterialExpression*>& InOutExpressions) const;
 
-	//~ Start UDMMaterialStageInput
+	//~ Begin UDMMaterialStageInput
 	DYNAMICMATERIALEDITOR_API virtual void UpdateOutputConnectors() override;
 	//~ End UDMMaterialStageInput
 
diff --git a/Engine/Plugins/Experimental/DynamicMaterial/Source/DynamicMaterialEditor/Public/Components/MaterialStageInputs/DMMSIThroughput.h b/Engine/Plugins/Experimental/DynamicMaterial/Source/DynamicMaterialEditor/Public/Components/MaterialStageInputs/DMMSIThroughput.h
--- a/Engine/Plugins/Experimental/DynamicMaterial/Source/DynamicMaterialEditor/Public/Components/MaterialStageInputs/DMMSIThroughput.h
+++ b/Engine/Plugins/Experimental/DynamicMaterial/Source/DynamicMaterialEditor/Public/Components/MaterialStageInputs/DMMSIThroughput.h
@@ -30,18 +30,19 @@ public:
 	UFUNCTION(BlueprintPure, Category = "Material Designer")
 	UDMMaterialSubStage* GetSubStage() const { return SubStage; }
 
-	//~ Start UDMMaterialStageInput
-	DYNAMICMATERIALEDITOR_API virtual FText GetComponentDescription() const override;
+	//~ Begin UDMMaterialStageInput
 	DYNAMICMATERIALEDITOR_API virtual FText GetChannelDescription(const FDMMaterialStageConnectorChannel& Channel) override;
 	//~ End UDMMaterialStageInput
 
-	//~ Start UDMMaterialStageSource
+	//~ Begin UDMMaterialStageSource
 	DYNAMICMATERIALEDITOR_API virtual void GenerateExpressions(const TSharedRef<FDMMaterialBuildState>& InBuildState) const override;
 	DYNAMICMATERIALEDITOR_API virtual int32 GetInnateMaskOutput(int32 OutputIndex, int32 OutputChannels) const;
 	DYNAMICMATERIALEDITOR_API virtual int32 GetOutputChannelOverride(int32 InOutputIndex) const override;
 	//~ End UDMMaterialStageSource
 
-	//~ Start UDMMaterialComponent
+	//~ Begin UDMMaterialComponent
+	DYNAMICMATERIALEDITOR_API virtual FText GetComponentDescription() const override;
+	DYNAMICMATERIALEDITOR_API virtual FSlateIcon GetComponentIcon() const override;
 	DYNAMICMATERIALEDITOR_API virtual bool IsPropertyVisible(FName Property) const override;
 	DYNAMICMATERIALEDITOR_API virtual void PostEditorDuplicate(UDynamicMaterialModel* InMaterialModel, UDMMaterialComponent* InParent) override;
 	//~ End UDMMaterialComponent
diff --git a/Engine/Plugins/Experimental/DynamicMaterial/Source/DynamicMaterialEditor/Public/Components/MaterialStageInputs/DMMSIValue.h b/Engine/Plugins/Experimental/DynamicMaterial/Source/DynamicMaterialEditor/Public/Components/MaterialStageInputs/DMMSIValue.h
--- a/Engine/Plugins/Experimental/DynamicMaterial/Source/DynamicMaterialEditor/Public/Components/MaterialStageInputs/DMMSIValue.h
+++ b/Engine/Plugins/Experimental/DynamicMaterial/Source/DynamicMaterialEditor/Public/Components/MaterialStageInputs/DMMSIValue.h
@@ -68,16 +68,17 @@ public:
 
 	DYNAMICMATERIALEDITOR_API void ApplyDefaultLayerSettings();
 
-	//~ Start UDMMaterialStageInput
+	//~ Begin UDMMaterialStageInput
 	DYNAMICMATERIALEDITOR_API virtual FText GetChannelDescription(const FDMMaterialStageConnectorChannel& Channel) override;
 	//~ End UDMMaterialStageInput
 
-	//~ Start UDMMaterialStageSource
+	//~ Begin UDMMaterialStageSource
 	DYNAMICMATERIALEDITOR_API virtual int32 GetInnateMaskOutput(int32 OutputIndex, int32 OutputChannels) const override;
-	//~ Start UDMMaterialStageSource
+	//~ Begin UDMMaterialStageSource
 
-	//~ Start UDMMaterialComponent
+	//~ Begin UDMMaterialComponent
 	DYNAMICMATERIALEDITOR_API virtual FText GetComponentDescription() const override;
+	DYNAMICMATERIALEDITOR_API virtual FSlateIcon GetComponentIcon() const override;
 	DYNAMICMATERIALEDITOR_API virtual void GenerateExpressions(const TSharedRef<FDMMaterialBuildState>& InBuildState) const override;
 	DYNAMICMATERIALEDITOR_API virtual void PostEditorDuplicate(UDynamicMaterialModel* InMaterialModel, UDMMaterialComponent* InParent) override;
 	//~ End UDMMaterialComponent
@@ -101,7 +102,7 @@ protected:
 
 	void ApplyWholeLayerValue();
 
-	//~ Start UDMMaterialStageInput
+	//~ Begin UDMMaterialStageInput
 	DYNAMICMATERIALEDITOR_API virtual void UpdateOutputConnectors() override;
 	//~ End UDMMaterialStageInput
 
