Material Designer
- Removed the source settings the details section.

diff --git a/Engine/Plugins/Experimental/DynamicMaterial/Source/DynamicMaterialEditor/Private/DynamicMaterialEditorModule.cpp b/Engine/Plugins/Experimental/DynamicMaterial/Source/DynamicMaterialEditor/Private/DynamicMaterialEditorModule.cpp
--- a/Engine/Plugins/Experimental/DynamicMaterial/Source/DynamicMaterialEditor/Private/DynamicMaterialEditorModule.cpp
+++ b/Engine/Plugins/Experimental/DynamicMaterial/Source/DynamicMaterialEditor/Private/DynamicMaterialEditorModule.cpp
@@ -40,7 +40,6 @@
 #include "UI/PropertyGenerators/DMMaterialStageFunctionPropertyRowGenerator.h"
 #include "UI/PropertyGenerators/DMMaterialValueDynamicPropertyRowGenerator.h"
 #include "UI/PropertyGenerators/DMMaterialValuePropertyRowGenerator.h"
-#include "UI/PropertyGenerators/DMStagePropertyRowGenerator.h"
 #include "UI/PropertyGenerators/DMTextureUVDynamicPropertyRowGenerator.h"
 #include "UI/PropertyGenerators/DMTextureUVPropertyRowGenerator.h"
 #include "UI/PropertyGenerators/DMThroughputPropertyRowGenerator.h"
@@ -233,7 +232,6 @@ void FDynamicMaterialEditorModule::StartupModule()
 	FDMMaterialModelDefaults::RegisterDefaultsDelegates();
 
 	RegisterComponentPropertyRowGeneratorDelegate<UDMMaterialComponent,            FDMComponentPropertyRowGenerator>();
-	RegisterComponentPropertyRowGeneratorDelegate<UDMMaterialStage,                FDMStagePropertyRowGenerator>();
 	RegisterComponentPropertyRowGeneratorDelegate<UDMMaterialValue,                FDMMaterialValuePropertyRowGenerator>();
 	RegisterComponentPropertyRowGeneratorDelegate<UDMMaterialValueDynamic,         FDMMaterialValueDynamicPropertyRowGenerator>();
 	RegisterComponentPropertyRowGeneratorDelegate<UDMTextureUV,                    FDMTextureUVPropertyRowGenerator>();
diff --git a/Engine/Plugins/Experimental/DynamicMaterial/Source/DynamicMaterialEditor/Private/UI/PropertyGenerators/DMStagePropertyRowGenerator.cpp b/Engine/Plugins/Experimental/DynamicMaterial/Source/DynamicMaterialEditor/Private/UI/PropertyGenerators/DMStagePropertyRowGenerator.cpp
deleted file mode 100644
diff --git a/Engine/Plugins/Experimental/DynamicMaterial/Source/DynamicMaterialEditor/Private/UI/PropertyGenerators/DMStagePropertyRowGenerator.h b/Engine/Plugins/Experimental/DynamicMaterial/Source/DynamicMaterialEditor/Private/UI/PropertyGenerators/DMStagePropertyRowGenerator.h
deleted file mode 100644
