Material Designer
- The save button is now only visible when editing a material designer asset that exist in the content browser.

diff --git a/Engine/Plugins/Experimental/DynamicMaterial/Source/DynamicMaterialEditor/Private/UI/Widgets/Editor/SDMToolBar.cpp b/Engine/Plugins/Experimental/DynamicMaterial/Source/DynamicMaterialEditor/Private/UI/Widgets/Editor/SDMToolBar.cpp
--- a/Engine/Plugins/Experimental/DynamicMaterial/Source/DynamicMaterialEditor/Private/UI/Widgets/Editor/SDMToolBar.cpp
+++ b/Engine/Plugins/Experimental/DynamicMaterial/Source/DynamicMaterialEditor/Private/UI/Widgets/Editor/SDMToolBar.cpp
@@ -98,7 +98,7 @@ TSharedRef<SWidget> SDMToolBar::CreateToolBarEntries()
 				.VAlign(EVerticalAlignment::VAlign_Center)
 				[
 					SAssignNew(SaveButtonWidget, SButton)
-					.IsEnabled(false)
+					.Visibility(EVisibility::Collapsed)
 					.ContentPadding(LargeIconToolBarButtonContentPadding)
 					.ButtonStyle(FDynamicMaterialEditorStyle::Get(), "HoverHintOnly")
 					.ToolTipText(LOCTEXT("MaterialDesignerSaveTooltip", "Save the Material Designer asset\n\nCaution: If this asset lives inside an actor, the actor/level will be saved."))
@@ -346,7 +346,7 @@ void SDMToolBar::SetButtonVisibilities()
 		bIsDynamic = !MaterialModelBase->IsA<UDynamicMaterialModel>();
 	}
 
-	SaveButtonWidget->SetEnabled(CanSave());
+	SaveButtonWidget->SetVisibility(CanSave() ? EVisibility::Visible : EVisibility::Collapsed);
 
 	if (bIsAsset)
 	{
@@ -739,7 +739,12 @@ FText SDMToolBar::GetAssetToolTip() const
 
 bool SDMToolBar::CanSave() const
 {
-	return !!GetSaveablePackage(GetMaterialModelBase());
+	if (UDynamicMaterialModelBase* MaterialModelBase = GetMaterialModelBase())
+	{
+		return !MaterialModelBase->GetTypedOuter<UWorld>();
+	}
+
+	return false;
 }
 
 const FSlateBrush* SDMToolBar::GetSaveIcon() const
