Material Designer
- Upgraded the look of the channels page.

diff --git a/Engine/Plugins/Experimental/DynamicMaterial/Source/DynamicMaterialEditor/Private/UI/Widgets/Editor/SDMMaterialProperties.cpp b/Engine/Plugins/Experimental/DynamicMaterial/Source/DynamicMaterialEditor/Private/UI/Widgets/Editor/SDMMaterialProperties.cpp
--- a/Engine/Plugins/Experimental/DynamicMaterial/Source/DynamicMaterialEditor/Private/UI/Widgets/Editor/SDMMaterialProperties.cpp
+++ b/Engine/Plugins/Experimental/DynamicMaterial/Source/DynamicMaterialEditor/Private/UI/Widgets/Editor/SDMMaterialProperties.cpp
@@ -11,6 +11,8 @@
 #include "DMWorldSubsystem.h"
 #include "Engine/World.h"
 #include "ICustomDetailsView.h"
+#include "Items/ICustomDetailsViewCustomCategoryItem.h"
+#include "Items/ICustomDetailsViewCustomItem.h"
 #include "Items/ICustomDetailsViewItem.h"
 #include "Model/DynamicMaterialModelBase.h"
 #include "Model/DynamicMaterialModelDynamic.h"
@@ -56,8 +58,14 @@ TSharedRef<SWidget> SDMMaterialProperties::CreateSlot_Content()
 	{
 		return SNullWidget::NullWidget;
 	}
+	UDynamicMaterialModelBase* MaterialModelBase = EditorWidget->GetMaterialModelBase();
 
-	UDynamicMaterialModel* MaterialModel = EditorWidget->GetMaterialModel();
+	if (!MaterialModelBase)
+	{
+		return SNullWidget::NullWidget;
+	}
+
+	UDynamicMaterialModel* MaterialModel = MaterialModelBase->ResolveMaterialModel();
 
 	if (!MaterialModel)
 	{
@@ -71,11 +79,41 @@ TSharedRef<SWidget> SDMMaterialProperties::CreateSlot_Content()
 		return SNullWidget::NullWidget;
 	}
 
+	FCustomDetailsViewArgs Args;
+	Args.bAllowGlobalExtensions = false;
+	Args.bAllowResetToDefault = false;
+	Args.bShowCategories = true;
+	Args.OnExpansionStateChanged.AddSP(this, &SDMMaterialProperties::OnExpansionStateChanged);
+
+	TSharedRef<ICustomDetailsView> DetailsView = ICustomDetailsViewModule::Get().CreateCustomDetailsView(Args);
+	FCustomDetailsViewItemId RootId = DetailsView->GetRootItem()->GetItemId();
+
+	auto CreateCategory = [&DetailsView, &RootId, MaterialModelBase](FName InName, const FText& InDisplayName)
+		{
+			TSharedRef<ICustomDetailsViewItem> CategoryItem = DetailsView->CreateCustomCategoryItem(InName, InDisplayName)->AsItem();
+			CategoryItem->RefreshItemId();
+			DetailsView->ExtendTree(RootId, ECustomDetailsTreeInsertPosition::Child, CategoryItem);
+
+			bool bExpansionState = true;
+			FDMWidgetStatics::Get().GetExpansionState(MaterialModelBase, InName, bExpansionState);
+
+			DetailsView->SetItemExpansionState(
+				CategoryItem->GetItemId(),
+				bExpansionState ? ECustomDetailsViewExpansion::SelfExpanded : ECustomDetailsViewExpansion::Collapsed
+			);
+
+			return CategoryItem;
+		};
+
+	TSharedRef<ICustomDetailsViewItem> ActiveCategory = CreateCategory(TEXT("Active"), LOCTEXT("ActiveCategory", "Active Channels"));
+	TSharedRef<ICustomDetailsViewItem> InactiveCategory = CreateCategory(TEXT("Inactive"), LOCTEXT("InactiveCategory", "Inactive Channels"));
+	TSharedRef<ICustomDetailsViewItem> IncompatibleCategory = CreateCategory(TEXT("Incompatible"), LOCTEXT("IncompatibleCategory", "Incompatible Channels"));
+
 	TSharedRef<SVerticalBox> List = SNew(SVerticalBox);
 
 	// Active properties first.
 	UE::DynamicMaterial::ForEachMaterialPropertyType(
-		[this, &List, EditorOnlyData](EDMMaterialPropertyType InMaterialProperty)
+		[this, &DetailsView, &ActiveCategory, EditorOnlyData](EDMMaterialPropertyType InMaterialProperty)
 		{
 			if (UE::DynamicMaterialEditor::Private::IsCustomMaterialProperty(InMaterialProperty))
 			{
@@ -90,7 +128,7 @@ TSharedRef<SWidget> SDMMaterialProperties::CreateSlot_Content()
 
 				if (bIsActive)
 				{
-					AddProperty(List, MaterialProperty);
+					AddProperty(DetailsView, ActiveCategory, MaterialProperty);
 				}
 			}
 
@@ -100,7 +138,7 @@ TSharedRef<SWidget> SDMMaterialProperties::CreateSlot_Content()
 
 	// Now inactive properties
 	UE::DynamicMaterial::ForEachMaterialPropertyType(
-		[this, &List, EditorOnlyData](EDMMaterialPropertyType InMaterialProperty)
+		[this, &DetailsView, &InactiveCategory, EditorOnlyData](EDMMaterialPropertyType InMaterialProperty)
 		{
 			if (UE::DynamicMaterialEditor::Private::IsCustomMaterialProperty(InMaterialProperty))
 			{
@@ -116,7 +154,7 @@ TSharedRef<SWidget> SDMMaterialProperties::CreateSlot_Content()
 
 				if (!bIsActive && bIsValid)
 				{
-					AddProperty(List, MaterialProperty);
+					AddProperty(DetailsView, InactiveCategory, MaterialProperty);
 				}
 			}
 
@@ -126,7 +164,7 @@ TSharedRef<SWidget> SDMMaterialProperties::CreateSlot_Content()
 
 	// Now invalid properties
 	UE::DynamicMaterial::ForEachMaterialPropertyType(
-		[this, &List, EditorOnlyData](EDMMaterialPropertyType InMaterialProperty)
+		[this, &DetailsView, &IncompatibleCategory, EditorOnlyData](EDMMaterialPropertyType InMaterialProperty)
 		{
 			if (UE::DynamicMaterialEditor::Private::IsCustomMaterialProperty(InMaterialProperty))
 			{
@@ -142,7 +180,7 @@ TSharedRef<SWidget> SDMMaterialProperties::CreateSlot_Content()
 
 				if (!bIsActive && !bIsValid)
 				{
-					AddProperty(List, MaterialProperty);
+					AddProperty(DetailsView, IncompatibleCategory, MaterialProperty);
 				}
 			}
 
@@ -150,15 +188,19 @@ TSharedRef<SWidget> SDMMaterialProperties::CreateSlot_Content()
 		}
 	);
 
-	return SNew(SBox)
-		.Padding(5.f)
-		[
-			List
-		];
+	DetailsView->RebuildTree(ECustomDetailsViewBuildType::InstantBuild);
+
+	return DetailsView;
 }
 
-void SDMMaterialProperties::AddProperty(const TSharedRef<SVerticalBox>& InList, UDMMaterialProperty* InProperty)
+void SDMMaterialProperties::AddProperty(const TSharedRef<ICustomDetailsView>& InDetailsView, const TSharedRef<ICustomDetailsViewItem>& InCategory,
+	UDMMaterialProperty* InProperty)
 {
+	if (!InProperty)
+	{
+		return;
+	}
+
 	TSharedPtr<SDMMaterialEditor> EditorWidget = EditorWidgetWeak.Pin();
 
 	if (!EditorWidget.IsValid())
@@ -180,80 +222,72 @@ void SDMMaterialProperties::AddProperty(const TSharedRef<SVerticalBox>& InList,
 		return;
 	}
 
-	EDMMaterialPropertyType MaterialProperty = InProperty ? InProperty->GetMaterialProperty() : EDMMaterialPropertyType::None;
+	TSharedPtr<ICustomDetailsViewCustomItem> Item = InDetailsView->CreateCustomItem(InProperty->GetClass()->GetFName());
+
+	if (!Item.IsValid())
+	{
+		return;
+	}
+
+	Item->SetWholeRowWidget(CreatePropertyRow(InProperty));
+
+	InDetailsView->ExtendTree(InCategory->GetItemId(), ECustomDetailsTreeInsertPosition::Child, Item->AsItem());
+}
+
+TSharedRef<SWidget> SDMMaterialProperties::CreatePropertyRow(UDMMaterialProperty* InProperty)
+{
+	// There are all ensured to be valid by the caller of this method
+	TSharedPtr<SDMMaterialEditor> EditorWidget = EditorWidgetWeak.Pin();
+	UDynamicMaterialModel* MaterialModel = EditorWidget->GetMaterialModel();
+	UDynamicMaterialModelEditorOnlyData* EditorOnlyData = UDynamicMaterialModelEditorOnlyData::Get(MaterialModel);
+	EDMMaterialPropertyType MaterialProperty = InProperty->GetMaterialProperty();
 
 	TSharedRef<SBox> PreviewWidgetContainer = SNew(SBox)
 		.WidthOverride(64.f)
-		.HeightOverride(23.f);
+		.HeightOverride(18.f);
 
 	TSharedRef<SWidget> PropertyName = CreateSlot_PropertyName(MaterialProperty);
 
 	TSharedRef<SWidget> Slider = SNullWidget::NullWidget;
 
-	bool bPadSlider = false;
+	const bool bValidProperty = InProperty->IsValidForModel(*EditorOnlyData);
+	const bool bPropertyEnabled = bValidProperty && InProperty->IsEnabled() && EditorOnlyData->GetSlotForMaterialProperty(MaterialProperty);;
 
-	if (InProperty)
+	if (bPropertyEnabled)
 	{
-		const bool bValidProperty = InProperty->IsValidForModel(*EditorOnlyData);
-		const bool bPropertyEnabled = bValidProperty && InProperty->IsEnabled() && EditorOnlyData->GetSlotForMaterialProperty(MaterialProperty);;
+		MaterialProperty = InProperty->GetMaterialProperty();
 
-		if (bPropertyEnabled)
-		{
-			MaterialProperty = InProperty->GetMaterialProperty();
-
-			PreviewWidgetContainer = SNew(SBox)
-				.WidthOverride(64.f)
-				.HeightOverride(64.f)
+		PreviewWidgetContainer = SNew(SBox)
+			.WidthOverride(64.f)
+			.HeightOverride(64.f)
+			[
+				SNew(SBorder)
+				.BorderBackgroundColor(FLinearColor(1, 1, 1, 1.f))
+				.Padding(2.0f)
+				.BorderImage(FAppStyle::GetBrush(UE::DynamicMaterialEditor::Private::EditorDarkBackground))
 				[
-					SNew(SBorder)
-					.BorderBackgroundColor(FLinearColor(1, 1, 1, 1.f))
-					.Padding(2.0f)
-					.BorderImage(FAppStyle::GetBrush(UE::DynamicMaterialEditor::Private::EditorDarkBackground))
-					[
-						SNew(SDMMaterialComponentPreview, EditorWidget.ToSharedRef(), InProperty)
-						.PreviewSize(FVector2D(60.f, 60.f))
-					]
-				];
-
-			PreviewWidgetContainer->SetCursor(EMouseCursor::Hand);
-			PreviewWidgetContainer->SetOnMouseButtonUp(FPointerEventHandler::CreateSP(this, &SDMMaterialProperties::OnPropertyClicked, MaterialProperty));
+					SNew(SDMMaterialComponentPreview, EditorWidget.ToSharedRef(), InProperty)
+					.PreviewSize(FVector2D(60.f, 60.f))
+				]
+			];
 
-			PropertyName->SetCursor(EMouseCursor::Hand);
-			PropertyName->SetOnMouseButtonUp(FPointerEventHandler::CreateSP(this, &SDMMaterialProperties::OnPropertyClicked, MaterialProperty));
+		PreviewWidgetContainer->SetCursor(EMouseCursor::Hand);
+		PreviewWidgetContainer->SetOnMouseButtonUp(FPointerEventHandler::CreateSP(this, &SDMMaterialProperties::OnPropertyClicked, MaterialProperty));
 
-			Slider = CreateGlobalSlider(InProperty);
-		}
-		else if (!bValidProperty)
-		{
-			Slider = SNew(STextBlock)
-				.Font(IDetailLayoutBuilder::GetDetailFont())
-				.ColorAndOpacity(FStyleColors::Notifications)
-				.Text(INVTEXT("Not valid for material type."));
+		PropertyName->SetCursor(EMouseCursor::Hand);
+		PropertyName->SetOnMouseButtonUp(FPointerEventHandler::CreateSP(this, &SDMMaterialProperties::OnPropertyClicked, MaterialProperty));
 
-			bPadSlider = true;
-		}
-
-		// If the property is just disabled, leave the slider widget blank.
+		Slider = CreateGlobalSlider(InProperty);
 	}
-	else
-	{
-		Slider = SNew(STextBlock)
-			.Font(IDetailLayoutBuilder::GetDetailFont())
-			.ColorAndOpacity(FStyleColors::Error)
-			.Text(INVTEXT("Invalid property."));
 
-		bPadSlider = true;
-	}
+	// If the property is just disabled, leave the slider widget blank.
 
-	InList->AddSlot()
-		.Padding(0.f, 0.f, 0.f, 10.f)
-		.AutoHeight()
-		[
-			SNew(SHorizontalBox)
+	return SNew(SHorizontalBox)
 
 			+ SHorizontalBox::Slot()
 			.AutoWidth()
 			.VAlign(EVerticalAlignment::VAlign_Top)
+			.Padding(0.f, 5.f, 0.f, 5.f)
 			[
 				PreviewWidgetContainer
 			]
@@ -262,7 +296,7 @@ void SDMMaterialProperties::AddProperty(const TSharedRef<SVerticalBox>& InList,
 			.FillWidth(1.f)
 			.HAlign(EHorizontalAlignment::HAlign_Left)
 			.VAlign(EVerticalAlignment::VAlign_Fill)
-			.Padding(5.f, 0.f, 0.f, 0.f)
+			.Padding(5.f, 5.f, 0.f, 5.f)
 			[
 				SNew(SVerticalBox)
 
@@ -293,12 +327,10 @@ void SDMMaterialProperties::AddProperty(const TSharedRef<SVerticalBox>& InList,
 				.FillHeight(1.f)
 				.HAlign(EHorizontalAlignment::HAlign_Fill)
 				.VAlign(EVerticalAlignment::VAlign_Center)
-				.Padding(0.f, bPadSlider ? 5.f : 0.f, 0.f, 0.f)
 				[
 					Slider
 				]
-			]
-		];
+			];
 }
 
 TSharedRef<SWidget> SDMMaterialProperties::CreateSlot_EnabledButton(EDMMaterialPropertyType InMaterialProperty)
@@ -526,6 +558,8 @@ TSharedRef<SWidget> SDMMaterialProperties::CreateGlobalSlider(UDMMaterialPropert
 		}
 	}
 
+	TGuardValue<bool> ConstructGuard(bConstructing, true);
+
 	FCustomDetailsViewArgs Args;
 	Args.KeyframeHandler = KeyframeHandler;
 	Args.bAllowGlobalExtensions = true;
@@ -595,4 +629,35 @@ TSharedRef<SWidget> SDMMaterialProperties::CreateGlobalSlider(UDMMaterialPropert
 		];
 }
 
+void SDMMaterialProperties::OnExpansionStateChanged(const TSharedRef<ICustomDetailsViewItem>& InItem, bool bInExpansionState)
+{
+	if (bConstructing)
+	{
+		return;
+	}
+
+	const FCustomDetailsViewItemId& ItemId = InItem->GetItemId();
+
+	if (ItemId.GetItemType() != static_cast<uint32>(EDetailNodeType::Category))
+	{
+		return;
+	}
+
+	TSharedPtr<SDMMaterialEditor> EditorWidget = EditorWidgetWeak.Pin();
+
+	if (!EditorWidget.IsValid())
+	{
+		return;
+	}
+
+	UDynamicMaterialModelBase* MaterialModelBase = EditorWidget->GetMaterialModelBase();
+
+	if (!MaterialModelBase)
+	{
+		return;
+	}
+
+	FDMWidgetStatics::Get().SetExpansionState(MaterialModelBase, *ItemId.GetItemName(), bInExpansionState);
+}
+
 #undef LOCTEXT_NAMESPACE
diff --git a/Engine/Plugins/Experimental/DynamicMaterial/Source/DynamicMaterialEditor/Private/UI/Widgets/Editor/SDMMaterialProperties.h b/Engine/Plugins/Experimental/DynamicMaterial/Source/DynamicMaterialEditor/Private/UI/Widgets/Editor/SDMMaterialProperties.h
--- a/Engine/Plugins/Experimental/DynamicMaterial/Source/DynamicMaterialEditor/Private/UI/Widgets/Editor/SDMMaterialProperties.h
+++ b/Engine/Plugins/Experimental/DynamicMaterial/Source/DynamicMaterialEditor/Private/UI/Widgets/Editor/SDMMaterialProperties.h
@@ -7,6 +7,7 @@
 #include "Containers/Array.h"
 #include "UI/Utils/DMWidgetSlot.h"
 
+class ICustomDetailsView;
 class ICustomDetailsViewItem;
 class SDMMaterialEditor;
 class SVerticalBox;
@@ -36,9 +37,14 @@ protected:
 
 	TArray<TSharedRef<ICustomDetailsViewItem>> GlobalItems;
 
+	bool bConstructing = false;
+
 	TSharedRef<SWidget> CreateSlot_Content();
 
-	void AddProperty(const TSharedRef<SVerticalBox>& InList, UDMMaterialProperty* InProperty);
+	void AddProperty(const TSharedRef<ICustomDetailsView>& InDetailsView, const TSharedRef<ICustomDetailsViewItem>& InCategory, 
+		UDMMaterialProperty* InProperty);
+
+	TSharedRef<SWidget> CreatePropertyRow(UDMMaterialProperty* InProperty);
 
 	TSharedRef<SWidget> CreateSlot_EnabledButton(EDMMaterialPropertyType InMaterialProperty);
 
@@ -51,4 +57,6 @@ protected:
 	FReply OnPropertyClicked(const FGeometry& InGeometry, const FPointerEvent& InPointerEvent, EDMMaterialPropertyType InMaterialProperty);
 
 	TSharedRef<SWidget> CreateGlobalSlider(UDMMaterialProperty* InProperty);
+
+	void OnExpansionStateChanged(const TSharedRef<ICustomDetailsViewItem>& InItem, bool bInExpansionState);
 };
