Material Designer
- Moved the blend mode selector to the layer list. (part 1.)

diff --git a/Engine/Plugins/Experimental/DynamicMaterial/Source/DynamicMaterialEditor/Private/UI/Widgets/Editor/SDMMaterialSlotEditor.cpp b/Engine/Plugins/Experimental/DynamicMaterial/Source/DynamicMaterialEditor/Private/UI/Widgets/Editor/SDMMaterialSlotEditor.cpp
--- a/Engine/Plugins/Experimental/DynamicMaterial/Source/DynamicMaterialEditor/Private/UI/Widgets/Editor/SDMMaterialSlotEditor.cpp
+++ b/Engine/Plugins/Experimental/DynamicMaterial/Source/DynamicMaterialEditor/Private/UI/Widgets/Editor/SDMMaterialSlotEditor.cpp
@@ -535,77 +535,7 @@ TSharedRef<SWidget> SDMMaterialSlotEditor::CreateSlot_SlotSettings()
 {
 	FDMWidgetStatics::Get().ClearPropertyHandles(this);
 
-	TSharedRef<SHorizontalBox> NewSlotSettings = SNew(SHorizontalBox)
-
-		+ SHorizontalBox::Slot()
-		.AutoWidth()
-		[
-			CreateSlot_LayerBlendMode()
-		]
-
-		+ SHorizontalBox::Slot()
-		.FillWidth(0.5f)
-		[
-			CreateSlot_LayerOpacity()
-		];
-
-	return NewSlotSettings;
-}
-
-TSharedRef<SWidget> SDMMaterialSlotEditor::CreateSlot_LayerBlendMode()
-{
-	TSubclassOf<UDMMaterialStageBlend> SelectedBlendMode = nullptr;
-
-	if (LayerViewSlot.IsValid())
-	{
-		if (const UDMMaterialLayerObject* SelectedLayer = LayerViewSlot->GetSelectedLayer())
-		{
-			if (UDMMaterialStage* BaseStage = SelectedLayer->GetFirstEnabledStage(EDMMaterialLayerStage::Base))
-			{
-				if (UDMMaterialStageSource* BaseStageSource = BaseStage->GetSource())
-				{
-					SelectedBlendMode = BaseStageSource->GetClass();
-				}
-			}
-		}
-	}
-
-	return SNew(SHorizontalBox)
-		.IsEnabled(!bIsDynamic)
-		.ToolTipText(LOCTEXT("MaterialDesignerInstanceBlendModeTooltip", "Change the Blend Mode for selected Material Layer."))
-		
-		+ SHorizontalBox::Slot()
-		.AutoWidth()
-
-		[
-			SNew(SBox)
-			.HeightOverride(22.f)
-			.HAlign(HAlign_Left)
-			.VAlign(VAlign_Center)
-			.Padding(5.f, 3.f, 5.f, 3.f)
-			[
-				SNew(STextBlock)
-				.Font(IDetailLayoutBuilder::GetDetailFont())
-				.Text(LOCTEXT("MaterialDesignerInstanceBlendMode", "Blend"))
-			]
-		]
-
-		+ SHorizontalBox::Slot()
-		.FillWidth(1.0f)
-		.MaxWidth(105.f)
-		[
-			SNew(SBox)
-			.HeightOverride(32.f)
-			.HAlign(HAlign_Fill)
-			.VAlign(VAlign_Center)
-			.VAlign(EVerticalAlignment::VAlign_Center)
-			.Padding(0.f, 3.f, 5.f, 3.f)
-			.Visibility(SelectedBlendMode.Get() ? EVisibility::Visible : EVisibility::Hidden)
-			[
-				SNew(SDMMaterialLayerBlendMode, SharedThis(this))
-				.SelectedItem(SelectedBlendMode)
-			]
-		];
+	return CreateSlot_LayerOpacity();
 }
 
 TSharedRef<SWidget> SDMMaterialSlotEditor::CreateSlot_LayerOpacity()
diff --git a/Engine/Plugins/Experimental/DynamicMaterial/Source/DynamicMaterialEditor/Private/UI/Widgets/Editor/SDMMaterialSlotEditor.h b/Engine/Plugins/Experimental/DynamicMaterial/Source/DynamicMaterialEditor/Private/UI/Widgets/Editor/SDMMaterialSlotEditor.h
--- a/Engine/Plugins/Experimental/DynamicMaterial/Source/DynamicMaterialEditor/Private/UI/Widgets/Editor/SDMMaterialSlotEditor.h
+++ b/Engine/Plugins/Experimental/DynamicMaterial/Source/DynamicMaterialEditor/Private/UI/Widgets/Editor/SDMMaterialSlotEditor.h
@@ -86,8 +86,6 @@ protected:
 
 	TSharedRef<SWidget> CreateSlot_SlotSettings();
 
-	TSharedRef<SWidget> CreateSlot_LayerBlendMode();
-
 	TSharedRef<SWidget> CreateSlot_LayerOpacity();
 
 	TSharedRef<SDMMaterialSlotLayerView> CreateSlot_LayerView();
diff --git a/Engine/Plugins/Experimental/DynamicMaterial/Source/DynamicMaterialEditor/Private/UI/Widgets/Editor/SlotEditor/SDMMaterialLayerBlendMode.cpp b/Engine/Plugins/Experimental/DynamicMaterial/Source/DynamicMaterialEditor/Private/UI/Widgets/Editor/SlotEditor/SDMMaterialLayerBlendMode.cpp
--- a/Engine/Plugins/Experimental/DynamicMaterial/Source/DynamicMaterialEditor/Private/UI/Widgets/Editor/SlotEditor/SDMMaterialLayerBlendMode.cpp
+++ b/Engine/Plugins/Experimental/DynamicMaterial/Source/DynamicMaterialEditor/Private/UI/Widgets/Editor/SlotEditor/SDMMaterialLayerBlendMode.cpp
@@ -32,11 +32,11 @@
 #include "Components/MaterialStageBlends/DMMSBVividLight.h"
 #include "DetailLayoutBuilder.h"
 #include "DynamicMaterialEditorStyle.h"
-#include "SDMMaterialSlotLayerView.h"
 #include "ToolMenu.h"
 #include "ToolMenuDelegates.h"
 #include "ToolMenus.h"
-#include "UI/Widgets/Editor/SDMMaterialSlotEditor.h"
+#include "UI/Widgets/Editor/SlotEditor/SDMMaterialSlotLayerItem.h"
+#include "UI/Widgets/Editor/SlotEditor/SDMMaterialSlotLayerView.h"
 #include "Widgets/Input/SComboButton.h"
 #include "Widgets/SNullWidget.h"
 #include "Widgets/Text/STextBlock.h"
@@ -139,9 +139,9 @@ void UDMSourceBlendModeContextObject::SetBlendModeWidget(const TSharedPtr<SDMMat
 	BlendModeWidgetWeak = InBlendModeWidget;
 }
 
-void SDMMaterialLayerBlendMode::Construct(const FArguments& InArgs, const TSharedRef<SDMMaterialSlotEditor> InSlotEditor)
+void SDMMaterialLayerBlendMode::Construct(const FArguments& InArgs, const TSharedRef<SDMMaterialSlotLayerItem> InSlotEditor)
 {
-	SlotEditorWidgetWeak = InSlotEditor;
+	LayerItemWidgetWeak = InSlotEditor;
 	SelectedItem = InArgs._SelectedItem;
 
 	SetCanTick(false);
@@ -275,14 +275,14 @@ void SDMMaterialLayerBlendMode::OnBlendModeSelected(UClass* InBlendClass)
 {
 	SelectedItem = InBlendClass;
 
-	TSharedPtr<SDMMaterialSlotEditor> SlotEditorWidget = SlotEditorWidgetWeak.Pin();
+	TSharedPtr<SDMMaterialSlotLayerItem> LayerItemWiget = LayerItemWidgetWeak.Pin();
 
-	if (!SlotEditorWidget.IsValid())
+	if (!LayerItemWiget.IsValid())
 	{
 		return;
 	}
 
-	UDMMaterialLayerObject* Layer = SlotEditorWidget->GetLayerView()->GetSelectedLayer();
+	UDMMaterialLayerObject* Layer = LayerItemWiget->GetLayer();
 
 	if (!Layer)
 	{
@@ -300,8 +300,6 @@ void SDMMaterialLayerBlendMode::OnBlendModeSelected(UClass* InBlendClass)
 	FDMScopedUITransaction Transaction(LOCTEXT("SetStageBlendMode", "Set Blend Mode"));
 	BaseStage->Modify();
 	BaseStage->ChangeSource<UDMMaterialStageBlend>(InBlendClass);
-
-	SlotEditorWidget->InvalidateLayerSettings();
 }
 
 bool SDMMaterialLayerBlendMode::CanSelectBlendMode(UClass* InBlendClass)
@@ -316,14 +314,14 @@ bool SDMMaterialLayerBlendMode::InBlendModeSelected(UClass* InBlendClass)
 
 bool SDMMaterialLayerBlendMode::IsSelectorEnabled() const
 {
-	TSharedPtr<SDMMaterialSlotEditor> SlotEditorWidget = SlotEditorWidgetWeak.Pin();
+	TSharedPtr<SDMMaterialSlotLayerItem> LayerItemWiget = LayerItemWidgetWeak.Pin();
 
-	if (!SlotEditorWidget.IsValid())
+	if (!LayerItemWiget.IsValid())
 	{
 		return false;
 	}
 
-	UDMMaterialLayerObject* Layer = SlotEditorWidget->GetLayerView()->GetSelectedLayer();
+	UDMMaterialLayerObject* Layer = LayerItemWiget->GetLayer();
 
 	if (!Layer)
 	{
diff --git a/Engine/Plugins/Experimental/DynamicMaterial/Source/DynamicMaterialEditor/Private/UI/Widgets/Editor/SlotEditor/SDMMaterialLayerBlendMode.h b/Engine/Plugins/Experimental/DynamicMaterial/Source/DynamicMaterialEditor/Private/UI/Widgets/Editor/SlotEditor/SDMMaterialLayerBlendMode.h
--- a/Engine/Plugins/Experimental/DynamicMaterial/Source/DynamicMaterialEditor/Private/UI/Widgets/Editor/SlotEditor/SDMMaterialLayerBlendMode.h
+++ b/Engine/Plugins/Experimental/DynamicMaterial/Source/DynamicMaterialEditor/Private/UI/Widgets/Editor/SlotEditor/SDMMaterialLayerBlendMode.h
@@ -15,7 +15,7 @@
 #include "SDMMaterialLayerBlendMode.generated.h"
 
 class SDMMaterialLayerBlendMode;
-class SDMMaterialSlotEditor;
+class SDMMaterialSlotLayerItem;
 class SWidget;
 class UClass;
 class UDMMaterialStageBlend;
@@ -51,7 +51,7 @@ public:
 
 	virtual ~SDMMaterialLayerBlendMode() override = default;
 
-	void Construct(const FArguments& InArgs, const TSharedRef<SDMMaterialSlotEditor> InSlotEditor);
+	void Construct(const FArguments& InArgs, const TSharedRef<SDMMaterialSlotLayerItem> InSlotEditor);
 
 protected:
 	static TArray<TStrongObjectPtr<UClass>> SupportedBlendClasses;
@@ -61,7 +61,7 @@ protected:
 	static void EnsureMenuRegistered();
 	static void MakeSourceBlendMenu(UToolMenu* InToolMenu);
 
-	TWeakPtr<SDMMaterialSlotEditor> SlotEditorWidgetWeak;
+	TWeakPtr<SDMMaterialSlotLayerItem> LayerItemWidgetWeak;
 	TAttribute<TSubclassOf<UDMMaterialStageBlend>> SelectedItem;
 
 	TSharedRef<SWidget> OnGenerateWidget(const FName InItem);
diff --git a/Engine/Plugins/Experimental/DynamicMaterial/Source/DynamicMaterialEditor/Private/UI/Widgets/Editor/SlotEditor/SDMMaterialSlotLayerItem.cpp b/Engine/Plugins/Experimental/DynamicMaterial/Source/DynamicMaterialEditor/Private/UI/Widgets/Editor/SlotEditor/SDMMaterialSlotLayerItem.cpp
--- a/Engine/Plugins/Experimental/DynamicMaterial/Source/DynamicMaterialEditor/Private/UI/Widgets/Editor/SlotEditor/SDMMaterialSlotLayerItem.cpp
+++ b/Engine/Plugins/Experimental/DynamicMaterial/Source/DynamicMaterialEditor/Private/UI/Widgets/Editor/SlotEditor/SDMMaterialSlotLayerItem.cpp
@@ -20,6 +20,7 @@
 #include "UI/Utils/DMWidgetStatics.h"
 #include "UI/Widgets/Editor/SDMMaterialComponentEditor.h"
 #include "UI/Widgets/Editor/SDMMaterialSlotEditor.h"
+#include "UI/Widgets/Editor/SlotEditor/SDMMaterialLayerBlendMode.h"
 #include "UI/Widgets/Editor/SlotEditor/SDMMaterialSlotLayerEffectView.h"
 #include "UI/Widgets/Editor/SlotEditor/SDMMaterialSlotLayerView.h"
 #include "UI/Widgets/Editor/SlotEditor/SDMMaterialStage.h"
@@ -304,11 +305,9 @@ TSharedRef<SWidget> SDMMaterialSlotLayerItem::CreateHeaderRowContent()
 				.FillHeight(1.0f)
 				.HAlign(HAlign_Left)
 				.VAlign(VAlign_Center)
-				.Padding(0.0f, 0.0f, 0.0f, VerticalSpacing)
+				.Padding(0.0f, VerticalSpacing, 0.0f, VerticalSpacing)
 				[
-					SNew(STextBlock)
-					.TextStyle(FDynamicMaterialEditorStyle::Get(), "SmallFont")
-					.Text(this, &SDMMaterialSlotLayerItem::GetStageDescription, EDMMaterialLayerStage::All)
+					CreateBlendModeSelector()
 				]
 			]
 			+ SHorizontalBox::Slot()
@@ -594,6 +593,25 @@ TSharedRef<SWidget> SDMMaterialSlotLayerItem::CreateStageSourceButton(EDMMateria
 		];
 }
 
+TSharedRef<SWidget> SDMMaterialSlotLayerItem::CreateBlendModeSelector()
+{
+	TSubclassOf<UDMMaterialStageBlend> SelectedBlendMode = nullptr;
+
+	if (const UDMMaterialLayerObject* Layer = GetLayer())
+	{
+		if (UDMMaterialStage* BaseStage = Layer->GetFirstEnabledStage(EDMMaterialLayerStage::Base))
+		{
+			if (UDMMaterialStageSource* BaseStageSource = BaseStage->GetSource())
+			{
+				SelectedBlendMode = BaseStageSource->GetClass();
+			}
+		}
+	}
+
+	return SNew(SDMMaterialLayerBlendMode, SharedThis(this))
+		.SelectedItem(SelectedBlendMode);
+}
+
 EVisibility SDMMaterialSlotLayerItem::GetEffectsListVisibility() const
 {
 	if (!LayerItem.IsValid())
@@ -916,19 +934,6 @@ FText SDMMaterialSlotLayerItem::GetBlendModeText() const
 	return LOCTEXT("Error", "Error");
 }
 
-FText SDMMaterialSlotLayerItem::GetStageDescription(EDMMaterialLayerStage InLayerStage) const
-{
-	if (UDMMaterialLayerObject* Layer = GetLayer())
-	{
-		if (UDMMaterialStage* Stage = Layer->GetFirstEnabledStage(InLayerStage))
-		{
-			return Stage->GetComponentDescription();
-		}
-	}
-
-	return GetDefault<UDMMaterialStage>()->GetComponentDescription();
-}
-
 const FSlateBrush* SDMMaterialSlotLayerItem::GetStageSourceButtonImage(EDMMaterialLayerStage InLayerStage) const
 {
 	if (UDMMaterialLayerObject* Layer = GetLayer())
diff --git a/Engine/Plugins/Experimental/DynamicMaterial/Source/DynamicMaterialEditor/Private/UI/Widgets/Editor/SlotEditor/SDMMaterialSlotLayerItem.h b/Engine/Plugins/Experimental/DynamicMaterial/Source/DynamicMaterialEditor/Private/UI/Widgets/Editor/SlotEditor/SDMMaterialSlotLayerItem.h
--- a/Engine/Plugins/Experimental/DynamicMaterial/Source/DynamicMaterialEditor/Private/UI/Widgets/Editor/SlotEditor/SDMMaterialSlotLayerItem.h
+++ b/Engine/Plugins/Experimental/DynamicMaterial/Source/DynamicMaterialEditor/Private/UI/Widgets/Editor/SlotEditor/SDMMaterialSlotLayerItem.h
@@ -82,6 +82,7 @@ protected:
 	TSharedPtr<SWidget> CreateLayerHeaderEditableText();
 	TSharedRef<SWidget> CreateEffectsToggleButton();
 	TSharedRef<SWidget> CreateStageSourceButton(EDMMaterialLayerStage InStage);
+	TSharedRef<SWidget> CreateBlendModeSelector();
 
 	EVisibility GetEffectsListVisibility() const;
 	EVisibility GetEffectsToggleButtonVisibility() const;
@@ -113,7 +114,6 @@ protected:
 	FText GetLayerHeaderText() const;
 	FText GetLayerIndexText() const;
 	FText GetBlendModeText() const;
-	FText GetStageDescription(EDMMaterialLayerStage InLayerStage) const;
 
 	/** Drag and Drop */
 	int32 OnLayerItemPaintDropIndicator(EItemDropZone InItemDropZone, const FPaintArgs& InArgs, const FGeometry& InAllottedGeometry,
