Material Designer
- Texture sets and multiple textures can now be dropped into the entire editor window.

diff --git a/Engine/Plugins/Experimental/DynamicMaterial/Source/DynamicMaterialEditor/Private/UI/Widgets/Editor/SDMMaterialSlotEditor.cpp b/Engine/Plugins/Experimental/DynamicMaterial/Source/DynamicMaterialEditor/Private/UI/Widgets/Editor/SDMMaterialSlotEditor.cpp
--- a/Engine/Plugins/Experimental/DynamicMaterial/Source/DynamicMaterialEditor/Private/UI/Widgets/Editor/SDMMaterialSlotEditor.cpp
+++ b/Engine/Plugins/Experimental/DynamicMaterial/Source/DynamicMaterialEditor/Private/UI/Widgets/Editor/SDMMaterialSlotEditor.cpp
@@ -20,7 +20,6 @@
 #include "CustomDetailsViewModule.h"
 #include "DetailLayoutBuilder.h"
 #include "DMTextureSet.h"
-#include "DMTextureSetBlueprintFunctionLibrary.h"
 #include "DMWorldSubsystem.h"
 #include "DynamicMaterialEditorStyle.h"
 #include "DynamicMaterialModule.h"
@@ -29,7 +28,6 @@
 #include "HAL/PlatformApplicationMisc.h"
 #include "Items/ICustomDetailsViewItem.h"
 #include "Materials/MaterialFunctionInterface.h"
-#include "Misc/MessageDialog.h"
 #include "Model/DynamicMaterialModel.h"
 #include "Model/DynamicMaterialModelDynamic.h"
 #include "Model/DynamicMaterialModelEditorOnlyData.h"
@@ -40,8 +38,8 @@
 #include "UI/DragDrop/DMSlotLayerDragDropOperation.h"
 #include "UI/Menus/DMMaterialSlotLayerAddEffectMenus.h"
 #include "UI/Menus/DMMaterialSlotLayerMenus.h"
+#include "UI/Utils/DMDropTargetPrivateSetter.h"
 #include "UI/Utils/DMWidgetStatics.h"
-#include "UI/Widgets/Editor/SlotEditor/SDMMaterialLayerBlendMode.h"
 #include "UI/Widgets/Editor/SlotEditor/SDMMaterialSlotLayerView.h"
 #include "UI/Widgets/SDMMaterialEditor.h"
 #include "Utils/DMMaterialSlotFunctionLibrary.h"
@@ -528,6 +526,9 @@ TSharedRef<SWidget> SDMMaterialSlotEditor::CreateSlot_Container()
 	const_cast<FSlotBase&>(DropTargetOverlayChildren->GetSlotAt(0)).AttachWidget(SecondChild);
 	const_cast<FSlotBase&>(DropTargetOverlayChildren->GetSlotAt(1)).AttachWidget(FirstChild);
 
+	using namespace UE::DynamicMaterialEditor::Private;
+	DropTarget::SetInvalidColor(DropTarget.Get(), FStyleColors::Transparent);
+
 	return NewContainer;
 }
 
@@ -853,11 +854,11 @@ bool SDMMaterialSlotEditor::OnAreAssetsAcceptableForDrop(TArrayView<FAssetData>
 	}
 
 	const TArray<UClass*> AllowedClasses = {
-		UTexture::StaticClass(),
-		UDMTextureSet::StaticClass(),
 		UMaterialFunctionInterface::StaticClass()
 	};
 
+	TArray<FAssetData> DroppedTextures;
+
 	for (const FAssetData& Asset : InAssets)
 	{
 		UClass* AssetClass = Asset.GetClass(EResolveClass::Yes);
@@ -874,6 +875,16 @@ bool SDMMaterialSlotEditor::OnAreAssetsAcceptableForDrop(TArrayView<FAssetData>
 				return true;
 			}
 		}
+
+		if (AssetClass->IsChildOf(UTexture::StaticClass()))
+		{
+			DroppedTextures.Add(Asset);
+		}
+	}
+
+	if (DroppedTextures.Num() == 1)
+	{
+		return true;
 	}
 
 	return false;
@@ -992,24 +1003,14 @@ void SDMMaterialSlotEditor::HandleDrop_Texture(UTexture* InTexture)
 
 void SDMMaterialSlotEditor::HandleDrop_CreateTextureSet(const TArray<FAssetData>& InTextureAssets)
 {
-	if (InTextureAssets.Num() < 2)
+	TSharedPtr<SDMMaterialEditor> EditorWidget = GetEditorWidget();
+
+	if (!EditorWidget.IsValid())
 	{
 		return;
 	}
 
-	UDMTextureSetBlueprintFunctionLibrary::CreateTextureSetFromAssetsInteractive(
-		InTextureAssets,
-		FDMTextureSetBuilderOnComplete::CreateSPLambda(
-			this,
-			[this](UDMTextureSet* InTextureSet, bool bInWasAccepted)
-			{
-				if (bInWasAccepted)
-				{
-					HandleDrop_TextureSet(InTextureSet);
-				}
-			}
-		)
-	);
+	EditorWidget->HandleDrop_CreateTextureSet(InTextureAssets);
 }
 
 void SDMMaterialSlotEditor::HandleDrop_TextureSet(UDMTextureSet* InTextureSet)
@@ -1021,48 +1022,7 @@ void SDMMaterialSlotEditor::HandleDrop_TextureSet(UDMTextureSet* InTextureSet)
 		return;
 	}
 
-	UDynamicMaterialModel* MaterialModel = EditorWidget->GetMaterialModel();
-
-	if (!MaterialModel)
-	{
-		return;
-	}
-
-	UDynamicMaterialModelEditorOnlyData* EditorOnlyData = UDynamicMaterialModelEditorOnlyData::Get(MaterialModel);
-
-	if (!EditorOnlyData)
-	{
-		return;
-	}
-
-	const EAppReturnType::Type Result = FMessageDialog::Open(
-		EAppMsgType::YesNoCancel,
-		LOCTEXT("ReplaceSlotsTextureSet",
-			"You are about to import a Material Designer Texture Set.\n\n"
-			"Do you want to replace the slot contents?\n"
-			"- Yes: All layers are deleted in the matching slots.\n"
-			"- No: New texture layers are added to the matching slots.\n"
-			"- Cancel: Abort this operation.")
-	);
-
-	FDMScopedUITransaction Transaction(LOCTEXT("DropTextureSet", "Drop Texture Set"));
-
-	switch (Result)
-	{
-		case EAppReturnType::No:
-			EditorOnlyData->Modify();
-			EditorOnlyData->AddTextureSet(InTextureSet, /* Replace */ false);
-			break;
-
-		case EAppReturnType::Yes:
-			EditorOnlyData->Modify();
-			EditorOnlyData->AddTextureSet(InTextureSet, /* Replace */ true);
-			break;
-
-		default:
-			Transaction.Transaction.Cancel();
-			break;
-	}
+	EditorWidget->HandleDrop_TextureSet(InTextureSet);
 }
 
 void SDMMaterialSlotEditor::HandleDrop_MaterialFunction(UMaterialFunctionInterface* InMaterialFunction)
diff --git a/Engine/Plugins/Experimental/DynamicMaterial/Source/DynamicMaterialEditor/Private/UI/Widgets/SDMMaterialDesigner.cpp b/Engine/Plugins/Experimental/DynamicMaterial/Source/DynamicMaterialEditor/Private/UI/Widgets/SDMMaterialDesigner.cpp
--- a/Engine/Plugins/Experimental/DynamicMaterial/Source/DynamicMaterialEditor/Private/UI/Widgets/SDMMaterialDesigner.cpp
+++ b/Engine/Plugins/Experimental/DynamicMaterial/Source/DynamicMaterialEditor/Private/UI/Widgets/SDMMaterialDesigner.cpp
@@ -3,6 +3,7 @@
 #include "UI/Widgets/SDMMaterialDesigner.h"
 
 #include "Containers/Set.h"
+#include "DMTextureSet.h"
 #include "DynamicMaterialEditorSettings.h"
 #include "DynamicMaterialModule.h"
 #include "Editor/SDMMaterialComponentEditor.h"
@@ -312,6 +313,7 @@ void SDMMaterialDesigner::SetWidget(const TSharedRef<SWidget>& InWidget, bool bI
 		TSharedRef<SAssetDropTarget> DropTarget = SNew(SAssetDropTarget)
 			.OnAreAssetsAcceptableForDrop(this, &SDMMaterialDesigner::OnAssetDraggedOver)
 			.OnAssetsDropped(this, &SDMMaterialDesigner::OnAssetsDropped)
+			.bSupportsMultiDrop(true)
 			[
 				InWidget
 			];
@@ -351,12 +353,22 @@ bool SDMMaterialDesigner::NeedsWizard(UDynamicMaterialModelBase* InMaterialModel
 
 bool SDMMaterialDesigner::OnAssetDraggedOver(TArrayView<FAssetData> InAssets)
 {
-	const TArray<UClass*> AllowedClasses = {
+	TArray<UClass*> AllowedClasses = {
 		AActor::StaticClass(),
 		UDynamicMaterialModelBase::StaticClass(),
 		UDynamicMaterialInstance::StaticClass()
 	};
 
+	const bool bIsEditor = Content.IsValid()
+		&& Content->GetWidgetClass().GetWidgetType() == SDMMaterialEditor::StaticWidgetClass().GetWidgetType();
+
+	if (bIsEditor)
+	{
+		AllowedClasses.Add(UDMTextureSet::StaticClass());
+	}
+
+	TArray<FAssetData> DroppedTextures;
+
 	for (const FAssetData& Asset : InAssets)
 	{
 		UClass* AssetClass = Asset.GetClass(EResolveClass::Yes);
@@ -373,6 +385,16 @@ bool SDMMaterialDesigner::OnAssetDraggedOver(TArrayView<FAssetData> InAssets)
 				return true;
 			}
 		}
+
+		if (AssetClass->IsChildOf(UTexture::StaticClass()))
+		{
+			DroppedTextures.Add(Asset);
+		}
+	}
+
+	if (bIsEditor && DroppedTextures.Num() > 1)
+	{
+		return true;
 	}
 
 	return false;
@@ -380,6 +402,8 @@ bool SDMMaterialDesigner::OnAssetDraggedOver(TArrayView<FAssetData> InAssets)
 
 void SDMMaterialDesigner::OnAssetsDropped(const FDragDropEvent& InDragDropEvent, TArrayView<FAssetData> InAssets)
 {
+	TArray<FAssetData> DroppedTextures;
+
 	for (const FAssetData& Asset : InAssets)
 	{
 		UClass* AssetClass = Asset.GetClass(EResolveClass::Yes);
@@ -410,6 +434,26 @@ void SDMMaterialDesigner::OnAssetsDropped(const FDragDropEvent& InDragDropEvent,
 				return;
 			}
 		}
+		else if (AssetClass->IsChildOf(UTexture::StaticClass()))
+		{
+			DroppedTextures.Add(Asset);
+		}
+		else if (AssetClass->IsChildOf(UDMTextureSet::StaticClass()))
+		{
+			if (Content.IsValid() && Content->GetWidgetClass().GetWidgetType() == SDMMaterialEditor::StaticWidgetClass().GetWidgetType())
+			{
+				StaticCastSharedPtr<SDMMaterialEditor>(Content)->HandleDrop_CreateTextureSet({Asset});
+				return;
+			}
+		}
+	}
+
+	if (DroppedTextures.Num() > 1)
+	{
+		if (Content.IsValid() && Content->GetWidgetClass().GetWidgetType() == SDMMaterialEditor::StaticWidgetClass().GetWidgetType())
+		{
+			StaticCastSharedPtr<SDMMaterialEditor>(Content)->HandleDrop_CreateTextureSet(DroppedTextures);
+		}
 	}
 }
 
diff --git a/Engine/Plugins/Experimental/DynamicMaterial/Source/DynamicMaterialEditor/Private/UI/Widgets/SDMMaterialEditor.cpp b/Engine/Plugins/Experimental/DynamicMaterial/Source/DynamicMaterialEditor/Private/UI/Widgets/SDMMaterialEditor.cpp
--- a/Engine/Plugins/Experimental/DynamicMaterial/Source/DynamicMaterialEditor/Private/UI/Widgets/SDMMaterialEditor.cpp
+++ b/Engine/Plugins/Experimental/DynamicMaterial/Source/DynamicMaterialEditor/Private/UI/Widgets/SDMMaterialEditor.cpp
@@ -6,6 +6,7 @@
 #include "Components/DMMaterialProperty.h"
 #include "Components/DMMaterialSlot.h"
 #include "Components/DMMaterialStage.h"
+#include "DMTextureSetBlueprintFunctionLibrary.h"
 #include "DynamicMaterialEditorCommands.h"
 #include "DynamicMaterialEditorSettings.h"
 #include "DynamicMaterialModule.h"
@@ -14,6 +15,7 @@
 #include "Framework/Commands/GenericCommands.h"
 #include "Materials/Material.h"
 #include "Materials/MaterialFunctionInterface.h"
+#include "Misc/MessageDialog.h"
 #include "Model/DynamicMaterialModel.h"
 #include "Model/DynamicMaterialModelBase.h"
 #include "Model/DynamicMaterialModelDynamic.h"
@@ -30,6 +32,7 @@
 #include "UI/Widgets/Editor/SDMToolBar.h"
 #include "UI/Widgets/SDMMaterialDesigner.h"
 #include "Utils/DMMaterialModelFunctionLibrary.h"
+#include "Utils/DMPrivate.h"
 #include "Widgets/Docking/SDockTab.h"
 #include "Widgets/Layout/SBox.h"
 #include "Widgets/Layout/SScrollBox.h"
@@ -843,6 +846,79 @@ bool SDMMaterialEditor::SetActivePage(const FDMMaterialEditorPage& InPage)
 	}
 }
 
+void SDMMaterialEditor::HandleDrop_CreateTextureSet(const TArray<FAssetData>& InTextureAssets)
+{
+	if (InTextureAssets.Num() < 2)
+	{
+		return;
+	}
+
+	UDMTextureSetBlueprintFunctionLibrary::CreateTextureSetFromAssetsInteractive(
+		InTextureAssets,
+		FDMTextureSetBuilderOnComplete::CreateSPLambda(
+			this,
+			[this](UDMTextureSet* InTextureSet, bool bInWasAccepted)
+			{
+				if (bInWasAccepted)
+				{
+					HandleDrop_TextureSet(InTextureSet);
+				}
+			}
+		)
+	);
+}
+
+void SDMMaterialEditor::HandleDrop_TextureSet(UDMTextureSet* InTextureSet)
+{
+	if (!InTextureSet)
+	{
+		return;
+	}
+
+	UDynamicMaterialModel* MaterialModel = GetMaterialModel();
+
+	if (!MaterialModel)
+	{
+		return;
+	}
+
+	UDynamicMaterialModelEditorOnlyData* EditorOnlyData = UDynamicMaterialModelEditorOnlyData::Get(MaterialModel);
+
+	if (!EditorOnlyData)
+	{
+		return;
+	}
+
+	const EAppReturnType::Type Result = FMessageDialog::Open(
+		EAppMsgType::YesNoCancel,
+		LOCTEXT("ReplaceSlotsTextureSet",
+			"You are about to import a Material Designer Texture Set.\n\n"
+			"Do you want to replace the slot contents?\n"
+			"- Yes: All layers are deleted in the matching slots.\n"
+			"- No: New texture layers are added to the matching slots.\n"
+			"- Cancel: Abort this operation.")
+	);
+
+	FDMScopedUITransaction Transaction(LOCTEXT("DropTextureSet", "Drop Texture Set"));
+
+	switch (Result)
+	{
+		case EAppReturnType::No:
+			EditorOnlyData->Modify();
+			EditorOnlyData->AddTextureSet(InTextureSet, /* Replace */ false);
+			break;
+
+		case EAppReturnType::Yes:
+			EditorOnlyData->Modify();
+			EditorOnlyData->AddTextureSet(InTextureSet, /* Replace */ true);
+			break;
+
+		default:
+			Transaction.Transaction.Cancel();
+			break;
+	}
+}
+
 bool SDMMaterialEditor::PageHistoryBack()
 {
 	const int32 NewPageIndex = PageHistoryActive - 1;
diff --git a/Engine/Plugins/Experimental/DynamicMaterial/Source/DynamicMaterialEditor/Private/UI/Widgets/SDMMaterialEditor.h b/Engine/Plugins/Experimental/DynamicMaterial/Source/DynamicMaterialEditor/Private/UI/Widgets/SDMMaterialEditor.h
--- a/Engine/Plugins/Experimental/DynamicMaterial/Source/DynamicMaterialEditor/Private/UI/Widgets/SDMMaterialEditor.h
+++ b/Engine/Plugins/Experimental/DynamicMaterial/Source/DynamicMaterialEditor/Private/UI/Widgets/SDMMaterialEditor.h
@@ -143,6 +143,10 @@ public:
 
 	bool SetActivePage(const FDMMaterialEditorPage& InPage);
 
+	void HandleDrop_CreateTextureSet(const TArray<FAssetData>& InTextureAssets);
+
+	void HandleDrop_TextureSet(UDMTextureSet* InTextureSet);
+
 	//~ Begin SWidget
 	virtual bool SupportsKeyboardFocus() const override;
 	virtual FReply OnKeyDown(const FGeometry& InMyGeometry, const FKeyEvent& InKeyEvent) override;
