Material Designer
- Extended the width of the opacity slider.

diff --git a/Engine/Plugins/Experimental/DynamicMaterial/Source/DynamicMaterialEditor/Private/UI/Widgets/Editor/SDMMaterialSlotEditor.cpp b/Engine/Plugins/Experimental/DynamicMaterial/Source/DynamicMaterialEditor/Private/UI/Widgets/Editor/SDMMaterialSlotEditor.cpp
--- a/Engine/Plugins/Experimental/DynamicMaterial/Source/DynamicMaterialEditor/Private/UI/Widgets/Editor/SDMMaterialSlotEditor.cpp
+++ b/Engine/Plugins/Experimental/DynamicMaterial/Source/DynamicMaterialEditor/Private/UI/Widgets/Editor/SDMMaterialSlotEditor.cpp
@@ -32,8 +32,6 @@
 #include "Model/DynamicMaterialModelDynamic.h"
 #include "Model/DynamicMaterialModelEditorOnlyData.h"
 #include "SAssetDropTarget.h"
-#include "ToolMenu.h"
-#include "ToolMenus.h"
 #include "UI/DragDrop/DMLayerEffectsDragDropOperation.h"
 #include "UI/DragDrop/DMSlotLayerDragDropOperation.h"
 #include "UI/Menus/DMMaterialSlotLayerAddEffectMenus.h"
@@ -55,6 +53,48 @@
 
 #define LOCTEXT_NAMESPACE "SDMMaterialSlotEditor"
 
+namespace UE::DynamicMaterialEditor::Private
+{
+	const FLazyName PropertyValueWidget = TEXT("SPropertyValueWidget");
+
+	TSharedPtr<SWidget> FindWidgetInHierarchy(const TSharedRef<SWidget>& InParent, const FName& InName)
+	{
+		if (InParent->GetType() == InName)
+		{
+			return InParent;
+		}
+
+		FChildren* Children = InParent->GetChildren();
+
+		if (!Children)
+		{
+			return nullptr;
+		}
+
+		const int32 ChildNum = Children->Num();
+
+		for (int32 Index = 0; Index < ChildNum; ++Index)
+		{
+			const TSharedRef<SWidget>& Widget = Children->GetChildAt(Index);
+
+			if (Widget->GetType() == InName)
+			{
+				return Widget;
+			}
+		}
+
+		for (int32 Index = 0; Index < ChildNum; ++Index)
+		{
+			if (TSharedPtr<SWidget> FoundChild = FindWidgetInHierarchy(Children->GetChildAt(Index), InName))
+			{
+				return FoundChild;
+			}
+		}
+
+		return nullptr;
+	}
+}
+
 void SDMMaterialSlotEditor::PrivateRegisterAttributes(FSlateAttributeDescriptor::FInitializer&)
 {
 }
@@ -575,6 +615,8 @@ TSharedRef<SWidget> SDMMaterialSlotEditor::CreateSlot_LayerOpacity()
 							}
 						}
 
+						FDMPropertyHandle PropertyHandle = FDMWidgetStatics::Get().GetPropertyHandle(this, OpacityValue, UDMMaterialValue::ValueName);
+
 						FCustomDetailsViewArgs Args;
 						Args.KeyframeHandler = KeyframeHandler;
 						Args.bAllowGlobalExtensions = true;
@@ -584,8 +626,6 @@ TSharedRef<SWidget> SDMMaterialSlotEditor::CreateSlot_LayerOpacity()
 						TSharedRef<ICustomDetailsView> DetailsView = ICustomDetailsViewModule::Get().CreateCustomDetailsView(Args);
 						FCustomDetailsViewItemId RootId = DetailsView->GetRootItem()->GetItemId();
 
-						FDMPropertyHandle PropertyHandle = FDMWidgetStatics::Get().GetPropertyHandle(this, OpacityValue, UDMMaterialValue::ValueName);
-
 						LayerOpacityItem = DetailsView->CreateDetailTreeItem(PropertyHandle.DetailTreeNode.ToSharedRef());
 
 						LayerOpacityItem->SetResetToDefaultOverride(FResetToDefaultOverride::Create(
@@ -603,6 +643,25 @@ TSharedRef<SWidget> SDMMaterialSlotEditor::CreateSlot_LayerOpacity()
 	TSharedPtr<SWidget> ValueWidget = LayerOpacityItem.IsValid() ? LayerOpacityItem->GetWidget(ECustomDetailsViewWidgetType::Value) : nullptr;
 	TSharedPtr<SWidget> ExtensionWidget = LayerOpacityItem.IsValid() ? LayerOpacityItem->GetWidget(ECustomDetailsViewWidgetType::Extensions) : nullptr;
 
+	using namespace UE::DynamicMaterialEditor::Private;
+
+	if (ValueWidget.IsValid())
+	{
+		if (TSharedPtr<SWidget> FoundPropertyValueWidget = FindWidgetInHierarchy(ValueWidget.ToSharedRef(), PropertyValueWidget))
+		{
+			if (FoundPropertyValueWidget.IsValid())
+			{
+				if (FChildren* Children = FoundPropertyValueWidget->GetChildren())
+				{
+					if (Children->Num() > 0)
+					{
+						ValueWidget = Children->GetChildAt(0);
+					}
+				}
+			}
+		}
+	}
+
 	return SNew(SHorizontalBox)
 		.ToolTipText(LOCTEXT("MaterialDesignerInstanceLayerOpacityTooltip", "Change the Opacity of the selected Material Layer."))
 		
