Material Designer
- Added the ability to change stage and property preview sizes.

diff --git a/Engine/Plugins/Experimental/DynamicMaterial/Source/DynamicMaterialEditor/Private/DynamicMaterialEditorSettings.cpp b/Engine/Plugins/Experimental/DynamicMaterial/Source/DynamicMaterialEditor/Private/DynamicMaterialEditorSettings.cpp
--- a/Engine/Plugins/Experimental/DynamicMaterial/Source/DynamicMaterialEditor/Private/DynamicMaterialEditorSettings.cpp
+++ b/Engine/Plugins/Experimental/DynamicMaterial/Source/DynamicMaterialEditor/Private/DynamicMaterialEditorSettings.cpp
@@ -272,7 +272,9 @@ void UDynamicMaterialEditorSettings::PostEditChangeProperty(FPropertyChangedEven
 		|| PropertyName == GET_MEMBER_NAME_CHECKED(UDynamicMaterialEditorSettings, SplitterLocation)
 		|| PropertyName == GET_MEMBER_NAME_CHECKED(UDynamicMaterialEditorSettings, bPreviewImagesUseTextureUVs)
 		|| PropertyName == GET_MEMBER_NAME_CHECKED(UDynamicMaterialEditorSettings, PreviewMesh)
-		|| PropertyName == GET_MEMBER_NAME_CHECKED(UDynamicMaterialEditorSettings, bShowPreviewBackground))
+		|| PropertyName == GET_MEMBER_NAME_CHECKED(UDynamicMaterialEditorSettings, bShowPreviewBackground)
+		|| PropertyName == GET_MEMBER_NAME_CHECKED(UDynamicMaterialEditorSettings, StagePreviewSize)
+		|| PropertyName == GET_MEMBER_NAME_CHECKED(UDynamicMaterialEditorSettings, PropertyPreviewSize))
 	{
 		OnSettingsChanged.Broadcast(InPropertyChangedEvent);
 	}
@@ -308,6 +310,8 @@ void UDynamicMaterialEditorSettings::ResetAllLayoutSettings()
 	bShowPreviewBackground = true;
 	bPreviewImagesUseTextureUVs = true;
 	bUVVisualizerVisible = true;
+	StagePreviewSize = 40.f;
+	PropertyPreviewSize = 64.f;
 }
 
 TArray<FDMMaterialEffectList> UDynamicMaterialEditorSettings::GetEffectList() const
diff --git a/Engine/Plugins/Experimental/DynamicMaterial/Source/DynamicMaterialEditor/Private/DynamicMaterialEditorSettings.h b/Engine/Plugins/Experimental/DynamicMaterial/Source/DynamicMaterialEditor/Private/DynamicMaterialEditorSettings.h
--- a/Engine/Plugins/Experimental/DynamicMaterial/Source/DynamicMaterialEditor/Private/DynamicMaterialEditorSettings.h
+++ b/Engine/Plugins/Experimental/DynamicMaterial/Source/DynamicMaterialEditor/Private/DynamicMaterialEditorSettings.h
@@ -190,6 +190,12 @@ public:
 	UPROPERTY(Config, EditAnywhere, BlueprintReadWrite, Category = "Layout")
 	bool bUseFullChannelNamesInTopSlimLayout;
 
+	UPROPERTY(Config, EditAnywhere, BlueprintReadWrite, Category = "Layout")
+	float StagePreviewSize;
+
+	UPROPERTY(Config, EditAnywhere, BlueprintReadWrite, Category = "Layout", meta = (DisplayName = "Channel Preview Size"))
+	float PropertyPreviewSize;
+
 	UPROPERTY(Config, EditAnywhere, BlueprintReadWrite, Category = "Thumnnails")
 	double ThumbnailSize;
 
diff --git a/Engine/Plugins/Experimental/DynamicMaterial/Source/DynamicMaterialEditor/Private/UI/Widgets/Editor/SDMMaterialProperties.cpp b/Engine/Plugins/Experimental/DynamicMaterial/Source/DynamicMaterialEditor/Private/UI/Widgets/Editor/SDMMaterialProperties.cpp
--- a/Engine/Plugins/Experimental/DynamicMaterial/Source/DynamicMaterialEditor/Private/UI/Widgets/Editor/SDMMaterialProperties.cpp
+++ b/Engine/Plugins/Experimental/DynamicMaterial/Source/DynamicMaterialEditor/Private/UI/Widgets/Editor/SDMMaterialProperties.cpp
@@ -9,6 +9,8 @@
 #include "CustomDetailsViewModule.h"
 #include "DetailLayoutBuilder.h"
 #include "DMWorldSubsystem.h"
+#include "DynamicMaterialEditorSettings.h"
+#include "DynamicMaterialModule.h"
 #include "Engine/World.h"
 #include "ICustomDetailsView.h"
 #include "Items/ICustomDetailsViewCustomCategoryItem.h"
@@ -35,11 +37,29 @@ void SDMMaterialProperties::PrivateRegisterAttributes(FSlateAttributeDescriptor:
 {
 }
 
+SDMMaterialProperties::~SDMMaterialProperties()
+{
+	if (!FDynamicMaterialModule::AreUObjectsSafe())
+	{
+		return;
+	}
+
+	if (UDynamicMaterialEditorSettings* Settings = UDynamicMaterialEditorSettings::Get())
+	{
+		Settings->GetOnSettingsChanged().RemoveAll(this);
+	}
+}
+
 void SDMMaterialProperties::Construct(const FArguments& InArgs, const TSharedRef<SDMMaterialEditor>& InEditorWidget)
 {
 	EditorWidgetWeak = InEditorWidget;
 
 	Content = TDMWidgetSlot<SWidget>(SharedThis(this), 0, CreateSlot_Content());
+
+	if (UDynamicMaterialEditorSettings* Settings = UDynamicMaterialEditorSettings::Get())
+	{
+		Settings->GetOnSettingsChanged().AddSP(this, &SDMMaterialProperties::OnSettingsUpdated);
+	}
 }
 
 void SDMMaterialProperties::Validate()
@@ -201,6 +221,13 @@ void SDMMaterialProperties::AddProperty(const TSharedRef<ICustomDetailsView>& In
 		return;
 	}
 
+	UDynamicMaterialEditorSettings* Settings = UDynamicMaterialEditorSettings::Get();
+
+	if (!Settings)
+	{
+		return;
+	}
+
 	TSharedPtr<SDMMaterialEditor> EditorWidget = EditorWidgetWeak.Pin();
 
 	if (!EditorWidget.IsValid())
@@ -240,10 +267,11 @@ TSharedRef<SWidget> SDMMaterialProperties::CreatePropertyRow(UDMMaterialProperty
 	TSharedPtr<SDMMaterialEditor> EditorWidget = EditorWidgetWeak.Pin();
 	UDynamicMaterialModel* MaterialModel = EditorWidget->GetMaterialModel();
 	UDynamicMaterialModelEditorOnlyData* EditorOnlyData = UDynamicMaterialModelEditorOnlyData::Get(MaterialModel);
-	EDMMaterialPropertyType MaterialProperty = InProperty->GetMaterialProperty();
+	UDynamicMaterialEditorSettings* Settings = UDynamicMaterialEditorSettings::Get();
+	const EDMMaterialPropertyType MaterialProperty = InProperty->GetMaterialProperty();
 
 	TSharedRef<SBox> PreviewWidgetContainer = SNew(SBox)
-		.WidthOverride(64.f)
+		.WidthOverride(Settings->PropertyPreviewSize + 4.f)
 		.HeightOverride(18.f);
 
 	TSharedRef<SWidget> PropertyName = CreateSlot_PropertyName(MaterialProperty);
@@ -255,22 +283,25 @@ TSharedRef<SWidget> SDMMaterialProperties::CreatePropertyRow(UDMMaterialProperty
 
 	if (bPropertyEnabled)
 	{
-		MaterialProperty = InProperty->GetMaterialProperty();
+		TSharedPtr<SDMMaterialComponentPreview> PreviewImage;
 
 		PreviewWidgetContainer = SNew(SBox)
-			.WidthOverride(64.f)
-			.HeightOverride(64.f)
+			.WidthOverride(Settings->PropertyPreviewSize + 4.f)
+			.HeightOverride(Settings->PropertyPreviewSize + 4.f)
 			[
 				SNew(SBorder)
 				.BorderBackgroundColor(FLinearColor(1, 1, 1, 1.f))
 				.Padding(2.0f)
 				.BorderImage(FAppStyle::GetBrush(UE::DynamicMaterialEditor::Private::EditorDarkBackground))
 				[
-					SNew(SDMMaterialComponentPreview, EditorWidget.ToSharedRef(), InProperty)
-					.PreviewSize(FVector2D(60.f, 60.f))
+					SAssignNew(PreviewImage, SDMMaterialComponentPreview, EditorWidget.ToSharedRef(), InProperty)
+					.PreviewSize(FVector2D(Settings->PropertyPreviewSize))
 				]
 			];
 
+		PropertyPreviewContainers.Add(PreviewWidgetContainer);
+		PropertyPreviews.Add(PreviewImage.ToSharedRef());
+
 		PreviewWidgetContainer->SetCursor(EMouseCursor::Hand);
 		PreviewWidgetContainer->SetOnMouseButtonUp(FPointerEventHandler::CreateSP(this, &SDMMaterialProperties::OnPropertyClicked, MaterialProperty));
 
@@ -279,6 +310,10 @@ TSharedRef<SWidget> SDMMaterialProperties::CreatePropertyRow(UDMMaterialProperty
 
 		Slider = CreateGlobalSlider(InProperty);
 	}
+	else
+	{
+		PropertyEmptyContainers.Add(PreviewWidgetContainer);
+	}
 
 	// If the property is just disabled, leave the slider widget blank.
 
@@ -286,7 +321,7 @@ TSharedRef<SWidget> SDMMaterialProperties::CreatePropertyRow(UDMMaterialProperty
 
 			+ SHorizontalBox::Slot()
 			.AutoWidth()
-			.VAlign(EVerticalAlignment::VAlign_Top)
+			.VAlign(EVerticalAlignment::VAlign_Center)
 			.Padding(0.f, 5.f, 0.f, 5.f)
 			[
 				PreviewWidgetContainer
@@ -660,4 +695,35 @@ void SDMMaterialProperties::OnExpansionStateChanged(const TSharedRef<ICustomDeta
 	FDMWidgetStatics::Get().SetExpansionState(MaterialModelBase, *ItemId.GetItemName(), bInExpansionState);
 }
 
+void SDMMaterialProperties::OnSettingsUpdated(const FPropertyChangedEvent& InPropertyChangedEvent)
+{
+	UDynamicMaterialEditorSettings* Settings = UDynamicMaterialEditorSettings::Get();
+
+	if (!Settings)
+	{
+		return;
+	}
+
+	const FName MemberName = InPropertyChangedEvent.GetMemberPropertyName();
+
+	if (MemberName == GET_MEMBER_NAME_CHECKED(UDynamicMaterialEditorSettings, PropertyPreviewSize))
+	{
+		for (const TSharedRef<SBox>& PropertyPreviewContainer : PropertyPreviewContainers)
+		{
+			PropertyPreviewContainer->SetWidthOverride(Settings->PropertyPreviewSize + 4.f);
+			PropertyPreviewContainer->SetHeightOverride(Settings->PropertyPreviewSize + 4.f);
+		}
+
+		for (const TSharedRef<SBox>& PropertyEmptyContainer : PropertyEmptyContainers)
+		{
+			PropertyEmptyContainer->SetWidthOverride(Settings->PropertyPreviewSize + 4.f);
+		}
+
+		for (const TSharedRef<SDMMaterialComponentPreview>& PropertyPreview : PropertyPreviews)
+		{
+			PropertyPreview->SetPreviewSize(FVector2D(Settings->PropertyPreviewSize));
+		}
+	}
+}
+
 #undef LOCTEXT_NAMESPACE
diff --git a/Engine/Plugins/Experimental/DynamicMaterial/Source/DynamicMaterialEditor/Private/UI/Widgets/Editor/SDMMaterialProperties.h b/Engine/Plugins/Experimental/DynamicMaterial/Source/DynamicMaterialEditor/Private/UI/Widgets/Editor/SDMMaterialProperties.h
--- a/Engine/Plugins/Experimental/DynamicMaterial/Source/DynamicMaterialEditor/Private/UI/Widgets/Editor/SDMMaterialProperties.h
+++ b/Engine/Plugins/Experimental/DynamicMaterial/Source/DynamicMaterialEditor/Private/UI/Widgets/Editor/SDMMaterialProperties.h
@@ -9,12 +9,15 @@
 
 class ICustomDetailsView;
 class ICustomDetailsViewItem;
+class SBox;
+class SDMMaterialComponentPreview;
 class SDMMaterialEditor;
 class SVerticalBox;
 class UDMMaterialProperty;
 class UDMMaterialSlot;
 enum class ECheckBoxState : uint8;
 enum class EDMMaterialPropertyType : uint8;
+struct FPropertyChangedEvent;
 
 class SDMMaterialProperties : public SCompoundWidget
 {
@@ -24,7 +27,7 @@ class SDMMaterialProperties : public SCompoundWidget
 	SLATE_END_ARGS()
 
 public:
-	virtual ~SDMMaterialProperties() override = default;
+	virtual ~SDMMaterialProperties() override;
 
 	void Construct(const FArguments& InArgs, const TSharedRef<SDMMaterialEditor>& InEditorWidget);
 
@@ -36,6 +39,9 @@ protected:
 	TDMWidgetSlot<SWidget> Content;
 
 	TArray<TSharedRef<ICustomDetailsViewItem>> GlobalItems;
+	TArray<TSharedRef<SBox>> PropertyPreviewContainers;
+	TArray<TSharedRef<SBox>> PropertyEmptyContainers;
+	TArray<TSharedRef<SDMMaterialComponentPreview>> PropertyPreviews;
 
 	bool bConstructing = false;
 
@@ -59,4 +65,6 @@ protected:
 	TSharedRef<SWidget> CreateGlobalSlider(UDMMaterialProperty* InProperty);
 
 	void OnExpansionStateChanged(const TSharedRef<ICustomDetailsViewItem>& InItem, bool bInExpansionState);
+
+	void OnSettingsUpdated(const FPropertyChangedEvent& InPropertyChangedEvent);
 };
diff --git a/Engine/Plugins/Experimental/DynamicMaterial/Source/DynamicMaterialEditor/Private/UI/Widgets/Editor/SlotEditor/SDMMaterialStage.cpp b/Engine/Plugins/Experimental/DynamicMaterial/Source/DynamicMaterialEditor/Private/UI/Widgets/Editor/SlotEditor/SDMMaterialStage.cpp
--- a/Engine/Plugins/Experimental/DynamicMaterial/Source/DynamicMaterialEditor/Private/UI/Widgets/Editor/SlotEditor/SDMMaterialStage.cpp
+++ b/Engine/Plugins/Experimental/DynamicMaterial/Source/DynamicMaterialEditor/Private/UI/Widgets/Editor/SlotEditor/SDMMaterialStage.cpp
@@ -13,6 +13,7 @@
 #include "Components/MaterialValues/DMMaterialValueTexture.h"
 #include "DynamicMaterialEditorSettings.h"
 #include "DynamicMaterialEditorStyle.h"
+#include "DynamicMaterialModule.h"
 #include "Engine/Texture.h"
 #include "Framework/Application/SlateApplication.h"
 #include "SAssetDropTarget.h"
@@ -36,6 +37,19 @@ void SDMMaterialStage::PrivateRegisterAttributes(FSlateAttributeDescriptor::FIni
 {
 }
 
+SDMMaterialStage::~SDMMaterialStage()
+{
+	if (!FDynamicMaterialModule::AreUObjectsSafe())
+	{
+		return;
+	}
+
+	if (UDynamicMaterialEditorSettings* Settings = UDynamicMaterialEditorSettings::Get())
+	{
+		Settings->GetOnSettingsChanged().RemoveAll(this);
+	}
+}
+
 void SDMMaterialStage::Construct(const FArguments& InArgs, const TSharedRef<SDMMaterialSlotLayerItem>& InSlotLayerItem, 
 	UDMMaterialStage* InStage)
 {
@@ -50,6 +64,13 @@ void SDMMaterialStage::Construct(const FArguments& InArgs, const TSharedRef<SDMM
 		return;
 	}
 
+	UDynamicMaterialEditorSettings* Settings = UDynamicMaterialEditorSettings::Get();
+
+	if (!Settings)
+	{
+		return;
+	}
+
 	TSharedPtr<SDMMaterialSlotLayerView> SlotLayerView = InSlotLayerItem->GetSlotLayerView();
 
 	if (!SlotLayerView.IsValid())
@@ -85,6 +106,7 @@ void SDMMaterialStage::Construct(const FArguments& InArgs, const TSharedRef<SDMM
 			.BorderImage(FAppStyle::Get().GetBrush("ContentBrowser.AssetTileItem.DropShadow"))
 			[
 				SNew(SOverlay)
+
 				+ SOverlay::Slot()
 				.HAlign(HAlign_Center)
 				.VAlign(VAlign_Center)
@@ -97,20 +119,21 @@ void SDMMaterialStage::Construct(const FArguments& InArgs, const TSharedRef<SDMM
 						SNew(SOverlay)
 						+ SOverlay::Slot()
 						[
-							SNew(SDMMaterialComponentPreview, EditorWidget.ToSharedRef(), InStage)
-							.PreviewSize(FVector2D(40.f))
+							SAssignNew(PreviewImage, SDMMaterialComponentPreview, EditorWidget.ToSharedRef(), InStage)
+							.PreviewSize(FVector2D(Settings->StagePreviewSize))
 						]
 						+ SOverlay::Slot()
 						.Padding(5.f)
 						[
 							SNew(SImage)
 							.Image(FAppStyle::Get().GetBrush("Icons.X"))
-							.DesiredSizeOverride(FVector2D(30.f))
+							.DesiredSizeOverride(FVector2D(Settings->StagePreviewSize - 10.f))
 							.ColorAndOpacity(FStyleColors::AccentRed)
 							.Visibility(this, &SDMMaterialStage::GetDisabledOverlayVisibility)
 						]
 					]
 				]
+
 				+ SOverlay::Slot()
 				.HAlign(HAlign_Fill)
 				.VAlign(VAlign_Fill)
@@ -124,23 +147,17 @@ void SDMMaterialStage::Construct(const FArguments& InArgs, const TSharedRef<SDMM
 		]
 	];
 
-	UDynamicMaterialEditorSettings* Settings = UDynamicMaterialEditorSettings::Get();
-
 	SetToolTip(
 		SNew(SToolTip)
 		.IsInteractive(false)
 		.BorderImage(FCoreStyle::Get().GetBrush("ToolTip.Background"))
 		[
-			SNew(SDMMaterialComponentPreview, EditorWidget.ToSharedRef(), InStage)
-			.PreviewSize(TAttribute<TOptional<FVector2D>>::CreateWeakLambda(
-				Settings,
-				[Settings]()
-				{
-					return FVector2D(Settings->ThumbnailSize, Settings->ThumbnailSize);
-				}
-			))
+			SAssignNew(ToolTipImage, SDMMaterialComponentPreview, EditorWidget.ToSharedRef(), InStage)
+			.PreviewSize(FVector2D(Settings->ThumbnailSize))
 		]
 	);
+
+	Settings->GetOnSettingsChanged().AddSP(this, &SDMMaterialStage::OnSettingsUpdated);
 }
 
 TSharedPtr<SDMMaterialSlotLayerItem> SDMMaterialStage::GetSlotLayerView() const
@@ -260,6 +277,34 @@ EVisibility SDMMaterialStage::GetDisabledOverlayVisibility() const
 	return EVisibility::Collapsed;
 }
 
+
+void SDMMaterialStage::OnSettingsUpdated(const FPropertyChangedEvent& InPropertyChangedEvent)
+{
+	UDynamicMaterialEditorSettings* Settings = UDynamicMaterialEditorSettings::Get();
+
+	if (!Settings)
+	{
+		return;
+	}
+
+	const FName MemberName = InPropertyChangedEvent.GetMemberPropertyName();
+
+	if (MemberName == GET_MEMBER_NAME_CHECKED(UDynamicMaterialEditorSettings, StagePreviewSize))
+	{		
+		if (PreviewImage.IsValid())
+		{
+			PreviewImage->SetPreviewSize(FVector2D(Settings->StagePreviewSize));
+		}
+	}
+	else if (MemberName == GET_MEMBER_NAME_CHECKED(UDynamicMaterialEditorSettings, ThumbnailSize))
+	{
+		if (ToolTipImage.IsValid())
+		{
+			ToolTipImage->SetPreviewSize(FVector2D(Settings->ThumbnailSize));
+		}
+	}
+}
+
 bool SDMMaterialStage::OnAssetDraggedOver(TArrayView<FAssetData> InAssets)
 {
 	UDMMaterialStage* Stage = GetStage();
@@ -355,7 +400,7 @@ void SDMMaterialStage::HandleDrop_Texture(UTexture* InTexture)
 	}
 	else if (StageSource->IsA<UDMMaterialStageThroughputLayerBlend>())
 	{
-		bool bHasAlpha = UE::DynamicMaterial::Private::HasAlpha(InTexture);
+		const bool bHasAlpha = UE::DynamicMaterial::Private::HasAlpha(InTexture);
 
 		UDMMaterialStageInputExpression* NewInput = UDMMaterialStageInputExpression::ChangeStageInput_Expression(
 			Stage,
diff --git a/Engine/Plugins/Experimental/DynamicMaterial/Source/DynamicMaterialEditor/Private/UI/Widgets/Editor/SlotEditor/SDMMaterialStage.h b/Engine/Plugins/Experimental/DynamicMaterial/Source/DynamicMaterialEditor/Private/UI/Widgets/Editor/SlotEditor/SDMMaterialStage.h
--- a/Engine/Plugins/Experimental/DynamicMaterial/Source/DynamicMaterialEditor/Private/UI/Widgets/Editor/SlotEditor/SDMMaterialStage.h
+++ b/Engine/Plugins/Experimental/DynamicMaterial/Source/DynamicMaterialEditor/Private/UI/Widgets/Editor/SlotEditor/SDMMaterialStage.h
@@ -7,10 +7,12 @@
 #include "Templates/SharedPointer.h"
 #include "UObject/WeakObjectPtr.h"
 
+class SDMMaterialComponentPreview;
 class SDMMaterialEditor;
 class SDMMaterialSlotLayerItem;
 class UDMMaterialStage;
 class UTexture;
+struct FPropertyChangedEvent;
 
 class SDMMaterialStage : public SCompoundWidget
 {
@@ -20,7 +22,7 @@ class SDMMaterialStage : public SCompoundWidget
 	SLATE_END_ARGS()
 
 public:
-	virtual ~SDMMaterialStage() override = default;
+	virtual ~SDMMaterialStage() override;
 
 	void Construct(const FArguments& InArgs, const TSharedRef<SDMMaterialSlotLayerItem>& InSlotLayerItem, UDMMaterialStage* InStage);
 
@@ -36,6 +38,8 @@ public:
 protected:
 	TWeakPtr<SDMMaterialSlotLayerItem> SlotLayerItemWeak;
 	TWeakObjectPtr<UDMMaterialStage> StageWeak;
+	TSharedPtr<SDMMaterialComponentPreview> PreviewImage;
+	TSharedPtr<SDMMaterialComponentPreview> ToolTipImage;
 
 	bool IsStageSelected() const;
 
@@ -43,6 +47,8 @@ protected:
 
 	EVisibility GetDisabledOverlayVisibility() const;
 
+	void OnSettingsUpdated(const FPropertyChangedEvent& InPropertyChangedEvent);
+
 	/** Drag and drop. */
 	bool OnAssetDraggedOver(TArrayView<FAssetData> InAssets);
 
diff --git a/Engine/Plugins/Experimental/DynamicMaterial/Source/DynamicMaterialEditor/Private/UI/Widgets/Visualizers/SDMMaterialComponentPreview.cpp b/Engine/Plugins/Experimental/DynamicMaterial/Source/DynamicMaterialEditor/Private/UI/Widgets/Visualizers/SDMMaterialComponentPreview.cpp
--- a/Engine/Plugins/Experimental/DynamicMaterial/Source/DynamicMaterialEditor/Private/UI/Widgets/Visualizers/SDMMaterialComponentPreview.cpp
+++ b/Engine/Plugins/Experimental/DynamicMaterial/Source/DynamicMaterialEditor/Private/UI/Widgets/Visualizers/SDMMaterialComponentPreview.cpp
@@ -18,6 +18,7 @@
 
 SDMMaterialComponentPreview::SDMMaterialComponentPreview()
 	: Brush(FSlateMaterialBrush(FVector2D(1.f, 1.f)))
+	, PreviewSize(FVector2D(48.f))
 {
 	Brush.SetUVRegion(FBox2f(FVector2f::ZeroVector, FVector2f::UnitVector));
 }
@@ -50,6 +51,7 @@ void SDMMaterialComponentPreview::Construct(const FArguments& InArgs, const TSha
 {
 	EditorWidgetWeak = InEditorWidget;
 	ComponentWeak = InComponent;
+	PreviewSize = InArgs._PreviewSize;
 
 	PreviewMaterialBaseWeak = InEditorWidget->GetPreviewMaterialManager()->CreatePreviewMaterial(InComponent);
 	PreviewMaterialDynamicWeak = InEditorWidget->GetPreviewMaterialManager()->CreatePreviewMaterialDynamic(PreviewMaterialBaseWeak.Get());
@@ -65,12 +67,37 @@ void SDMMaterialComponentPreview::Construct(const FArguments& InArgs, const TSha
 
 	ChildSlot
 	[
-		SNew(SImage)
+		SAssignNew(PreviewImage, SImage)
 		.Image(&Brush)
-		.DesiredSizeOverride(InArgs._PreviewSize)
+		.DesiredSizeOverride(PreviewSize)
 	];
 }
 
+FSlateMaterialBrush& SDMMaterialComponentPreview::GetBrush()
+{
+	return Brush;
+}
+
+const FVector2D& SDMMaterialComponentPreview::GetPreviewSize() const
+{
+	return PreviewSize;
+}
+
+void SDMMaterialComponentPreview::SetPreviewSize(const FVector2D& InSize)
+{
+	if (PreviewSize.Equals(InSize))
+	{
+		return;
+	}
+
+	PreviewSize = InSize;
+
+	if (PreviewImage.IsValid())
+	{
+		PreviewImage->SetDesiredSizeOverride(PreviewSize);
+	}
+}
+
 void SDMMaterialComponentPreview::Tick(const FGeometry& AllottedGeometry, const double InCurrentTime, const float InDeltaTime)
 {
 	SCompoundWidget::Tick(AllottedGeometry, InCurrentTime, InDeltaTime);
diff --git a/Engine/Plugins/Experimental/DynamicMaterial/Source/DynamicMaterialEditor/Private/UI/Widgets/Visualizers/SDMMaterialComponentPreview.h b/Engine/Plugins/Experimental/DynamicMaterial/Source/DynamicMaterialEditor/Private/UI/Widgets/Visualizers/SDMMaterialComponentPreview.h
--- a/Engine/Plugins/Experimental/DynamicMaterial/Source/DynamicMaterialEditor/Private/UI/Widgets/Visualizers/SDMMaterialComponentPreview.h
+++ b/Engine/Plugins/Experimental/DynamicMaterial/Source/DynamicMaterialEditor/Private/UI/Widgets/Visualizers/SDMMaterialComponentPreview.h
@@ -7,6 +7,7 @@
 #include "Delegates/IDelegateInstance.h"
 #include "DMEDefs.h"
 #include "SlateMaterialBrush.h"
+#include "UI/Utils/DMWidgetSlot.h"
 #include "UObject/WeakObjectPtr.h"
 #include "UObject/WeakObjectPtrTemplates.h"
 
@@ -24,7 +25,7 @@ public:
 	SLATE_BEGIN_ARGS(SDMMaterialComponentPreview)
 		: _PreviewSize(FVector2D(48.f, 48.f))
 		{}
-		SLATE_ATTRIBUTE(TOptional<FVector2D>, PreviewSize)
+		SLATE_ARGUMENT(FVector2D, PreviewSize)
 	SLATE_END_ARGS()
 
 	SDMMaterialComponentPreview();
@@ -33,7 +34,11 @@ public:
 
 	void Construct(const FArguments& InArgs, const TSharedRef<SDMMaterialEditor>& InEditorWidget, UDMMaterialComponent* InComponent);
 
-	FSlateMaterialBrush& GetBrush() { return Brush; }
+	FSlateMaterialBrush& GetBrush();
+
+	const FVector2D& GetPreviewSize() const;
+
+	void SetPreviewSize(const FVector2D& InSize);
 
 	//~ Begin SWidget
 	virtual void Tick(const FGeometry& InAllottedGeometry, const double InCurrentTime, const float InDeltaTime) override;
@@ -47,6 +52,8 @@ protected:
 	TWeakObjectPtr<UMaterialInstanceDynamic> PreviewMaterialDynamicWeak;
 	FSlateMaterialBrush Brush;
 	FDelegateHandle EndOfFrameDelegateHandle;
+	TSharedPtr<SImage> PreviewImage;
+	FVector2D PreviewSize;
 
 	void OnComponentUpdated(UDMMaterialComponent* InComponent, UDMMaterialComponent* InSource, EDMUpdateType InUpdateType);
 
