Material Designer
- Fixed localisation conflict.

diff --git a/Engine/Plugins/Experimental/DynamicMaterial/Source/DynamicMaterialEditor/Private/UI/Menus/DMMaterialSlotLayerMenus.cpp b/Engine/Plugins/Experimental/DynamicMaterial/Source/DynamicMaterialEditor/Private/UI/Menus/DMMaterialSlotLayerMenus.cpp
--- a/Engine/Plugins/Experimental/DynamicMaterial/Source/DynamicMaterialEditor/Private/UI/Menus/DMMaterialSlotLayerMenus.cpp
+++ b/Engine/Plugins/Experimental/DynamicMaterial/Source/DynamicMaterialEditor/Private/UI/Menus/DMMaterialSlotLayerMenus.cpp
@@ -439,7 +439,7 @@ void FDMMaterialSlotLayerMenus::AddSlotMenuEntry(const TSharedPtr<SDMMaterialSlo
 		FToolMenuEntry::InitMenuEntry(
 			NAME_None,
 			InName,
-			LOCTEXT("AddValueStageSpecificTooltip", "Add a Material Stage based on this Material Slot."),
+			LOCTEXT("AddSlotStageSpecificTooltip", "Add a Material Stage based on this Material Slot."),
 			FSlateIcon(),
 			FUIAction(FExecuteAction::CreateWeakLambda(
 				TargetSlot,
