Material Designer
- Reworked the material wizard to also allow creation of dynamic materials from templates.
- Wizard and "Create with Material Designer" buttons no longer create materials before selecting options.
- Material Instance factory no longer automatically opens the asset.

diff --git a/Engine/Plugins/Experimental/Avalanche/Source/AvalancheEditor/Private/MaterialDesigner/AvaMaterialDesignerTextureAssetFactory.cpp b/Engine/Plugins/Experimental/Avalanche/Source/AvalancheEditor/Private/MaterialDesigner/AvaMaterialDesignerTextureAssetFactory.cpp
--- a/Engine/Plugins/Experimental/Avalanche/Source/AvalancheEditor/Private/MaterialDesigner/AvaMaterialDesignerTextureAssetFactory.cpp
+++ b/Engine/Plugins/Experimental/Avalanche/Source/AvalancheEditor/Private/MaterialDesigner/AvaMaterialDesignerTextureAssetFactory.cpp
@@ -14,8 +14,10 @@
 #include "Components/MaterialValues/DMMaterialValueTexture.h"
 #include "DynamicMeshes/AvaShapeDynMeshBase.h"
 #include "DynamicMeshes/AvaShapeRectangleDynMesh.h"
+#include "Engine/Level.h"
 #include "Engine/Texture.h"
 #include "Framework/Application/SlateApplication.h"
+#include "IDynamicMaterialEditorModule.h"
 #include "Material/DynamicMaterialInstance.h"
 #include "Material/DynamicMaterialInstanceFactory.h"
 #include "Model/DynamicMaterialModel.h"
@@ -125,6 +127,12 @@ AActor* UAvaMaterialDesignerTextureAssetFactory::SpawnActor(UObject* InAsset, UL
 	EditorOnlyData->SetChannelListPreset(TEXT("Emissive"));
 	EditorOnlyData->OnWizardComplete();
 
+	ON_SCOPE_EXIT
+	{
+		const IDynamicMaterialEditorModule& MaterialDesignerModule = IDynamicMaterialEditorModule::Get();
+		MaterialDesignerModule.OpenMaterialInstance(NewInstance, InLevel->GetWorld(), /* Invoke Tab */ true);
+	};
+
 	const UDMMaterialSlot* Slot = EditorOnlyData->GetSlotForMaterialProperty(EDMMaterialPropertyType::EmissiveColor);
 
 	if (!Slot)
diff --git a/Engine/Plugins/Experimental/Avalanche/Source/AvalancheShapes/Private/DynamicMeshes/AvaShapeDynMeshBase.cpp b/Engine/Plugins/Experimental/Avalanche/Source/AvalancheShapes/Private/DynamicMeshes/AvaShapeDynMeshBase.cpp
--- a/Engine/Plugins/Experimental/Avalanche/Source/AvalancheShapes/Private/DynamicMeshes/AvaShapeDynMeshBase.cpp
+++ b/Engine/Plugins/Experimental/Avalanche/Source/AvalancheShapes/Private/DynamicMeshes/AvaShapeDynMeshBase.cpp
@@ -1,6 +1,7 @@
 // Copyright Epic Games, Inc. All Rights Reserved.
 
 #include "DynamicMeshes/AvaShapeDynMeshBase.h"
+
 #include "Async/Async.h"
 #include "Async/ParallelFor.h"
 #include "AvaShapeActor.h"
@@ -28,6 +29,7 @@
 #if WITH_EDITOR
 #include "Editor.h"
 #include "Elements/Actor/ActorElementEditorViewportInteractionCustomization.h"
+#include "IDynamicMaterialEditorModule.h"
 #include "LevelEditor/AvaLevelEditorUtils.h"
 #include "Material/DynamicMaterialInstanceFactory.h"
 #include "PropertyEditorModule.h"
@@ -1190,6 +1192,11 @@ void UAvaShapeDynamicMeshBase::OnMaterialTypeChanged(int32 MaterialIndex)
 					UDynamicMaterialInstance* NewInstance = NewObject<UDynamicMaterialInstance>(this);
 #endif
 					SetMaterial(MaterialIndex, NewInstance);
+
+#if WITH_EDITOR
+					const IDynamicMaterialEditorModule& MaterialDesignerModule = IDynamicMaterialEditorModule::Get();
+					MaterialDesignerModule.OpenMaterialInstance(NewInstance, GetWorld(), /* Invoke Tab */ true);
+#endif
 				}
 				break;
 			}
diff --git a/Engine/Plugins/Experimental/DynamicMaterial/Source/DynamicMaterialEditor/Private/DMContentBrowserIntegration.cpp b/Engine/Plugins/Experimental/DynamicMaterial/Source/DynamicMaterialEditor/Private/DMContentBrowserIntegration.cpp
--- a/Engine/Plugins/Experimental/DynamicMaterial/Source/DynamicMaterialEditor/Private/DMContentBrowserIntegration.cpp
+++ b/Engine/Plugins/Experimental/DynamicMaterial/Source/DynamicMaterialEditor/Private/DMContentBrowserIntegration.cpp
@@ -114,7 +114,7 @@ void FDMContentBrowserIntegration::OnCreateMaterialDesignerInstanceFromTextureSe
 		return;
 	}
 
-	UDynamicMaterialInstance* Instance = Cast<UDynamicMaterialInstance>(GetMutableDefault<UDynamicMaterialInstanceFactory>()->FactoryCreateNew(
+	UDynamicMaterialInstance* NewInstance = Cast<UDynamicMaterialInstance>(GetMutableDefault<UDynamicMaterialInstanceFactory>()->FactoryCreateNew(
 		UDynamicMaterialInstance::StaticClass(),
 		GetTransientPackage(),
 		NAME_None,
@@ -123,12 +123,18 @@ void FDMContentBrowserIntegration::OnCreateMaterialDesignerInstanceFromTextureSe
 		GWarn
 	));
 
-	if (!Instance)
+	if (!NewInstance)
 	{
 		return;
 	}
 
-	UDynamicMaterialModelEditorOnlyData* EditorOnlyData = UDynamicMaterialModelEditorOnlyData::Get(Instance);
+	ON_SCOPE_EXIT
+	{
+		const IDynamicMaterialEditorModule& MaterialDesignerModule = IDynamicMaterialEditorModule::Get();
+		MaterialDesignerModule.OpenMaterialInstance(NewInstance, nullptr, /* Invoke Tab */ true);
+	};
+
+	UDynamicMaterialModelEditorOnlyData* EditorOnlyData = UDynamicMaterialModelEditorOnlyData::Get(NewInstance);
 
 	if (!EditorOnlyData)
 	{
@@ -157,10 +163,10 @@ void FDMContentBrowserIntegration::OnCreateMaterialDesignerInstanceFromTextureSe
 		return;
 	}
 
-	Instance->SetFlags(RF_Standalone | RF_Public);
-	Instance->Rename(*UniqueAssetName, Package, REN_DontCreateRedirectors);
+	NewInstance->SetFlags(RF_Standalone | RF_Public);
+	NewInstance->Rename(*UniqueAssetName, Package, REN_DontCreateRedirectors);
 
-	FAssetRegistryModule::AssetCreated(Instance);
+	FAssetRegistryModule::AssetCreated(NewInstance);
 }
 
 void FDMContentBrowserIntegration::UpdateMaterialDesignerInstanceFromTextureSet(TArray<FAssetData> InSelectedAssets, bool bInReplace)
diff --git a/Engine/Plugins/Experimental/DynamicMaterial/Source/DynamicMaterialEditor/Private/DetailsPanel/Widgets/SDMDetailsPanelMaterialInterfaceWidget.cpp b/Engine/Plugins/Experimental/DynamicMaterial/Source/DynamicMaterialEditor/Private/DetailsPanel/Widgets/SDMDetailsPanelMaterialInterfaceWidget.cpp
--- a/Engine/Plugins/Experimental/DynamicMaterial/Source/DynamicMaterialEditor/Private/DetailsPanel/Widgets/SDMDetailsPanelMaterialInterfaceWidget.cpp
+++ b/Engine/Plugins/Experimental/DynamicMaterial/Source/DynamicMaterialEditor/Private/DetailsPanel/Widgets/SDMDetailsPanelMaterialInterfaceWidget.cpp
@@ -7,7 +7,6 @@
 #include "DynamicMaterialEditorStyle.h"
 #include "IDynamicMaterialEditorModule.h"
 #include "Material/DynamicMaterialInstance.h"
-#include "Material/DynamicMaterialInstanceFactory.h"
 #include "Model/DynamicMaterialModel.h"
 #include "PropertyCustomizationHelpers.h"
 #include "PropertyHandle.h"
@@ -146,32 +145,29 @@ FReply SDMDetailsPanelMaterialInterfaceWidget::CreateDynamicMaterialInstance()
 	// We already have an instance, so we don't need to create one
 	if (Instance)
 	{
-		return FReply::Unhandled();
+		return FReply::Handled();
 	}
 
 	TArray<UObject*> OuterObjects;
 	PropertyHandle->GetOuterObjects(OuterObjects);
 
-	if (OuterObjects.IsEmpty())
+	if (OuterObjects.IsEmpty() || !IsValid(OuterObjects[0]))
 	{
-		return FReply::Unhandled();
+		return FReply::Handled();
 	}
 
-	UDynamicMaterialInstanceFactory* DynamicMaterialInstanceFactory = NewObject<UDynamicMaterialInstanceFactory>();
-	check(DynamicMaterialInstanceFactory);
-
-	UDynamicMaterialInstance* NewInstance = Cast<UDynamicMaterialInstance>(DynamicMaterialInstanceFactory->FactoryCreateNew(
-		UDynamicMaterialInstance::StaticClass(),
-		OuterObjects[0],
-		"DynamicMaterialInstance",
-		RF_NoFlags,
-		nullptr,
-		GWarn
-	));
+	const IDynamicMaterialEditorModule& MaterialDesignerModule = IDynamicMaterialEditorModule::Get();
+	constexpr bool bInvokeTab = true;
 
-	PropertyHandle->SetValueFromFormattedString(NewInstance->GetPathName());
+	if (FProperty* Property = PropertyHandle->GetProperty())
+	{
+		if (Property->IsA<FObjectPropertyBase>())
+		{
+			MaterialDesignerModule.OpenMaterialObjectProperty({OuterObjects[0], Property}, OuterObjects[0]->GetWorld(), bInvokeTab);
+		}
+	}
 
-	return OpenDynamicMaterialInstanceTab();
+	return FReply::Handled();
 }
 
 FReply SDMDetailsPanelMaterialInterfaceWidget::ClearDynamicMaterialInstance()
@@ -181,7 +177,7 @@ FReply SDMDetailsPanelMaterialInterfaceWidget::ClearDynamicMaterialInstance()
 	// We don't have an instance, so we don't need to clear it (and don't clear non-MDIs)
 	if (!Instance)
 	{
-		return FReply::Unhandled();
+		return FReply::Handled();
 	}
 
 	SetDynamicMaterialInstance(nullptr);
@@ -196,7 +192,7 @@ FReply SDMDetailsPanelMaterialInterfaceWidget::OpenDynamicMaterialInstanceTab()
 	// We don't have a MDI, so don't try to open it.
 	if (!Instance)
 	{
-		return FReply::Unhandled();
+		return FReply::Handled();
 	}
 
 	const IDynamicMaterialEditorModule& MaterialDesignerModule = IDynamicMaterialEditorModule::Get();
diff --git a/Engine/Plugins/Experimental/DynamicMaterial/Source/DynamicMaterialEditor/Private/DetailsPanel/Widgets/SDMDetailsPanelTabSpawner.cpp b/Engine/Plugins/Experimental/DynamicMaterial/Source/DynamicMaterialEditor/Private/DetailsPanel/Widgets/SDMDetailsPanelTabSpawner.cpp
--- a/Engine/Plugins/Experimental/DynamicMaterial/Source/DynamicMaterialEditor/Private/DetailsPanel/Widgets/SDMDetailsPanelTabSpawner.cpp
+++ b/Engine/Plugins/Experimental/DynamicMaterial/Source/DynamicMaterialEditor/Private/DetailsPanel/Widgets/SDMDetailsPanelTabSpawner.cpp
@@ -2,13 +2,12 @@
 
 #include "DetailsPanel/Widgets/SDMDetailsPanelTabSpawner.h"
 
-#include "AssetRegistry/AssetRegistryModule.h"
 #include "AssetToolsModule.h"
+#include "DMObjectMaterialProperty.h"
 #include "DynamicMaterialEditorStyle.h"
 #include "IAssetTools.h"
-#include "Model/DynamicMaterialModel.h"
+#include "IDynamicMaterialEditorModule.h"
 #include "Model/DynamicMaterialModelBase.h"
-#include "Model/DynamicMaterialModelFactory.h"
 #include "PropertyCustomizationHelpers.h"
 #include "PropertyHandle.h"
 #include "Widgets/Input/SButton.h"
@@ -127,46 +126,29 @@ FReply SDMDetailsPanelTabSpawner::CreateDynamicMaterialModel()
 	// We already have a builder, so we don't need to create one
 	if (MaterialModelBase)
 	{
-		return FReply::Unhandled();
+		return FReply::Handled();
 	}
 
 	TArray<UObject*> OuterObjects;
 	PropertyHandle->GetOuterObjects(OuterObjects);
 
-	if (OuterObjects.IsEmpty())
+	if (OuterObjects.IsEmpty() || !IsValid(OuterObjects[0]))
 	{
-		return FReply::Unhandled();
+		return FReply::Handled();
 	}
 
-	FString PackageName, AssetName;
-	IAssetTools& AssetTools = FModuleManager::LoadModuleChecked<FAssetToolsModule>("AssetTools").Get();
-	AssetTools.CreateUniqueAssetName(
-		UDynamicMaterialModelFactory::BaseDirectory / UDynamicMaterialModelFactory::BaseName + FGuid::NewGuid().ToString(),
-		"",
-		PackageName,
-		AssetName
-	);
-
-	UPackage* Package = CreatePackage(*PackageName);
-	check(Package);
-
-	UDynamicMaterialModelFactory* DynamicMaterialModelFactory = NewObject<UDynamicMaterialModelFactory>();
-	check(DynamicMaterialModelFactory);
-
-	UDynamicMaterialModel* NewModel = Cast<UDynamicMaterialModel>(DynamicMaterialModelFactory->FactoryCreateNew(
-		UDynamicMaterialModelBase::StaticClass(),
-		Package,
-		*AssetName,
-		RF_Standalone | RF_Public,
-		nullptr,
-		GWarn
-	));
+	const IDynamicMaterialEditorModule& MaterialDesignerModule = IDynamicMaterialEditorModule::Get();
+	constexpr bool bInvokeTab = true;
 
-	FAssetRegistryModule::AssetCreated(NewModel);
-
-	PropertyHandle->SetValueFromFormattedString(NewModel->GetPathName());
+	if (FProperty* Property = PropertyHandle->GetProperty())
+	{
+		if (Property->IsA<FObjectPropertyBase>())
+		{
+			MaterialDesignerModule.OpenMaterialObjectProperty({OuterObjects[0], Property}, OuterObjects[0]->GetWorld(), bInvokeTab);
+		}
+	}
 
-	return OpenDynamicMaterialModelTab();
+	return FReply::Handled();
 }
 
 FReply SDMDetailsPanelTabSpawner::ClearDynamicMaterialModel()
@@ -176,7 +158,7 @@ FReply SDMDetailsPanelTabSpawner::ClearDynamicMaterialModel()
 	// We don't have a builder, so we don't need to clear it
 	if (!MaterialModelBase)
 	{
-		return FReply::Unhandled();
+		return FReply::Handled();
 	}
 
 	SetMaterialModelBase(nullptr);
@@ -191,7 +173,7 @@ FReply SDMDetailsPanelTabSpawner::OpenDynamicMaterialModelTab()
 	// We don't have a builder, so we can't open it
 	if (!MaterialModelBase)
 	{
-		return FReply::Unhandled();
+		return FReply::Handled();
 	}
 
 	IAssetTools& AssetTools = FModuleManager::LoadModuleChecked<FAssetToolsModule>("AssetTools").Get();
diff --git a/Engine/Plugins/Experimental/DynamicMaterial/Source/DynamicMaterialEditor/Private/DetailsPanel/Widgets/SDMMaterialListExtensionWidget.cpp b/Engine/Plugins/Experimental/DynamicMaterial/Source/DynamicMaterialEditor/Private/DetailsPanel/Widgets/SDMMaterialListExtensionWidget.cpp
--- a/Engine/Plugins/Experimental/DynamicMaterial/Source/DynamicMaterialEditor/Private/DetailsPanel/Widgets/SDMMaterialListExtensionWidget.cpp
+++ b/Engine/Plugins/Experimental/DynamicMaterial/Source/DynamicMaterialEditor/Private/DetailsPanel/Widgets/SDMMaterialListExtensionWidget.cpp
@@ -157,26 +157,31 @@ FReply SDMMaterialListExtensionWidget::CreateDynamicMaterialInstance()
 	// We already have an instance, so we don't need to create one
 	if (Instance)
 	{
-		return FReply::Unhandled();
+		return FReply::Handled();
+	}
+
+	TSharedPtr<FMaterialItemView> ListItem = MaterialItemViewWeak.Pin();
+
+	if (!ListItem.IsValid())
+	{
+		return FReply::Handled();
 	}
 
-	FName InstanceName = MakeUniqueObjectName(CurrentComponentWeak.Get(), UDynamicMaterialInstance::StaticClass(), "DynamicMaterialInstance");
+	UPrimitiveComponent* Component = CurrentComponentWeak.Get();
 
-	UDynamicMaterialInstanceFactory* DynamicMaterialInstanceFactory = NewObject<UDynamicMaterialInstanceFactory>();
-	check(DynamicMaterialInstanceFactory);
+	if (!Component)
+	{
+		return FReply::Handled();
+	}
 
-	UDynamicMaterialInstance* NewInstance = Cast<UDynamicMaterialInstance>(DynamicMaterialInstanceFactory->FactoryCreateNew(
-		UDynamicMaterialInstance::StaticClass(),
-		CurrentComponentWeak.Get(),
-		InstanceName,
-		RF_NoFlags,
-		nullptr,
-		GWarn
-	));
+	const FDMObjectMaterialProperty MaterialProperty(Component, ListItem->GetMaterialListItem().SlotIndex);
 
-	SetDynamicMaterialInstance(NewInstance);
+	constexpr bool bInvokeTab = true;
 
-	return OpenDynamicMaterialInstanceTab();
+	const IDynamicMaterialEditorModule& MaterialDesignerModule = IDynamicMaterialEditorModule::Get();
+	MaterialDesignerModule.OpenMaterialObjectProperty(MaterialProperty, Component->GetWorld(), bInvokeTab);
+
+	return FReply::Handled();
 }
 
 FReply SDMMaterialListExtensionWidget::ClearDynamicMaterialInstance()
@@ -186,7 +191,7 @@ FReply SDMMaterialListExtensionWidget::ClearDynamicMaterialInstance()
 	// We don't have an instance, so we don't need to clear it (or any other asset in its place)
 	if (!Instance)
 	{
-		return FReply::Unhandled();
+		return FReply::Handled();
 	}
 
 	SetDynamicMaterialInstance(nullptr);
diff --git a/Engine/Plugins/Experimental/DynamicMaterial/Source/DynamicMaterialEditor/Private/DynamicMaterialEditorSettings.cpp b/Engine/Plugins/Experimental/DynamicMaterial/Source/DynamicMaterialEditor/Private/DynamicMaterialEditorSettings.cpp
--- a/Engine/Plugins/Experimental/DynamicMaterial/Source/DynamicMaterialEditor/Private/DynamicMaterialEditorSettings.cpp
+++ b/Engine/Plugins/Experimental/DynamicMaterial/Source/DynamicMaterialEditor/Private/DynamicMaterialEditorSettings.cpp
@@ -10,8 +10,11 @@
 #include "GenericPlatform/GenericApplication.h"
 #include "Interfaces/IPluginManager.h"
 #include "ISettingsModule.h"
+#include "Material/DynamicMaterialInstance.h"
 #include "Materials/MaterialFunctionInterface.h"
 #include "Misc/Paths.h"
+#include "Model/DynamicMaterialModel.h"
+#include "Model/DynamicMaterialModelBase.h"
 #include "Modules/ModuleManager.h"
 #include "UObject/UObjectIterator.h"
 
@@ -131,6 +134,7 @@ UDynamicMaterialEditorSettings::UDynamicMaterialEditorSettings()
 
 	bFollowSelection = true;
 	bUseLinearColorForVectors = true;
+	CustomTemplateFolders = {};
 
 	ResetAllLayoutSettings();
 
@@ -316,7 +320,7 @@ void UDynamicMaterialEditorSettings::ResetAllLayoutSettings()
 
 TArray<FDMMaterialEffectList> UDynamicMaterialEditorSettings::GetEffectList() const
 {
-	TArray<FDMMaterialEffectList> Effects = {};
+	TArray<FDMMaterialEffectList> Effects;
 
 	IAssetRegistry* AssetRegistry = IAssetRegistry::Get();
 
@@ -424,6 +428,59 @@ TArray<FDMMaterialEffectList> UDynamicMaterialEditorSettings::GetEffectList() co
 	return Effects;
 }
 
+TArray<FAssetData> UDynamicMaterialEditorSettings::GetTemplateList() const
+{
+	TArray<FAssetData> Templates;
+
+	IAssetRegistry* AssetRegistry = IAssetRegistry::Get();
+
+	if (!AssetRegistry)
+	{
+		return Templates;
+	}
+
+	TSharedPtr<IPlugin> Plugin = IPluginManager::Get().FindPlugin(UE_PLUGIN_NAME);
+
+	if (!Plugin.IsValid())
+	{
+		return Templates;
+	}
+
+	const FString PluginEffectPath = Plugin->GetMountedAssetPath() / "Templates";
+
+	TArray<FName> AssetPaths = {*PluginEffectPath};
+	AssetPaths.Append(UDynamicMaterialEditorSettings::Get()->CustomTemplateFolders);
+
+	TArray<FAssetData> Assets;
+	AssetRegistry->GetAssetsByPaths(AssetPaths, Assets, /* bRecursive */ true, /* Only Assets on Disk */ true);
+
+	for (const FAssetData& Asset : Assets)
+	{
+		if (Asset.GetClass(EResolveClass::Yes) != UDynamicMaterialInstance::StaticClass())
+		{
+			continue;
+		}
+
+		UDynamicMaterialInstance* Instance = Cast<UDynamicMaterialInstance>(Asset.GetAsset());
+
+		if (!Instance)
+		{
+			continue;
+		}
+
+		UDynamicMaterialModel* MaterialModel = Cast<UDynamicMaterialModel>(Instance->GetMaterialModelBase());
+
+		if (!MaterialModel)
+		{
+			continue;
+		}
+
+		Templates.Add(Asset);
+	}
+
+	return Templates;
+}
+
 const FDMDefaultMaterialPropertySlotValue& UDynamicMaterialEditorSettings::GetDefaultSlotValue(EDMMaterialPropertyType InProperty) const
 {
 	if (const FDMDefaultMaterialPropertySlotValue* OverridePtr = DefaultSlotValueOverrides.Find(InProperty))
diff --git a/Engine/Plugins/Experimental/DynamicMaterial/Source/DynamicMaterialEditor/Private/DynamicMaterialEditorSettings.h b/Engine/Plugins/Experimental/DynamicMaterial/Source/DynamicMaterialEditor/Private/DynamicMaterialEditorSettings.h
--- a/Engine/Plugins/Experimental/DynamicMaterial/Source/DynamicMaterialEditor/Private/DynamicMaterialEditorSettings.h
+++ b/Engine/Plugins/Experimental/DynamicMaterial/Source/DynamicMaterialEditor/Private/DynamicMaterialEditorSettings.h
@@ -11,6 +11,7 @@
 #include "DynamicMaterialEditorSettings.generated.h"
 
 class FModifierKeysState;
+class UDynamicMaterialModel;
 class UMaterialFunctionInterface;
 enum EOrientation : int;
 struct FDMMaterialChannelListPreset;
@@ -243,6 +244,16 @@ public:
 	UPROPERTY(Config)
 	bool bValidatedPresets = false;
 
+	/*
+	 * Add paths to search for custom material templates.
+	 *
+	 * Format examples:
+	 * - /Game/Some/Path
+	 * - /Plugin/Some/Path
+	 */
+	UPROPERTY(Config, EditAnywhere, BlueprintReadWrite, Category = "Templates")
+	TArray<FName> CustomTemplateFolders;
+
 	/** This variable is accessed in multiple places, so this is a quick accessor. */
 	static bool IsUseLinearColorForVectorsEnabled();
 
@@ -252,6 +263,8 @@ public:
 
 	TArray<FDMMaterialEffectList> GetEffectList() const;
 
+	TArray<FAssetData> GetTemplateList() const;
+
 	const FDMDefaultMaterialPropertySlotValue& GetDefaultSlotValue(EDMMaterialPropertyType InProperty) const;
 
 	const FDMMaterialChannelListPreset* GetPresetByName(FName InName) const;
diff --git a/Engine/Plugins/Experimental/DynamicMaterial/Source/DynamicMaterialEditor/Private/Material/DynamicMaterialInstanceFactory.cpp b/Engine/Plugins/Experimental/DynamicMaterial/Source/DynamicMaterialEditor/Private/Material/DynamicMaterialInstanceFactory.cpp
--- a/Engine/Plugins/Experimental/DynamicMaterial/Source/DynamicMaterialEditor/Private/Material/DynamicMaterialInstanceFactory.cpp
+++ b/Engine/Plugins/Experimental/DynamicMaterial/Source/DynamicMaterialEditor/Private/Material/DynamicMaterialInstanceFactory.cpp
@@ -2,11 +2,9 @@
 
 #include "Material/DynamicMaterialInstanceFactory.h"
 
-#include "AssetToolsModule.h"
 #include "DMDefs.h"
 #include "EngineAnalytics.h"
 #include "GameFramework/Actor.h"
-#include "IAssetTools.h"
 #include "Material/DynamicMaterialInstance.h"
 #include "Model/DynamicMaterialModel.h"
 #include "Model/DynamicMaterialModelFactory.h"
@@ -73,9 +71,6 @@ UObject* UDynamicMaterialInstanceFactory::FactoryCreateNew(UClass* Class, UObjec
 		}
 	}
 
-	IAssetTools& AssetTools = FModuleManager::LoadModuleChecked<FAssetToolsModule>("AssetTools").Get();
-	AssetTools.OpenEditorForAssets({NewInstance});
-
 	if (FEngineAnalytics::IsAvailable())
 	{
 		static const FString AssetType = TEXT("Asset");
diff --git a/Engine/Plugins/Experimental/DynamicMaterial/Source/DynamicMaterialEditor/Private/Model/DynamicMaterialModelEditorOnlyData.cpp b/Engine/Plugins/Experimental/DynamicMaterial/Source/DynamicMaterialEditor/Private/Model/DynamicMaterialModelEditorOnlyData.cpp
--- a/Engine/Plugins/Experimental/DynamicMaterial/Source/DynamicMaterialEditor/Private/Model/DynamicMaterialModelEditorOnlyData.cpp
+++ b/Engine/Plugins/Experimental/DynamicMaterial/Source/DynamicMaterialEditor/Private/Model/DynamicMaterialModelEditorOnlyData.cpp
@@ -182,7 +182,6 @@ UDynamicMaterialModelEditorOnlyData::UDynamicMaterialModelEditorOnlyData()
 	, bOutputTranslucentVelocityEnabled(true)
 	, bNaniteTessellationEnabled(true)
 	, bResponsiveAAEnabled(true)
-	, ChannelListPreset(NAME_None)
 {
 	BaseColor =           CreateDefaultSubobject<UDMMaterialPropertyBaseColor>(          "MaterialProperty_BaseColor");
 	EmissiveColor =       CreateDefaultSubobject<UDMMaterialPropertyEmissiveColor>(      "MaterialProperty_EmissiveColor");
@@ -241,82 +240,12 @@ void UDynamicMaterialModelEditorOnlyData::AssignPropertyAlphaValues()
 	SurfaceThickness   ->AddComponent(AlphaValueName, MaterialModel->GetGlobalParameterValue(UDynamicMaterialModel::GlobalSurfaceThicknessValueName));
 }
 
-TArray<FName> UDynamicMaterialModelEditorOnlyData::GetPresetOptions() const
-{
-	UDynamicMaterialEditorSettings* Settings = UDynamicMaterialEditorSettings::Get();
-
-	if (!Settings)
-	{
-		return {ChannelListPreset};
-	}
-
-	TArray<FName> PresetNames;
-	PresetNames.Reserve(Settings->MaterialChannelPresets.Num());
-
-	for (const FDMMaterialChannelListPreset& Preset : Settings->MaterialChannelPresets)
-	{
-		PresetNames.Add(Preset.Name);
-	}
-
-	return PresetNames;
-}
-
-void UDynamicMaterialModelEditorOnlyData::OnChannelListPresetChanged()
-{
-	const FDMMaterialChannelListPreset* Preset = GetDefault<UDynamicMaterialEditorSettings>()->GetPresetByName(ChannelListPreset);
-
-	if (!Preset)
-	{
-		return;
-	}
-
-	EnsurePresetSlots();
-
-	SetBlendMode(Preset->DefaultBlendMode);
-	SetShadingModel(Preset->DefaultShadingModel);
-	SetHasPixelAnimation(Preset->bDefaultAnimated);
-	SetIsTwoSided(Preset->bDefaultTwoSided);
-}
-
-void UDynamicMaterialModelEditorOnlyData::EnsurePresetSlots()
-{
-	const FDMMaterialChannelListPreset* Preset = GetDefault<UDynamicMaterialEditorSettings>()->GetPresetByName(ChannelListPreset);
-
-	if (!Preset)
-	{
-		return;
-	}
-
-	UE::DynamicMaterial::ForEachMaterialPropertyType(
-		[this, Preset](EDMMaterialPropertyType InProperty)
-		{
-			if (InProperty == EDMMaterialPropertyType::OpacityMask)
-			{
-				return EDMIterationResult::Continue;
-			}
-
-			if (Preset->IsPropertyEnabled(InProperty))
-			{
-				AddSlotForMaterialProperty(InProperty);
-			}
-			else
-			{
-				RemoveSlotForMaterialProperty(InProperty);
-			}
-
-			return EDMIterationResult::Continue;
-		});
-}
-
 void UDynamicMaterialModelEditorOnlyData::OnDomainChanged()
 {
 	if (Domain == EMaterialDomain::MD_PostProcess)
 	{
 		const FDMUpdateGuard Guard;
 
-		// Make sure we have the correct slots, then adjust for domains.
-		EnsurePresetSlots();
-
 		// Post process only supports emissive.
 		UDMMaterialSlot* BaseColorSlot = GetSlotForMaterialProperty(EDMMaterialPropertyType::BaseColor);
 		UDMMaterialSlot* EmissiveSlot = GetSlotForMaterialProperty(EDMMaterialPropertyType::EmissiveColor);
@@ -347,8 +276,6 @@ void UDynamicMaterialModelEditorOnlyData::OnDomainChanged()
 
 void UDynamicMaterialModelEditorOnlyData::OnBlendModeChanged()
 {
-	EnsurePresetSlots();
-
 	switch (BlendMode)
 	{
 		case EBlendMode::BLEND_Opaque:
@@ -390,9 +317,6 @@ void UDynamicMaterialModelEditorOnlyData::Initialize()
 	}
 
 	AssignPropertyAlphaValues();
-
-	// Will choose appropriate type between BaseColor and EmissiveColor
-	AddSlotForMaterialProperty(EDMMaterialPropertyType::BaseColor);
 }
 
 UMaterial* UDynamicMaterialModelEditorOnlyData::GetGeneratedMaterial() const
@@ -845,7 +769,7 @@ bool UDynamicMaterialModelEditorOnlyData::AddTextureSet(UDMTextureSet* InTexture
 
 bool UDynamicMaterialModelEditorOnlyData::NeedsWizard() const
 {
-	return ChannelListPreset.IsNone();
+	return PropertySlotMap.IsEmpty();
 }
 
 void UDynamicMaterialModelEditorOnlyData::OnWizardComplete()
@@ -858,9 +782,37 @@ void UDynamicMaterialModelEditorOnlyData::OnWizardComplete()
 
 void UDynamicMaterialModelEditorOnlyData::SetChannelListPreset(FName InPresetName)
 {
-	ChannelListPreset = InPresetName;
+	const FDMMaterialChannelListPreset* Preset = GetDefault<UDynamicMaterialEditorSettings>()->GetPresetByName(InPresetName);
+
+	if (!Preset)
+	{
+		return;
+	}
+
+	UE::DynamicMaterial::ForEachMaterialPropertyType(
+		[this, Preset](EDMMaterialPropertyType InProperty)
+		{
+			if (InProperty == EDMMaterialPropertyType::OpacityMask)
+			{
+				return EDMIterationResult::Continue;
+			}
 
-	OnChannelListPresetChanged();
+			if (Preset->IsPropertyEnabled(InProperty))
+			{
+				AddSlotForMaterialProperty(InProperty);
+			}
+			else
+			{
+				RemoveSlotForMaterialProperty(InProperty);
+			}
+
+			return EDMIterationResult::Continue;
+		});
+
+	SetBlendMode(Preset->DefaultBlendMode);
+	SetShadingModel(Preset->DefaultShadingModel);
+	SetHasPixelAnimation(Preset->bDefaultAnimated);
+	SetIsTwoSided(Preset->bDefaultTwoSided);
 }
 
 const FMaterialStatistics& UDynamicMaterialModelEditorOnlyData::GetMaterialStats() const
@@ -1123,11 +1075,6 @@ void UDynamicMaterialModelEditorOnlyData::SetResponsiveAAEnabled(bool bInEnabled
 	OnMaterialFlagChanged();
 }
 
-FName UDynamicMaterialModelEditorOnlyData::GetChannelListPreset() const
-{
-	return ChannelListPreset;
-}
-
 void UDynamicMaterialModelEditorOnlyData::OpenMaterialEditor() const
 {
 	if (!IsValid(MaterialModel) || !IsValid(MaterialModel->DynamicMaterial))
@@ -1534,11 +1481,7 @@ void UDynamicMaterialModelEditorOnlyData::PostEditChangeChainProperty(FPropertyC
 
 	const FName Property = InPropertyChangedEvent.GetMemberPropertyName();
 
-	if (Property == GET_MEMBER_NAME_CHECKED(ThisClass, ChannelListPreset))
-	{
-		OnChannelListPresetChanged();
-	}
-	else if (Property == GET_MEMBER_NAME_CHECKED(ThisClass, Domain))
+	 if (Property == GET_MEMBER_NAME_CHECKED(ThisClass, Domain))
 	{
 		OnDomainChanged();
 	}
@@ -1560,50 +1503,6 @@ void UDynamicMaterialModelEditorOnlyData::PostEditChangeChainProperty(FPropertyC
 	}
 }
 
-void UDynamicMaterialModelEditorOnlyData::Serialize(FArchive& InAr)
-{
-	InAr.UsingCustomVersion(FDynamicMaterialModelEditorOnlyDataVersion::GUID);
-
-	Super::Serialize(InAr);
-
-	const int32 Version = InAr.CustomVer(FDynamicMaterialModelEditorOnlyDataVersion::GUID);
-
-	// The default blend mode was changed from translucent and to opaque. If we have a
-	// transparent channel, let's set it back to translucent.
-	if (Version < FDynamicMaterialModelEditorOnlyDataVersion::GlobalValueRename)
-	{
-		if (ChannelListPreset.IsNone())
-		{
-			// Try to guess from the available slots.
-			if (GetSlotForMaterialProperty(EDMMaterialPropertyType::Opacity))
-			{
-				ChannelListPreset = "Translucent";
-				BlendMode = BLEND_Translucent;
-				ShadingModel = EDMMaterialShadingModel::Unlit;
-			}
-			else if (GetSlotForMaterialProperty(EDMMaterialPropertyType::OpacityMask))
-			{
-				ChannelListPreset = "Translucent";
-				BlendMode = BLEND_Masked;
-				ShadingModel = EDMMaterialShadingModel::Unlit;
-			}
-			if (GetSlotForMaterialProperty(EDMMaterialPropertyType::EmissiveColor))
-			{
-				ChannelListPreset = "Translucent";
-				BlendMode = BLEND_Translucent;
-				ShadingModel = EDMMaterialShadingModel::Unlit;
-			}
-			// Else let's stay on opaque
-			else
-			{
-				BlendMode = BLEND_Opaque;
-				ShadingModel = EDMMaterialShadingModel::DefaultLit;
-				ChannelListPreset = "Opaque";
-			}
-		}
-	}
-}
-
 void UDynamicMaterialModelEditorOnlyData::OnValueUpdated(UDMMaterialValue* InValue, EDMUpdateType InUpdateType)
 {
 	check(InValue);
diff --git a/Engine/Plugins/Experimental/DynamicMaterial/Source/DynamicMaterialEditor/Private/UI/PropertyGenerators/DMMaterialModelPropertyRowGenerator.cpp b/Engine/Plugins/Experimental/DynamicMaterial/Source/DynamicMaterialEditor/Private/UI/PropertyGenerators/DMMaterialModelPropertyRowGenerator.cpp
--- a/Engine/Plugins/Experimental/DynamicMaterial/Source/DynamicMaterialEditor/Private/UI/PropertyGenerators/DMMaterialModelPropertyRowGenerator.cpp
+++ b/Engine/Plugins/Experimental/DynamicMaterial/Source/DynamicMaterialEditor/Private/UI/PropertyGenerators/DMMaterialModelPropertyRowGenerator.cpp
@@ -72,8 +72,7 @@ void FDMMaterialModelPropertyRowGenerator::AddMaterialModelProperties(const TSha
 
 	if (UDynamicMaterialModelEditorOnlyData* EditorOnlyData = UDynamicMaterialModelEditorOnlyData::Get(InMaterialModelBase))
 	{
-		AddVariable(InGlobalSettingEditorWidget, InMaterialModelBase, InOutPropertyRows, EditorOnlyData, 
-			GET_MEMBER_NAME_CHECKED(UDynamicMaterialModelEditorOnlyData, ChannelListPreset));
+		// TODO: Add selector for non-dynamic materials to change their preset. It's no longer on the EOD.
 
 		AddVariable(InGlobalSettingEditorWidget, InMaterialModelBase, InOutPropertyRows, EditorOnlyData,
 			GET_MEMBER_NAME_CHECKED(UDynamicMaterialModelEditorOnlyData, Domain));
diff --git a/Engine/Plugins/Experimental/DynamicMaterial/Source/DynamicMaterialEditor/Private/UI/Widgets/SDMMaterialDesigner.cpp b/Engine/Plugins/Experimental/DynamicMaterial/Source/DynamicMaterialEditor/Private/UI/Widgets/SDMMaterialDesigner.cpp
--- a/Engine/Plugins/Experimental/DynamicMaterial/Source/DynamicMaterialEditor/Private/UI/Widgets/SDMMaterialDesigner.cpp
+++ b/Engine/Plugins/Experimental/DynamicMaterial/Source/DynamicMaterialEditor/Private/UI/Widgets/SDMMaterialDesigner.cpp
@@ -200,18 +200,11 @@ void SDMMaterialDesigner::OpenObjectMaterialProperty_Internal(const FDMObjectMat
 		{
 			SetEditorView(InObjectMaterialProperty);
 		}
-
-		return;
 	}
-
-	if (AActor* MaterialActor = InObjectMaterialProperty.GetTypedOuter<AActor>())
+	else
 	{
-		TArray<FDMObjectMaterialProperty> ActorProperties = UDMMaterialModelFunctionLibrary::GetActorMaterialProperties(MaterialActor);
-		SetMaterialSelectorView(MaterialActor, MoveTemp(ActorProperties));
-		return;
+		SetWizardView(InObjectMaterialProperty);
 	}
-
-	SetSelectPromptView();
 }
 
 void SDMMaterialDesigner::OpenActor_Internal(AActor* InActor)
@@ -235,7 +228,14 @@ void SDMMaterialDesigner::OpenActor_Internal(AActor* InActor)
 		}
 	}
 
-	SetMaterialSelectorView(InActor, MoveTemp(ActorProperties));
+	if (ActorProperties.Num() == 1)
+	{
+		OpenObjectMaterialProperty_Internal(ActorProperties[0]);
+	}
+	else
+	{
+		SetMaterialSelectorView(InActor, MoveTemp(ActorProperties));
+	}
 }
 
 void SDMMaterialDesigner::SetEmptyView()
diff --git a/Engine/Plugins/Experimental/DynamicMaterial/Source/DynamicMaterialEditor/Private/UI/Widgets/SDMMaterialWizard.cpp b/Engine/Plugins/Experimental/DynamicMaterial/Source/DynamicMaterialEditor/Private/UI/Widgets/SDMMaterialWizard.cpp
--- a/Engine/Plugins/Experimental/DynamicMaterial/Source/DynamicMaterialEditor/Private/UI/Widgets/SDMMaterialWizard.cpp
+++ b/Engine/Plugins/Experimental/DynamicMaterial/Source/DynamicMaterialEditor/Private/UI/Widgets/SDMMaterialWizard.cpp
@@ -2,21 +2,28 @@
 
 #include "UI/Widgets/SDMMaterialWizard.h"
 
+#include "AssetRegistry/AssetData.h"
+#include "AssetThumbnail.h"
 #include "DMDefs.h"
 #include "DynamicMaterialEditorSettings.h"
 #include "DynamicMaterialEditorStyle.h"
 #include "DynamicMaterialModule.h"
 #include "Engine/EngineTypes.h"
+#include "Material/DynamicMaterialInstance.h"
+#include "Material/DynamicMaterialInstanceFactory.h"
 #include "Model/DynamicMaterialModel.h"
 #include "Model/DynamicMaterialModelEditorOnlyData.h"
 #include "Styling/SlateTypes.h"
 #include "Styling/StyleColors.h"
+#include "ThumbnailRendering/ThumbnailManager.h"
 #include "UI/Widgets/SDMMaterialDesigner.h"
+#include "Utils/DMMaterialModelFunctionLibrary.h"
 #include "Utils/DMPrivate.h"
 #include "Widgets/Input/SButton.h"
 #include "Widgets/Input/SCheckBox.h"
 #include "Widgets/Layout/SBox.h"
 #include "Widgets/Layout/SExpandableArea.h"
+#include "Widgets/Layout/SWidgetSwitcher.h"
 #include "Widgets/Layout/SWrapBox.h"
 #include "Widgets/SBoxPanel.h"
 #include "Widgets/Text/STextBlock.h"
@@ -43,14 +50,11 @@ SDMMaterialWizard::~SDMMaterialWizard()
 		return;
 	}
 
-	if (TSharedPtr<SDMMaterialDesigner> DesignerWidget = DesignerWidgetWeak.Pin())
+	if (UDynamicMaterialModel* MaterialModel = MaterialModelWeak.Get())
 	{
-		if (UDynamicMaterialModel* MaterialModel = Cast<UDynamicMaterialModel>(DesignerWidget->GetMaterialModelBase()))
+		if (UDynamicMaterialModelEditorOnlyData* EditorOnlyData = UDynamicMaterialModelEditorOnlyData::Get(MaterialModel))
 		{
-			if (UDynamicMaterialModelEditorOnlyData* EditorOnlyData = UDynamicMaterialModelEditorOnlyData::Get(MaterialModel))
-			{
-				EditorOnlyData->GetOnMaterialBuiltDelegate().RemoveAll(this);
-			}
+			EditorOnlyData->GetOnMaterialBuiltDelegate().RemoveAll(this);
 		}
 	}
 }
@@ -108,6 +112,16 @@ UDynamicMaterialModel* SDMMaterialWizard::GetMaterialModel() const
 	return MaterialModelWeak.Get();
 }
 
+TArray<FAssetData> SDMMaterialWizard::GetTemplateMaterials()
+{
+	if (const UDynamicMaterialEditorSettings* Settings = GetDefault<UDynamicMaterialEditorSettings>())
+	{
+		return Settings->GetTemplateList();
+	}
+
+	return {};
+}
+
 TSharedRef<SWidget> SDMMaterialWizard::CreateLayout()
 {
 	using namespace UE::DynamicMaterialDesigner::Private;
@@ -128,54 +142,127 @@ TSharedRef<SWidget> SDMMaterialWizard::CreateLayout()
 
 			+ SVerticalBox::Slot()
 			.AutoHeight()
-			.HAlign(HAlign_Left)
-			.Padding(0.0f, SeparationDistance, 0.0f, TitleContentDistance)
+			.Padding(0.0f, SeparationDistance, 0.0f, 0.f)
 			[
-				SNew(STextBlock)
-				.TextStyle(FDynamicMaterialEditorStyle::Get(), "BoldFont")
-				.Text(LOCTEXT("MaterialType", "Material Type"))
+				CreateModeSelector()
 			]
 
 			+ SVerticalBox::Slot()
 			.AutoHeight()
-			.HAlign(EHorizontalAlignment::HAlign_Fill)
-			.Padding(0.0f, 0.f, 0.0f, TitleContentDistance)
+			.Padding(0.0f, SeparationDistance, 0.0f, 0.f)
 			[
-				CreateChannelPresets()
+				SAssignNew(Switcher, SWidgetSwitcher)
+				+ SWidgetSwitcher::Slot()
+				[
+					CreateTemplateListLayout()
+				]
+				+ SWidgetSwitcher::Slot()
+				[
+					CreateSelectPresetLayout()
+				]
 			]
+		];
+}
 
-			+ SVerticalBox::Slot()
-			.AutoHeight()
-			.HAlign(HAlign_Left)
-			.Padding(0.0f, SeparationDistance, 0.0f, TitleContentDistance)
+TSharedRef<SWidget> SDMMaterialWizard::CreateModeSelector()
+{
+	return SNew(SHorizontalBox)
+
+		+ SHorizontalBox::Slot()
+		.FillContentWidth(1.f)
+		.HAlign(EHorizontalAlignment::HAlign_Right)
+		.Padding(0.f, 0.f, 5.f, 0.f)
+		[
+			SNew(SCheckBox)
+			.Style(FAppStyle::Get(), "DetailsView.SectionButton")
+			.HAlign(EHorizontalAlignment::HAlign_Center)
+			.Padding(FMargin(10.f, 6.f))
+			.IsChecked(this, &SDMMaterialWizard::IsModeSelected, 0)
+			.OnCheckStateChanged(this, &SDMMaterialWizard::SetMode, 0)
+			.ToolTipText(LOCTEXT("TemplateModeToolTip", "Create a new Dynamic Material Instance based on a template."))
+			.Content()
 			[
 				SNew(STextBlock)
 				.TextStyle(FDynamicMaterialEditorStyle::Get(), "BoldFont")
-				.Text(LOCTEXT("AvailableChannels", "Available Channels"))
+				.Text(LOCTEXT("TemplateMode", "New Instance"))
 			]
+		]
 
-			+ SVerticalBox::Slot()
-			.AutoHeight()
-			.HAlign(EHorizontalAlignment::HAlign_Fill)
+		+ SHorizontalBox::Slot()
+		.FillContentWidth(1.f)
+		.HAlign(EHorizontalAlignment::HAlign_Left)
+		.Padding(5.f, 0.f, 0.f, 0.f)
+		[
+			SNew(SCheckBox)
+			.Style(FAppStyle::Get(), "DetailsView.SectionButton")
+			.HAlign(EHorizontalAlignment::HAlign_Center)
+			.Padding(FMargin(10.f, 6.f))
+			.IsChecked(this, &SDMMaterialWizard::IsModeSelected, 1)
+			.OnCheckStateChanged(this, &SDMMaterialWizard::SetMode, 1)
+			.ToolTipText(LOCTEXT("PresetModeToolTip", "Set up a new Material based on simple channel presets."))
+			.Content()
 			[
-				SAssignNew(PresetChannelContainer, SBox)
-				.HAlign(EHorizontalAlignment::HAlign_Fill)
-				[
-					CreateChannelList()
-				]
+				SNew(STextBlock)
+				.TextStyle(FDynamicMaterialEditorStyle::Get(), "BoldFont")
+				.Text(LOCTEXT("PresetMode", "New Template"))
 			]
+		];
+}
 
-			+ SVerticalBox::Slot()
-			.AutoHeight()
+TSharedRef<SWidget> SDMMaterialWizard::CreateSelectPresetLayout()
+{
+	using namespace UE::DynamicMaterialDesigner::Private;
+
+	return SNew(SVerticalBox)
+		+ SVerticalBox::Slot()
+		.AutoHeight()
+		.HAlign(HAlign_Left)
+		.Padding(0.0f, SeparationDistance, 0.0f, TitleContentDistance)
+		[
+			SNew(STextBlock)
+			.TextStyle(FDynamicMaterialEditorStyle::Get(), "BoldFont")
+			.Text(LOCTEXT("MaterialType", "Material Type"))
+		]
+
+		+ SVerticalBox::Slot()
+		.AutoHeight()
+		.HAlign(EHorizontalAlignment::HAlign_Fill)
+		.Padding(0.0f, 0.f, 0.0f, TitleContentDistance)
+		[
+			CreateSelectPreset_ChannelPresets()
+		]
+
+		+ SVerticalBox::Slot()
+		.AutoHeight()
+		.HAlign(HAlign_Left)
+		.Padding(0.0f, SeparationDistance, 0.0f, TitleContentDistance)
+		[
+			SNew(STextBlock)
+			.TextStyle(FDynamicMaterialEditorStyle::Get(), "BoldFont")
+			.Text(LOCTEXT("AvailableChannels", "Available Channels"))
+		]
+
+		+ SVerticalBox::Slot()
+		.AutoHeight()
+		.HAlign(EHorizontalAlignment::HAlign_Fill)
+		[
+			SAssignNew(PresetChannelContainer, SBox)
 			.HAlign(EHorizontalAlignment::HAlign_Fill)
-			.Padding(0.0f, SeparationDistance, 0.0f, 0.f)
 			[
-				CreateAcceptButton()
+				CreateSelectPreset_ChannelList()
 			]
+		]
+
+		+ SVerticalBox::Slot()
+		.AutoHeight()
+		.HAlign(EHorizontalAlignment::HAlign_Fill)
+		.Padding(0.0f, SeparationDistance, 0.0f, 0.f)
+		[
+			CreateSelectPreset_AcceptButton()
 		];
 }
 
-TSharedRef<SWidget> SDMMaterialWizard::CreateChannelPresets()
+TSharedRef<SWidget> SDMMaterialWizard::CreateSelectPreset_ChannelPresets()
 {
 	using namespace UE::DynamicMaterialDesigner::Private;
 
@@ -205,7 +292,7 @@ TSharedRef<SWidget> SDMMaterialWizard::CreateChannelPresets()
 	return ChannelPresets;
 }
 
-TSharedRef<SWidget> SDMMaterialWizard::CreateChannelList()
+TSharedRef<SWidget> SDMMaterialWizard::CreateSelectPreset_ChannelList()
 {
 	using namespace UE::DynamicMaterialDesigner::Private;
 
@@ -216,8 +303,6 @@ TSharedRef<SWidget> SDMMaterialWizard::CreateChannelList()
 		return SNullWidget::NullWidget;
 	}
 
-	UEnum* MaterialPropertyEnum = StaticEnum<EDMMaterialPropertyType>();
-
 	TSharedRef<SWrapBox> ChannelPresets = SNew(SWrapBox)
 		.UseAllottedSize(true)
 		.InnerSlotPadding(WrapBoxSlotPadding)
@@ -248,7 +333,61 @@ TSharedRef<SWidget> SDMMaterialWizard::CreateChannelList()
 	return ChannelPresets;
 }
 
-TSharedRef<SWidget> SDMMaterialWizard::CreateAcceptButton()
+TSharedRef<SWidget> SDMMaterialWizard::CreateTemplateListLayout()
+{
+	TSharedRef<SWrapBox> WrapBox = SNew(SWrapBox)
+		.InnerSlotPadding(FVector2D(5.f))
+		.UseAllottedSize(true);
+
+	TArray<FAssetData> Templates = GetTemplateMaterials();
+
+	for (const FAssetData& Template : Templates)
+	{
+		WrapBox->AddSlot()
+			[
+				CreateTemplateList_Entry(Template)
+			];
+	}
+
+	return WrapBox;
+}
+
+TSharedRef<SWidget> SDMMaterialWizard::CreateTemplateList_Entry(const FAssetData& InTemplateAsset)
+{
+	Assets.Add(InTemplateAsset);
+
+	TSharedRef<FAssetThumbnail> Thumbnail = MakeShared<FAssetThumbnail>(InTemplateAsset, 100.f, 100.f, UThumbnailManager::Get().GetSharedThumbnailPool());
+	Thumbnails.Add(Thumbnail);
+
+	TSharedRef<SWidget> Entry = SNew(SBox)
+		.WidthOverride(100.f)
+		.HeightOverride(123.f)
+		.Cursor(EMouseCursor::Hand)
+		[
+			SNew(SVerticalBox)
+			+ SVerticalBox::Slot()
+			.AutoHeight()
+			[
+				Thumbnail->MakeThumbnailWidget()
+			]
+			+ SVerticalBox::Slot()
+			.AutoHeight()
+			.Padding(0.f, 5.f, 0.f, 0.f)
+			[
+				SNew(STextBlock)
+				.WrapTextAt(100.f)
+				.WrappingPolicy(ETextWrappingPolicy::AllowPerCharacterWrapping)
+				.TextStyle(FDynamicMaterialEditorStyle::Get(), "RegularFont")
+				.Text(FText::FromName(InTemplateAsset.AssetName))
+			]
+		];
+
+	Entry->SetOnMouseButtonDown(FPointerEventHandler::CreateSP(this, &SDMMaterialWizard::OnTemplateMouseDown, Assets.Num() - 1));
+
+	return Entry;
+}
+
+TSharedRef<SWidget> SDMMaterialWizard::CreateSelectPreset_AcceptButton()
 {
 	using namespace UE::DynamicMaterialDesigner::Private;
 
@@ -282,32 +421,48 @@ void SDMMaterialWizard::Preset_OnChange(ECheckBoxState InState, FName InPresetNa
 
 		if (PresetChannelContainer.IsValid())
 		{
-			PresetChannelContainer->SetContent(CreateChannelList());
+			PresetChannelContainer->SetContent(CreateSelectPreset_ChannelList());
 		}
 	}
 }
 
 FReply SDMMaterialWizard::Accept_OnClick()
 {
-	UDynamicMaterialModel* MaterialModel = GetMaterialModel();
-
-	if (!MaterialModel)
+	if (CurrentPreset.IsNone())
 	{
-		return FReply::Handled();
+		FReply::Handled();
 	}
 
-	UDynamicMaterialModelEditorOnlyData* EditorOnlyData = UDynamicMaterialModelEditorOnlyData::Get(MaterialModel);
+	TSharedPtr<SDMMaterialDesigner> DesignerWidget = GetDesignerWidget();
 
-	if (!EditorOnlyData)
+	if (!DesignerWidget.IsValid())
 	{
 		return FReply::Handled();
 	}
 
-	EditorOnlyData->GetOnMaterialBuiltDelegate().RemoveAll(this);
-	EditorOnlyData->SetChannelListPreset(CurrentPreset);
-	EditorOnlyData->OnWizardComplete();
+	if (UDynamicMaterialModel* MaterialModel = GetMaterialModel())
+	{
+		SetChannelListInModel(CurrentPreset, MaterialModel);
+	}
+	else if (MaterialObjectProperty.IsSet())
+	{
+		if (MaterialObjectProperty->IsValid())
+		{
+			CreateTemplateMaterialInActor(CurrentPreset, MaterialObjectProperty.GetValue());
+		}
+		else
+		{
+			UE::DynamicMaterialEditor::Private::LogError(TEXT("Invalid actor property to create new template in."));
+
+			DesignerWidget->ShowSelectPrompt();
+		}
+	}
+	else
+	{
+		UE::DynamicMaterialEditor::Private::LogError(TEXT("Missing material information for new template."));
 
-	OpenMaterialInEditor();
+		DesignerWidget->ShowSelectPrompt();
+	}
 
 	return FReply::Handled();
 }
@@ -349,6 +504,178 @@ void SDMMaterialWizard::OpenMaterialInEditor()
 		return;
 	}
 
+	DesignerWidget->Empty();
+
+	if (MaterialObjectProperty.IsSet())
+	{
+		DesignerWidget->OpenObjectMaterialProperty(MaterialObjectProperty.GetValue());
+	}
+	else if (UDynamicMaterialModel* MaterialModel = GetMaterialModel())
+	{
+		DesignerWidget->OpenMaterialModelBase(MaterialModel);
+	}
+}
+
+ECheckBoxState SDMMaterialWizard::IsModeSelected(int32 InMode) const
+{
+	if (Switcher.IsValid() && Switcher->GetActiveWidgetIndex() == InMode)
+	{
+		return ECheckBoxState::Checked;
+	}
+
+	return ECheckBoxState::Unchecked;
+}
+
+void SDMMaterialWizard::SetMode(ECheckBoxState InState, int32 InMode)
+{
+	if (InState == ECheckBoxState::Checked && Switcher.IsValid())
+	{
+		Switcher->SetActiveWidgetIndex(InMode);
+	}
+}
+
+FReply SDMMaterialWizard::OnTemplateMouseDown(const FGeometry& InGeometry, const FPointerEvent& InPointerEvent, int32 InAssetIndex)
+{
+	if (!Assets.IsValidIndex(InAssetIndex))
+	{
+		return FReply::Handled();
+	}
+
+	TSharedPtr<SDMMaterialDesigner> DesignerWidget = GetDesignerWidget();
+
+	if (!DesignerWidget.IsValid())
+	{
+		return FReply::Handled();
+	}
+
+	UDynamicMaterialInstance* TemplateInstance = Cast<UDynamicMaterialInstance>(Assets[InAssetIndex].GetAsset());
+
+	if (!TemplateInstance)
+	{
+		return FReply::Handled();
+	}
+
+	UDynamicMaterialModel* TemplateModel = Cast<UDynamicMaterialModel>(TemplateInstance->GetMaterialModelBase());
+
+	if (!TemplateModel)
+	{
+		return FReply::Handled();
+	}
+
+	if (UDynamicMaterialModel* MaterialModel = GetMaterialModel())
+	{
+		if (UDynamicMaterialInstance* Instance = MaterialModel->GetDynamicMaterialInstance())
+		{
+			CreateDynamicMaterialInInstance(TemplateModel, Instance);
+		}
+		else
+		{
+			UE::DynamicMaterialEditor::Private::LogError(TEXT("Unable to find material instance to create new dynamic material in."));
+
+			DesignerWidget->ShowSelectPrompt();
+		}
+	}
+	else if (MaterialObjectProperty.IsSet())
+	{
+		if (MaterialObjectProperty->IsValid())
+		{
+			if (UDynamicMaterialInstance* Instance = MaterialObjectProperty->GetMaterial())
+			{
+				CreateDynamicMaterialInInstance(TemplateModel, Instance);
+			}
+			else
+			{
+				CreateNewDynamicInstanceInActor(TemplateModel, MaterialObjectProperty.GetValue());
+			}
+		}
+		else
+		{
+			UE::DynamicMaterialEditor::Private::LogError(TEXT("Invalid actor property to create new dynamic material in."));
+
+			DesignerWidget->ShowSelectPrompt();
+		}
+	}
+	else
+	{
+		UE::DynamicMaterialEditor::Private::LogError(TEXT("Missing material information for dynamic material."));
+
+		DesignerWidget->ShowSelectPrompt();
+	}
+
+	return FReply::Handled();
+}
+
+void SDMMaterialWizard::CreateDynamicMaterialInInstance(UDynamicMaterialModel* InTemplateModel, UDynamicMaterialInstance* InToInstance)
+{
+	TSharedPtr<SDMMaterialDesigner> DesignerWidget = GetDesignerWidget();
+
+	if (!DesignerWidget.IsValid())
+	{
+		return;
+	}
+
+	if (!UDMMaterialModelFunctionLibrary::CreateDynamicModelInInstance(InTemplateModel, InToInstance))
+	{
+		UE::DynamicMaterialEditor::Private::LogError(TEXT("Failed to create new dynamic model in existing instance."));
+		return;
+	}
+
+	DesignerWidget->OpenMaterialModelBase(InToInstance->GetMaterialModelBase());
+}
+
+void SDMMaterialWizard::CreateNewDynamicInstanceInActor(UDynamicMaterialModel* InFromModel, FDMObjectMaterialProperty& InMaterialObjectProperty)
+{
+	TSharedPtr<SDMMaterialDesigner> DesignerWidget = GetDesignerWidget();
+
+	if (!DesignerWidget.IsValid())
+	{
+		return;
+	}
+
+	UObject* Outer = InMaterialObjectProperty.GetOuter();
+
+	if (!Outer)
+	{
+		return;
+	}
+
+	UDynamicMaterialInstanceFactory* Factory = NewObject<UDynamicMaterialInstanceFactory>(GetTransientPackage());
+
+	UDynamicMaterialInstance* NewInstance = Cast<UDynamicMaterialInstance>(Factory->FactoryCreateNew(
+		UDynamicMaterialInstance::StaticClass(),
+		Outer,
+		NAME_None,
+		RF_Transactional | RF_Public,
+		nullptr,
+		GWarn
+	));
+
+	if (!NewInstance)
+	{
+		UE::DynamicMaterialEditor::Private::LogError(TEXT("Failed to create new material instance."));
+		return;
+	}
+
+	if (!UDMMaterialModelFunctionLibrary::CreateDynamicModelInInstance(InFromModel, NewInstance))
+	{
+		UE::DynamicMaterialEditor::Private::LogError(TEXT("Failed to create new dynamic model in new instance."));
+		return;
+	}
+
+	InMaterialObjectProperty.SetMaterial(NewInstance);
+
+	DesignerWidget->OpenObjectMaterialProperty(InMaterialObjectProperty);
+}
+
+void SDMMaterialWizard::SetChannelListInModel(FName InChannelList, UDynamicMaterialModel* InMaterialModel)
+{
+	TSharedPtr<SDMMaterialDesigner> DesignerWidget = GetDesignerWidget();
+
+	if (!DesignerWidget.IsValid())
+	{
+		return;
+	}
+
 	UDynamicMaterialModel* MaterialModel = GetMaterialModel();
 
 	if (!MaterialModel)
@@ -356,16 +683,62 @@ void SDMMaterialWizard::OpenMaterialInEditor()
 		return;
 	}
 
-	DesignerWidget->Empty();
+	UDynamicMaterialModelEditorOnlyData* EditorOnlyData = UDynamicMaterialModelEditorOnlyData::Get(MaterialModel);
 
-	if (MaterialObjectProperty.IsSet())
+	if (!EditorOnlyData)
 	{
-		DesignerWidget->OpenObjectMaterialProperty(MaterialObjectProperty.GetValue());
+		return;
 	}
-	else
+
+	EditorOnlyData->GetOnMaterialBuiltDelegate().RemoveAll(this);
+	EditorOnlyData->SetChannelListPreset(InChannelList);
+	EditorOnlyData->OnWizardComplete();
+
+	DesignerWidget->OpenMaterialModelBase(MaterialModel);
+}
+
+void SDMMaterialWizard::CreateTemplateMaterialInActor(FName InChannelList, FDMObjectMaterialProperty& InMaterialObjectProperty)
+{
+	TSharedPtr<SDMMaterialDesigner> DesignerWidget = GetDesignerWidget();
+
+	if (!DesignerWidget.IsValid())
 	{
-		DesignerWidget->OpenMaterialModelBase(MaterialModel);
+		return;
 	}
+
+	UObject* Outer = InMaterialObjectProperty.GetOuter();
+
+	if (!Outer)
+	{
+		return;
+	}
+
+	UDynamicMaterialInstanceFactory* Factory = NewObject<UDynamicMaterialInstanceFactory>(GetTransientPackage());
+
+	UDynamicMaterialInstance* NewInstance = Cast<UDynamicMaterialInstance>(Factory->FactoryCreateNew(
+		UDynamicMaterialInstance::StaticClass(),
+		Outer,
+		NAME_None,
+		RF_Transactional | RF_Public,
+		nullptr,
+		GWarn
+	));
+
+	if (!NewInstance)
+	{
+		UE::DynamicMaterialEditor::Private::LogError(TEXT("Failed to create new material instance."));
+		return;
+	}
+
+	if (UDynamicMaterialModelEditorOnlyData* EditorOnlyData = UDynamicMaterialModelEditorOnlyData::Get(NewInstance))
+	{
+		EditorOnlyData->SetChannelListPreset(InChannelList);
+		EditorOnlyData->OnWizardComplete();
+	}
+
+	InMaterialObjectProperty.SetMaterial(NewInstance);
+
+	DesignerWidget->OpenMaterialModelBase(NewInstance->GetMaterialModelBase());
 }
 
 #undef LOCTEXT_NAMESPACE
diff --git a/Engine/Plugins/Experimental/DynamicMaterial/Source/DynamicMaterialEditor/Private/UI/Widgets/SDMMaterialWizard.h b/Engine/Plugins/Experimental/DynamicMaterial/Source/DynamicMaterialEditor/Private/UI/Widgets/SDMMaterialWizard.h
--- a/Engine/Plugins/Experimental/DynamicMaterial/Source/DynamicMaterialEditor/Private/UI/Widgets/SDMMaterialWizard.h
+++ b/Engine/Plugins/Experimental/DynamicMaterial/Source/DynamicMaterialEditor/Private/UI/Widgets/SDMMaterialWizard.h
@@ -10,9 +10,11 @@
 #include "UObject/NameTypes.h"
 #include "UObject/WeakObjectPtr.h"
 
+class FAssetThumbnail;
 class SBox;
 class SDMMaterialDesigner;
 class SWidget;
+class SWidgetSwitcher;
 class UDynamicMaterialModel;
 enum class ECheckBoxState : uint8;
 
@@ -38,16 +40,25 @@ public:
 	UDynamicMaterialModel* GetMaterialModel() const;
 
 protected:
+	static TArray<FAssetData> GetTemplateMaterials();
+
 	TWeakPtr<SDMMaterialDesigner> DesignerWidgetWeak;
 	FName CurrentPreset;
 	TSharedPtr<SBox> PresetChannelContainer;
 	TWeakObjectPtr<UDynamicMaterialModel> MaterialModelWeak;
 	TOptional<FDMObjectMaterialProperty> MaterialObjectProperty;
+	TSharedPtr<SWidgetSwitcher> Switcher;
+	TArray<TSharedRef<FAssetThumbnail>> Thumbnails;
+	TArray<FAssetData> Assets;
 
 	TSharedRef<SWidget> CreateLayout();
-	TSharedRef<SWidget> CreateChannelPresets();
-	TSharedRef<SWidget> CreateChannelList();
-	TSharedRef<SWidget> CreateAcceptButton();
+	TSharedRef<SWidget> CreateModeSelector();
+	TSharedRef<SWidget> CreateSelectPresetLayout();
+	TSharedRef<SWidget> CreateSelectPreset_ChannelPresets();
+	TSharedRef<SWidget> CreateSelectPreset_ChannelList();
+	TSharedRef<SWidget> CreateSelectPreset_AcceptButton();
+	TSharedRef<SWidget> CreateTemplateListLayout();
+	TSharedRef<SWidget> CreateTemplateList_Entry(const FAssetData& InTemplateAsset);
 
 	ECheckBoxState Preset_GetState(FName InPresetName) const;
 	void Preset_OnChange(ECheckBoxState InState, FName InPresetName);
@@ -57,4 +68,18 @@ protected:
 	void OnMaterialBuilt(UDynamicMaterialModelBase* InMaterialModel);
 
 	void OpenMaterialInEditor();
+
+	ECheckBoxState IsModeSelected(int32 InMode) const;
+
+	void SetMode(ECheckBoxState InState, int32 InMode);
+
+	FReply OnTemplateMouseDown(const FGeometry& InGeometry, const FPointerEvent& InPointerEvent, int32 InAssetIndex);
+
+	void CreateDynamicMaterialInInstance(UDynamicMaterialModel* InTemplateModel, UDynamicMaterialInstance* InToInstance);
+
+	void CreateNewDynamicInstanceInActor(UDynamicMaterialModel* InFromModel, FDMObjectMaterialProperty& InMaterialObjectProperty);
+
+	void SetChannelListInModel(FName InChannelList, UDynamicMaterialModel* InMaterialModel);
+
+	void CreateTemplateMaterialInActor(FName InChannelList, FDMObjectMaterialProperty& InMaterialObjectProperty);
 };
diff --git a/Engine/Plugins/Experimental/DynamicMaterial/Source/DynamicMaterialEditor/Private/Utils/DMMaterialModelFunctionLibrary.cpp b/Engine/Plugins/Experimental/DynamicMaterial/Source/DynamicMaterialEditor/Private/Utils/DMMaterialModelFunctionLibrary.cpp
--- a/Engine/Plugins/Experimental/DynamicMaterial/Source/DynamicMaterialEditor/Private/Utils/DMMaterialModelFunctionLibrary.cpp
+++ b/Engine/Plugins/Experimental/DynamicMaterial/Source/DynamicMaterialEditor/Private/Utils/DMMaterialModelFunctionLibrary.cpp
@@ -21,6 +21,7 @@
 #include "Materials/Material.h"
 #include "Model/DynamicMaterialModel.h"
 #include "Model/DynamicMaterialModelDynamic.h"
+#include "Model/DynamicMaterialModelEditorOnlyData.h"
 #include "Model/IDynamicMaterialModelEditorOnlyDataInterface.h"
 #include "UObject/Package.h"
 #include "UObject/UObjectGlobals.h"
@@ -533,4 +534,90 @@ bool UDMMaterialModelFunctionLibrary::IsModelValid(UDynamicMaterialModelBase* In
 	return true;
 }
 
+bool UDMMaterialModelFunctionLibrary::DuplicateModelBetweenInstances(UDynamicMaterialModel* InFromModel, UDynamicMaterialInstance* InToInstance)
+{
+	if (!InFromModel || !InToInstance)
+	{
+		return false;
+	}
+
+	UDynamicMaterialModelBase* CurrentModel = InToInstance->GetMaterialModelBase();
+	const FString CurrentName = CurrentModel->GetName();
+
+	if (CurrentModel)
+	{
+		const FString NewName = CurrentName + TEXT("_OLD");
+		CurrentModel->Rename(*NewName, GetTransientPackage(), UE::DynamicMaterial::RenameFlags);
+
+		InToInstance->SetMaterialModel(nullptr);
+	}
+
+	FObjectDuplicationParameters Params = InitStaticDuplicateObjectParams(InFromModel, InToInstance, InFromModel->GetFName(),
+		InFromModel->GetFlags(), nullptr, EDuplicateMode::Normal, EInternalObjectFlags::None);
+
+	UDynamicMaterialModel* NewModel = Cast<UDynamicMaterialModel>(StaticDuplicateObjectEx(Params));
+
+	if (!NewModel)
+	{
+		UE::DynamicMaterialEditor::Private::LogError(TEXT("Failed to copy Material Model."));
+
+		// Put back the original model
+		CurrentModel->Rename(*CurrentName, InToInstance, UE::DynamicMaterial::RenameFlags);
+
+		return false;
+	}
+
+	NewModel->Rename(*CurrentName, InToInstance, UE::DynamicMaterial::RenameFlags);
+
+	InToInstance->SetMaterialModel(NewModel);
+	NewModel->SetDynamicMaterialInstance(InToInstance);
+	InToInstance->InitializeMIDPublic();
+
+	if (UDynamicMaterialModelEditorOnlyData* EditorOnlyData = UDynamicMaterialModelEditorOnlyData::Get(NewModel))
+	{
+		EditorOnlyData->RequestMaterialBuild();
+	}
+
+	return true;
+}
+
+bool UDMMaterialModelFunctionLibrary::CreateDynamicModelInInstance(UDynamicMaterialModel* InFromModel, UDynamicMaterialInstance* InToInstance)
+{
+	if (!InFromModel || !InToInstance)
+	{
+		return false;
+	}
+
+	UDynamicMaterialModelBase* CurrentModel = InToInstance->GetMaterialModelBase();
+	const FString CurrentName = CurrentModel->GetName();
+
+	if (CurrentModel)
+	{
+		const FString NewName = CurrentName + TEXT("_OLD");
+		CurrentModel->Rename(*NewName, GetTransientPackage(), UE::DynamicMaterial::RenameFlags);
+
+		InToInstance->SetMaterialModel(nullptr);
+	}
+
+	UDynamicMaterialModelDynamic* NewModelDynamic = UDynamicMaterialModelDynamic::Create(InToInstance, InFromModel);
+
+	if (!NewModelDynamic)
+	{
+		UE::DynamicMaterialEditor::Private::LogError(TEXT("Failed to make Dynamic Material."));
+
+		// Put back the original model
+		CurrentModel->Rename(*CurrentName, InToInstance, UE::DynamicMaterial::RenameFlags);
+
+		return false;
+	}
+
+	NewModelDynamic->Rename(*CurrentName, InToInstance, UE::DynamicMaterial::RenameFlags);
+
+	InToInstance->SetMaterialModel(NewModelDynamic);
+	NewModelDynamic->SetDynamicMaterialInstance(InToInstance);
+	InToInstance->InitializeMIDPublic();
+
+	return true;
+}
+
 #undef LOCTEXT_NAMESPACE
diff --git a/Engine/Plugins/Experimental/DynamicMaterial/Source/DynamicMaterialEditor/Public/Model/DynamicMaterialModelEditorOnlyData.h b/Engine/Plugins/Experimental/DynamicMaterial/Source/DynamicMaterialEditor/Public/Model/DynamicMaterialModelEditorOnlyData.h
--- a/Engine/Plugins/Experimental/DynamicMaterial/Source/DynamicMaterialEditor/Public/Model/DynamicMaterialModelEditorOnlyData.h
+++ b/Engine/Plugins/Experimental/DynamicMaterial/Source/DynamicMaterialEditor/Public/Model/DynamicMaterialModelEditorOnlyData.h
@@ -149,9 +149,6 @@ public:
 	UFUNCTION(BlueprintCallable, Category = "Material Designer")
 	DYNAMICMATERIALEDITOR_API void SetResponsiveAAEnabled(bool bInEnabled);
 
-	UFUNCTION(BlueprintCallable, Category = "Material Designer")
-	DYNAMICMATERIALEDITOR_API FName GetChannelListPreset() const;
-
 	UFUNCTION(BlueprintCallable, Category = "Material Designer")
 	DYNAMICMATERIALEDITOR_API void SetChannelListPreset(FName InPresetName);
 
@@ -238,7 +235,6 @@ public:
 	DYNAMICMATERIALEDITOR_API virtual void PostEditImport() override;
 	DYNAMICMATERIALEDITOR_API virtual void PostDuplicate(bool bInDuplicateForPIE) override;
 	DYNAMICMATERIALEDITOR_API virtual void PostEditChangeChainProperty(FPropertyChangedChainEvent& InPropertyChangedEvent) override;
-	DYNAMICMATERIALEDITOR_API virtual void Serialize(FArchive& InAr) override;
 	//~ End UObject
 
 	//~ Begin IDMBuildable
@@ -301,11 +297,6 @@ protected:
 	UPROPERTY(EditInstanceOnly, BlueprintReadOnly, Category = "Material Designer", meta = (DisplayName = "Responsive AA"))
 	bool bResponsiveAAEnabled;
 
-	/** Sets the available channels and default material properties. */
-	UPROPERTY(EditInstanceOnly, BlueprintReadOnly, Category = "Material Designer", 
-		meta = (GetOptions = GetPresetOptions, NoResetToDefault, DisplayName = "Material Type Preset"))
-	FName ChannelListPreset;
-
 	UPROPERTY(VisibleInstanceOnly, BlueprintReadOnly, TextExportTransient, Category = "Material Designer")
 	TMap<EDMMaterialPropertyType, TObjectPtr<UDMMaterialSlot>> PropertySlotMap;
 
@@ -411,11 +402,6 @@ protected:
 
 	void AssignPropertyAlphaValues();
 
-	UFUNCTION()
-	TArray<FName> GetPresetOptions() const;
-
-	void OnChannelListPresetChanged();
-	void EnsurePresetSlots();
 	void OnDomainChanged();
 	void OnBlendModeChanged();
 	void OnShadingModelChanged();
diff --git a/Engine/Plugins/Experimental/DynamicMaterial/Source/DynamicMaterialEditor/Public/Utils/DMMaterialModelFunctionLibrary.h b/Engine/Plugins/Experimental/DynamicMaterial/Source/DynamicMaterialEditor/Public/Utils/DMMaterialModelFunctionLibrary.h
--- a/Engine/Plugins/Experimental/DynamicMaterial/Source/DynamicMaterialEditor/Public/Utils/DMMaterialModelFunctionLibrary.h
+++ b/Engine/Plugins/Experimental/DynamicMaterial/Source/DynamicMaterialEditor/Public/Utils/DMMaterialModelFunctionLibrary.h
@@ -45,4 +45,10 @@ public:
 	DYNAMICMATERIALEDITOR_API static UDynamicMaterialInstance* ExportToTemplateMaterialInstance(UDynamicMaterialModelDynamic* InMaterialModelDynamic, const FString& InSavePath);
 
 	DYNAMICMATERIALEDITOR_API static bool IsModelValid(UDynamicMaterialModelBase* InMaterialModelBase);
+
+	DYNAMICMATERIALEDITOR_API static bool DuplicateModelBetweenInstances(UDynamicMaterialModel* InFromModel, UDynamicMaterialInstance* InToInstance);
+
+	DYNAMICMATERIALEDITOR_API static bool CreateDynamicModelInInstance(UDynamicMaterialModel* InFromModel, UDynamicMaterialInstance* InToInstance);
+
+
 };
