Material Designer
- Horde issue 569583

diff --git a/Engine/Plugins/Experimental/DynamicMaterial/Source/DynamicMaterialEditor/Private/Utils/DMMaterialModelFunctionLibrary.cpp b/Engine/Plugins/Experimental/DynamicMaterial/Source/DynamicMaterialEditor/Private/Utils/DMMaterialModelFunctionLibrary.cpp
--- a/Engine/Plugins/Experimental/DynamicMaterial/Source/DynamicMaterialEditor/Private/Utils/DMMaterialModelFunctionLibrary.cpp
+++ b/Engine/Plugins/Experimental/DynamicMaterial/Source/DynamicMaterialEditor/Private/Utils/DMMaterialModelFunctionLibrary.cpp
@@ -542,10 +542,12 @@ bool UDMMaterialModelFunctionLibrary::DuplicateModelBetweenInstances(UDynamicMat
 	}
 
 	UDynamicMaterialModelBase* CurrentModel = InToInstance->GetMaterialModelBase();
-	const FString CurrentName = CurrentModel->GetName();
+	FString CurrentName;
 
 	if (CurrentModel)
 	{
+		CurrentName = CurrentModel->GetName();
+
 		const FString NewName = CurrentName + TEXT("_OLD");
 		CurrentModel->Rename(*NewName, GetTransientPackage(), UE::DynamicMaterial::RenameFlags);
 
@@ -562,12 +564,12 @@ bool UDMMaterialModelFunctionLibrary::DuplicateModelBetweenInstances(UDynamicMat
 		UE::DynamicMaterialEditor::Private::LogError(TEXT("Failed to copy Material Model."));
 
 		// Put back the original model
-		CurrentModel->Rename(*CurrentName, InToInstance, UE::DynamicMaterial::RenameFlags);
+		CurrentModel->Rename(CurrentName.IsEmpty() ? nullptr : *CurrentName, InToInstance, UE::DynamicMaterial::RenameFlags);
 
 		return false;
 	}
 
-	NewModel->Rename(*CurrentName, InToInstance, UE::DynamicMaterial::RenameFlags);
+	NewModel->Rename(CurrentName.IsEmpty() ? nullptr : *CurrentName, InToInstance, UE::DynamicMaterial::RenameFlags);
 
 	InToInstance->SetMaterialModel(NewModel);
 	NewModel->SetDynamicMaterialInstance(InToInstance);
@@ -589,10 +591,12 @@ bool UDMMaterialModelFunctionLibrary::CreateDynamicModelInInstance(UDynamicMater
 	}
 
 	UDynamicMaterialModelBase* CurrentModel = InToInstance->GetMaterialModelBase();
-	const FString CurrentName = CurrentModel->GetName();
+	FString CurrentName;	
 
 	if (CurrentModel)
 	{
+		CurrentName = CurrentModel->GetName();
+
 		const FString NewName = CurrentName + TEXT("_OLD");
 		CurrentModel->Rename(*NewName, GetTransientPackage(), UE::DynamicMaterial::RenameFlags);
 
@@ -606,12 +610,12 @@ bool UDMMaterialModelFunctionLibrary::CreateDynamicModelInInstance(UDynamicMater
 		UE::DynamicMaterialEditor::Private::LogError(TEXT("Failed to make Dynamic Material."));
 
 		// Put back the original model
-		CurrentModel->Rename(*CurrentName, InToInstance, UE::DynamicMaterial::RenameFlags);
+		CurrentModel->Rename(CurrentName.IsEmpty() ? nullptr : *CurrentName, InToInstance, UE::DynamicMaterial::RenameFlags);
 
 		return false;
 	}
 
-	NewModelDynamic->Rename(*CurrentName, InToInstance, UE::DynamicMaterial::RenameFlags);
+	NewModelDynamic->Rename(CurrentName.IsEmpty() ? nullptr : *CurrentName, InToInstance, UE::DynamicMaterial::RenameFlags);
 
 	InToInstance->SetMaterialModel(NewModelDynamic);
 	NewModelDynamic->SetDynamicMaterialInstance(InToInstance);
