Material Designer
- Added back in the stage row generator to remove the alpha value appearing in texture settings.

diff --git a/Engine/Plugins/Experimental/DynamicMaterial/Source/DynamicMaterialEditor/Private/DynamicMaterialEditorModule.cpp b/Engine/Plugins/Experimental/DynamicMaterial/Source/DynamicMaterialEditor/Private/DynamicMaterialEditorModule.cpp
--- a/Engine/Plugins/Experimental/DynamicMaterial/Source/DynamicMaterialEditor/Private/DynamicMaterialEditorModule.cpp
+++ b/Engine/Plugins/Experimental/DynamicMaterial/Source/DynamicMaterialEditor/Private/DynamicMaterialEditorModule.cpp
@@ -40,6 +40,7 @@
 #include "UI/PropertyGenerators/DMMaterialStageFunctionPropertyRowGenerator.h"
 #include "UI/PropertyGenerators/DMMaterialValueDynamicPropertyRowGenerator.h"
 #include "UI/PropertyGenerators/DMMaterialValuePropertyRowGenerator.h"
+#include "UI/PropertyGenerators/DMStagePropertyRowGenerator.h"
 #include "UI/PropertyGenerators/DMTextureUVDynamicPropertyRowGenerator.h"
 #include "UI/PropertyGenerators/DMTextureUVPropertyRowGenerator.h"
 #include "UI/PropertyGenerators/DMThroughputPropertyRowGenerator.h"
@@ -232,6 +233,7 @@ void FDynamicMaterialEditorModule::StartupModule()
 	FDMMaterialModelDefaults::RegisterDefaultsDelegates();
 
 	RegisterComponentPropertyRowGeneratorDelegate<UDMMaterialComponent,            FDMComponentPropertyRowGenerator>();
+	RegisterComponentPropertyRowGeneratorDelegate<UDMMaterialStage,                FDMStagePropertyRowGenerator>();
 	RegisterComponentPropertyRowGeneratorDelegate<UDMMaterialValue,                FDMMaterialValuePropertyRowGenerator>();
 	RegisterComponentPropertyRowGeneratorDelegate<UDMMaterialValueDynamic,         FDMMaterialValueDynamicPropertyRowGenerator>();
 	RegisterComponentPropertyRowGeneratorDelegate<UDMTextureUV,                    FDMTextureUVPropertyRowGenerator>();
diff --git a/Engine/Plugins/Experimental/DynamicMaterial/Source/DynamicMaterialEditor/Private/UI/PropertyGenerators/DMStagePropertyRowGenerator.cpp b/Engine/Plugins/Experimental/DynamicMaterial/Source/DynamicMaterialEditor/Private/UI/PropertyGenerators/DMStagePropertyRowGenerator.cpp
new file mode 100644
--- /dev/null
+++ b/Engine/Plugins/Experimental/DynamicMaterial/Source/DynamicMaterialEditor/Private/UI/PropertyGenerators/DMStagePropertyRowGenerator.cpp
@@ -0,0 +1,64 @@
+// Copyright Epic Games, Inc. All Rights Reserved.
+
+#include "UI/PropertyGenerators/DMStagePropertyRowGenerator.h"
+
+#include "Components/DMMaterialComponent.h"
+#include "Components/DMMaterialLayer.h"
+#include "Components/DMMaterialStage.h"
+#include "Components/DMMaterialStageSource.h"
+#include "DetailLayoutBuilder.h"
+#include "DMEDefs.h"
+#include "DynamicMaterialEditorModule.h"
+#include "Internationalization/Text.h"
+#include "Internationalization/Text.h"
+#include "Model/DynamicMaterialModel.h"
+#include "Model/DynamicMaterialModelBase.h"
+#include "UI/Menus/DMMaterialStageSourceMenus.h"
+#include "UI/Widgets/Editor/SDMMaterialComponentEditor.h"
+#include "UI/Widgets/Editor/SDMMaterialSlotEditor.h"
+#include "UI/Widgets/Editor/SlotEditor/SDMMaterialSlotLayerItem.h"
+#include "UI/Widgets/Editor/SlotEditor/SDMMaterialSlotLayerView.h"
+#include "UI/Widgets/Editor/SlotEditor/SDMMaterialStage.h"
+#include "UI/Widgets/SDMMaterialEditor.h"
+#include "Widgets/Input/SComboButton.h"
+
+#define LOCTEXT_NAMESPACE "DMStagePropertyRowGenerator"
+
+const TSharedRef<FDMStagePropertyRowGenerator>& FDMStagePropertyRowGenerator::Get()
+{
+	static TSharedRef<FDMStagePropertyRowGenerator> Generator = MakeShared<FDMStagePropertyRowGenerator>();
+	return Generator;
+}
+
+void FDMStagePropertyRowGenerator::AddComponentProperties(const TSharedRef<SDMMaterialComponentEditor>& InComponentEditorWidget, UDMMaterialComponent* InComponent,
+	TArray<FDMPropertyHandle>& InOutPropertyRows, TSet<UDMMaterialComponent*>& InOutProcessedObjects)
+{
+	if (!IsValid(InComponent))
+	{
+		return;
+	}
+
+	if (InOutProcessedObjects.Contains(InComponent))
+	{
+		return;
+	}
+
+	UDMMaterialStage* Stage = Cast<UDMMaterialStage>(InComponent);
+
+	if (!Stage)
+	{
+		return;
+	}
+
+	UDMMaterialStageSource* Source = Stage->GetSource();
+
+	if (!Source)
+	{
+		return;
+	}
+
+	FDynamicMaterialEditorModule::GeneratorComponentPropertyRows(InComponentEditorWidget, Source, InOutPropertyRows, InOutProcessedObjects);
+	FDMComponentPropertyRowGenerator::AddComponentProperties(InComponentEditorWidget, Stage, InOutPropertyRows, InOutProcessedObjects);
+}
+
+#undef LOCTEXT_NAMESPACE
diff --git a/Engine/Plugins/Experimental/DynamicMaterial/Source/DynamicMaterialEditor/Private/UI/PropertyGenerators/DMStagePropertyRowGenerator.h b/Engine/Plugins/Experimental/DynamicMaterial/Source/DynamicMaterialEditor/Private/UI/PropertyGenerators/DMStagePropertyRowGenerator.h
new file mode 100644
--- /dev/null
+++ b/Engine/Plugins/Experimental/DynamicMaterial/Source/DynamicMaterialEditor/Private/UI/PropertyGenerators/DMStagePropertyRowGenerator.h
@@ -0,0 +1,22 @@
+// Copyright Epic Games, Inc. All Rights Reserved.
+
+#pragma once
+
+#include "UI/PropertyGenerators/DMComponentPropertyRowGenerator.h"
+
+class FText;
+class SWidget;
+class UDMMaterialStage;
+
+class FDMStagePropertyRowGenerator : public FDMComponentPropertyRowGenerator
+{
+public:
+	static const TSharedRef<FDMStagePropertyRowGenerator>& Get();
+
+	virtual ~FDMStagePropertyRowGenerator() override = default;
+
+	//~ Begin FDMComponentPropertyRowGenerator
+	virtual void AddComponentProperties(const TSharedRef<SDMMaterialComponentEditor>& InComponentEditorWidget, UDMMaterialComponent* InComponent,
+		TArray<FDMPropertyHandle>& InOutPropertyRows, TSet<UDMMaterialComponent*>& InOutProcessedObjects) override;
+	//~ End FDMComponentPropertyRowGenerator
+};
