Material Designer
- Using content browser widget for template selector.

diff --git a/Engine/Plugins/Experimental/DynamicMaterial/Source/DynamicMaterialEditor/Private/DynamicMaterialEditorSettings.cpp b/Engine/Plugins/Experimental/DynamicMaterial/Source/DynamicMaterialEditor/Private/DynamicMaterialEditorSettings.cpp
--- a/Engine/Plugins/Experimental/DynamicMaterial/Source/DynamicMaterialEditor/Private/DynamicMaterialEditorSettings.cpp
+++ b/Engine/Plugins/Experimental/DynamicMaterial/Source/DynamicMaterialEditor/Private/DynamicMaterialEditorSettings.cpp
@@ -134,7 +134,6 @@ UDynamicMaterialEditorSettings::UDynamicMaterialEditorSettings()
 
 	bFollowSelection = true;
 	bUseLinearColorForVectors = true;
-	CustomTemplateFolders = {};
 
 	ResetAllLayoutSettings();
 
@@ -428,59 +427,6 @@ TArray<FDMMaterialEffectList> UDynamicMaterialEditorSettings::GetEffectList() co
 	return Effects;
 }
 
-TArray<FAssetData> UDynamicMaterialEditorSettings::GetTemplateList() const
-{
-	TArray<FAssetData> Templates;
-
-	IAssetRegistry* AssetRegistry = IAssetRegistry::Get();
-
-	if (!AssetRegistry)
-	{
-		return Templates;
-	}
-
-	TSharedPtr<IPlugin> Plugin = IPluginManager::Get().FindPlugin(UE_PLUGIN_NAME);
-
-	if (!Plugin.IsValid())
-	{
-		return Templates;
-	}
-
-	const FString PluginEffectPath = Plugin->GetMountedAssetPath() / "Templates";
-
-	TArray<FName> AssetPaths = {*PluginEffectPath};
-	AssetPaths.Append(UDynamicMaterialEditorSettings::Get()->CustomTemplateFolders);
-
-	TArray<FAssetData> Assets;
-	AssetRegistry->GetAssetsByPaths(AssetPaths, Assets, /* bRecursive */ true, /* Only Assets on Disk */ true);
-
-	for (const FAssetData& Asset : Assets)
-	{
-		if (Asset.GetClass(EResolveClass::Yes) != UDynamicMaterialInstance::StaticClass())
-		{
-			continue;
-		}
-
-		UDynamicMaterialInstance* Instance = Cast<UDynamicMaterialInstance>(Asset.GetAsset());
-
-		if (!Instance)
-		{
-			continue;
-		}
-
-		UDynamicMaterialModel* MaterialModel = Cast<UDynamicMaterialModel>(Instance->GetMaterialModelBase());
-
-		if (!MaterialModel)
-		{
-			continue;
-		}
-
-		Templates.Add(Asset);
-	}
-
-	return Templates;
-}
-
 const FDMDefaultMaterialPropertySlotValue& UDynamicMaterialEditorSettings::GetDefaultSlotValue(EDMMaterialPropertyType InProperty) const
 {
 	if (const FDMDefaultMaterialPropertySlotValue* OverridePtr = DefaultSlotValueOverrides.Find(InProperty))
diff --git a/Engine/Plugins/Experimental/DynamicMaterial/Source/DynamicMaterialEditor/Private/DynamicMaterialEditorSettings.h b/Engine/Plugins/Experimental/DynamicMaterial/Source/DynamicMaterialEditor/Private/DynamicMaterialEditorSettings.h
--- a/Engine/Plugins/Experimental/DynamicMaterial/Source/DynamicMaterialEditor/Private/DynamicMaterialEditorSettings.h
+++ b/Engine/Plugins/Experimental/DynamicMaterial/Source/DynamicMaterialEditor/Private/DynamicMaterialEditorSettings.h
@@ -244,16 +244,6 @@ public:
 	UPROPERTY(Config)
 	bool bValidatedPresets = false;
 
-	/*
-	 * Add paths to search for custom material templates.
-	 *
-	 * Format examples:
-	 * - /Game/Some/Path
-	 * - /Plugin/Some/Path
-	 */
-	UPROPERTY(Config, EditAnywhere, BlueprintReadWrite, Category = "Templates")
-	TArray<FName> CustomTemplateFolders;
-
 	/** This variable is accessed in multiple places, so this is a quick accessor. */
 	static bool IsUseLinearColorForVectorsEnabled();
 
@@ -263,8 +253,6 @@ public:
 
 	TArray<FDMMaterialEffectList> GetEffectList() const;
 
-	TArray<FAssetData> GetTemplateList() const;
-
 	const FDMDefaultMaterialPropertySlotValue& GetDefaultSlotValue(EDMMaterialPropertyType InProperty) const;
 
 	const FDMMaterialChannelListPreset* GetPresetByName(FName InName) const;
diff --git a/Engine/Plugins/Experimental/DynamicMaterial/Source/DynamicMaterialEditor/Private/UI/Widgets/SDMMaterialWizard.cpp b/Engine/Plugins/Experimental/DynamicMaterial/Source/DynamicMaterialEditor/Private/UI/Widgets/SDMMaterialWizard.cpp
--- a/Engine/Plugins/Experimental/DynamicMaterial/Source/DynamicMaterialEditor/Private/UI/Widgets/SDMMaterialWizard.cpp
+++ b/Engine/Plugins/Experimental/DynamicMaterial/Source/DynamicMaterialEditor/Private/UI/Widgets/SDMMaterialWizard.cpp
@@ -3,26 +3,31 @@
 #include "UI/Widgets/SDMMaterialWizard.h"
 
 #include "AssetRegistry/AssetData.h"
+#include "AssetTextFilter.h"
 #include "AssetThumbnail.h"
 #include "DMDefs.h"
 #include "DynamicMaterialEditorSettings.h"
 #include "DynamicMaterialEditorStyle.h"
 #include "DynamicMaterialModule.h"
 #include "Engine/EngineTypes.h"
+#include "IContentBrowserSingleton.h"
 #include "Material/DynamicMaterialInstance.h"
 #include "Material/DynamicMaterialInstanceFactory.h"
 #include "Model/DynamicMaterialModel.h"
 #include "Model/DynamicMaterialModelEditorOnlyData.h"
+#include "SAssetSearchBox.h"
+#include "SAssetView.h"
 #include "Styling/SlateTypes.h"
 #include "Styling/StyleColors.h"
-#include "ThumbnailRendering/ThumbnailManager.h"
 #include "UI/Widgets/SDMMaterialDesigner.h"
 #include "Utils/DMMaterialModelFunctionLibrary.h"
 #include "Utils/DMPrivate.h"
 #include "Widgets/Input/SButton.h"
 #include "Widgets/Input/SCheckBox.h"
+#include "Widgets/Layout/SBorder.h"
 #include "Widgets/Layout/SBox.h"
 #include "Widgets/Layout/SExpandableArea.h"
+#include "Widgets/Layout/SScrollBox.h"
 #include "Widgets/Layout/SWidgetSwitcher.h"
 #include "Widgets/Layout/SWrapBox.h"
 #include "Widgets/SBoxPanel.h"
@@ -98,7 +103,12 @@ void SDMMaterialWizard::Construct(const FArguments& InArgs, const TSharedRef<SDM
 
 	ChildSlot
 	[
-		CreateLayout()
+		SNew(SScrollBox)
+		+ SScrollBox::Slot()
+		.VAlign(EVerticalAlignment::VAlign_Fill)
+		[
+			CreateLayout()
+		]
 	];
 }
 
@@ -112,16 +122,6 @@ UDynamicMaterialModel* SDMMaterialWizard::GetMaterialModel() const
 	return MaterialModelWeak.Get();
 }
 
-TArray<FAssetData> SDMMaterialWizard::GetTemplateMaterials()
-{
-	if (const UDynamicMaterialEditorSettings* Settings = GetDefault<UDynamicMaterialEditorSettings>())
-	{
-		return Settings->GetTemplateList();
-	}
-
-	return {};
-}
-
 TSharedRef<SWidget> SDMMaterialWizard::CreateLayout()
 {
 	using namespace UE::DynamicMaterialDesigner::Private;
@@ -335,56 +335,99 @@ TSharedRef<SWidget> SDMMaterialWizard::CreateSelectPreset_ChannelList()
 
 TSharedRef<SWidget> SDMMaterialWizard::CreateTemplateListLayout()
 {
-	TSharedRef<SWrapBox> WrapBox = SNew(SWrapBox)
-		.InnerSlotPadding(FVector2D(5.f))
-		.UseAllottedSize(true);
-
-	TArray<FAssetData> Templates = GetTemplateMaterials();
-
-	for (const FAssetData& Template : Templates)
-	{
-		WrapBox->AddSlot()
-			[
-				CreateTemplateList_Entry(Template)
-			];
-	}
-
-	return WrapBox;
-}
-
-TSharedRef<SWidget> SDMMaterialWizard::CreateTemplateList_Entry(const FAssetData& InTemplateAsset)
-{
-	Assets.Add(InTemplateAsset);
-
-	TSharedRef<FAssetThumbnail> Thumbnail = MakeShared<FAssetThumbnail>(InTemplateAsset, 100.f, 100.f, UThumbnailManager::Get().GetSharedThumbnailPool());
-	Thumbnails.Add(Thumbnail);
-
-	TSharedRef<SWidget> Entry = SNew(SBox)
-		.WidthOverride(100.f)
-		.HeightOverride(123.f)
-		.Cursor(EMouseCursor::Hand)
+	FAssetPickerConfig PickerConfig;
+	PickerConfig.SelectionMode = ESelectionMode::Single;
+	PickerConfig.Filter.ClassPaths.Add(FTopLevelAssetPath(UDynamicMaterialInstance::StaticClass()));
+	PickerConfig.Filter.bIncludeOnlyOnDiskAssets = true;
+	PickerConfig.ThumbnailLabel = EThumbnailLabel::AssetName;
+	PickerConfig.OnShouldFilterAsset.BindSP(this, &SDMMaterialWizard::ShouldFilterOutAsset);
+	PickerConfig.InitialAssetViewType = EAssetViewType::Tile;
+	PickerConfig.bAutohideSearchBar = false;
+	PickerConfig.bFocusSearchBoxWhenOpened = false;
+	PickerConfig.bShowBottomToolbar = false;
+	PickerConfig.bAllowDragging = false;
+	PickerConfig.bAllowRename = false;
+	PickerConfig.bCanShowClasses = false;
+	PickerConfig.bCanShowFolders = true;
+	PickerConfig.bCanShowReadOnlyFolders = true;
+	PickerConfig.bCanShowRealTimeThumbnails = true;
+	PickerConfig.bCanShowDevelopersFolder = true;
+	PickerConfig.bForceShowEngineContent = true;
+	PickerConfig.bForceShowPluginContent = true;
+	PickerConfig.bAddFilterUI = false;
+
+	PickerConfig.Filter.PackagePaths.Reset();
+
+	TextFilter = MakeShared<FAssetTextFilter>();
+	TextFilter->SetIncludeClassName(true);
+	TextFilter->SetIncludeAssetPath(false);
+	TextFilter->SetIncludeCollectionNames(false);
+
+	AssetView = SNew(SAssetView)
+		.InitialCategoryFilter(EContentBrowserItemCategoryFilter::IncludeAssets)
+		.SelectionMode(PickerConfig.SelectionMode)
+		.OnShouldFilterAsset(PickerConfig.OnShouldFilterAsset)
+		.OnItemsActivated(this, &SDMMaterialWizard::OnAssetsActivated)
+		.OnIsAssetValidForCustomToolTip(PickerConfig.OnIsAssetValidForCustomToolTip)
+		.OnGetCustomAssetToolTip(PickerConfig.OnGetCustomAssetToolTip)
+		.OnVisualizeAssetToolTip(PickerConfig.OnVisualizeAssetToolTip)
+		.OnAssetToolTipClosing(PickerConfig.OnAssetToolTipClosing)
+		.InitialSourcesData(FSourcesData())
+		.InitialBackendFilter(PickerConfig.Filter)
+		.InitialViewType(PickerConfig.InitialAssetViewType)
+		.InitialAssetSelection(PickerConfig.InitialAssetSelection)
+		.ShowBottomToolbar(PickerConfig.bShowBottomToolbar)
+		.OnAssetTagWantsToBeDisplayed(PickerConfig.OnAssetTagWantsToBeDisplayed)
+		.OnGetCustomSourceAssets(PickerConfig.OnGetCustomSourceAssets)
+		.AllowDragging(PickerConfig.bAllowDragging)
+		.CanShowClasses(PickerConfig.bCanShowClasses)
+		.CanShowFolders(PickerConfig.bCanShowFolders)
+		.CanShowReadOnlyFolders(PickerConfig.bCanShowReadOnlyFolders)
+		.ShowPathInColumnView(PickerConfig.bShowPathInColumnView)
+		.ShowTypeInColumnView(PickerConfig.bShowTypeInColumnView)
+		.ShowViewOptions(false)
+		.SortByPathInColumnView(PickerConfig.bSortByPathInColumnView)
+		.FilterRecursivelyWithBackendFilter(false)
+		.CanShowRealTimeThumbnails(PickerConfig.bCanShowRealTimeThumbnails)
+		.CanShowDevelopersFolder(PickerConfig.bCanShowDevelopersFolder)
+		.ForceShowEngineContent(PickerConfig.bForceShowEngineContent)
+		.ForceShowPluginContent(PickerConfig.bForceShowPluginContent)
+		.HighlightedText(TAttribute<FText>(this, &SDMMaterialWizard::GetSearchText))
+		.ThumbnailLabel(PickerConfig.ThumbnailLabel)
+		.AssetShowWarningText(PickerConfig.AssetShowWarningText)
+		.AllowFocusOnSync(false)
+		.HiddenColumnNames(PickerConfig.HiddenColumnNames)
+		.CustomColumns(PickerConfig.CustomColumns)
+		.InitialThumbnailSize(EThumbnailSize::Small)
+		.ShowTypeInTileView(false)
+		.TextFilter(TextFilter);
+
+	AssetView->RequestSlowFullListRefresh();
+
+	return SNew(SBorder)
+		.BorderImage(FAppStyle::Get().GetBrush("Menu.Background"))
+		.Padding(3.f, 3.f, 3.f, 3.f)
+		.VAlign(EVerticalAlignment::VAlign_Fill)
 		[
 			SNew(SVerticalBox)
+
 			+ SVerticalBox::Slot()
 			.AutoHeight()
+			.Padding(0.f, 0.f, 0.f, 3.f)
 			[
-				Thumbnail->MakeThumbnailWidget()
+				SAssignNew(AssetSearchBox, SAssetSearchBox)
+				.HintText(NSLOCTEXT("ContentBrowser", "SearchBoxHint", "Search Assets"))
+				.OnTextChanged(this, &SDMMaterialWizard::OnSearchBoxChanged)
+				.OnTextCommitted(this, &SDMMaterialWizard::OnSearchBoxCommitted)
+				.DelayChangeNotificationsWhileTyping(true)
 			]
+
 			+ SVerticalBox::Slot()
-			.AutoHeight()
-			.Padding(0.f, 5.f, 0.f, 0.f)
+			.FillHeight(1.f)
 			[
-				SNew(STextBlock)
-				.WrapTextAt(100.f)
-				.WrappingPolicy(ETextWrappingPolicy::AllowPerCharacterWrapping)
-				.TextStyle(FDynamicMaterialEditorStyle::Get(), "RegularFont")
-				.Text(FText::FromName(InTemplateAsset.AssetName))
+				AssetView.ToSharedRef()
 			]
 		];
-
-	Entry->SetOnMouseButtonDown(FPointerEventHandler::CreateSP(this, &SDMMaterialWizard::OnTemplateMouseDown, Assets.Num() - 1));
-
-	return Entry;
 }
 
 TSharedRef<SWidget> SDMMaterialWizard::CreateSelectPreset_AcceptButton()
@@ -534,39 +577,101 @@ void SDMMaterialWizard::SetMode(ECheckBoxState InState, int32 InMode)
 	}
 }
 
-FReply SDMMaterialWizard::OnTemplateMouseDown(const FGeometry& InGeometry, const FPointerEvent& InPointerEvent, int32 InAssetIndex)
+void SDMMaterialWizard::OnSearchBoxChanged(const FText& InSearchText)
+{
+	SetSearchText(InSearchText);
+}
+
+void SDMMaterialWizard::OnSearchBoxCommitted(const FText& InSearchText, ETextCommit::Type InCommitInfo)
 {
-	if (!Assets.IsValidIndex(InAssetIndex))
+	SetSearchText(InSearchText);
+}
+
+FText SDMMaterialWizard::GetSearchText() const
+{
+	if (TextFilter.IsValid())
 	{
-		return FReply::Handled();
+		return TextFilter->GetRawFilterText();
 	}
 
-	TSharedPtr<SDMMaterialDesigner> DesignerWidget = GetDesignerWidget();
+	return FText::GetEmpty();
+}
 
-	if (!DesignerWidget.IsValid())
+void SDMMaterialWizard::SetSearchText(const FText& InSearchText)
+{
+	if (InSearchText.ToString().Equals(TextFilter->GetRawFilterText().ToString(), ESearchCase::CaseSensitive))
 	{
-		return FReply::Handled();
+		return;
 	}
 
-	UDynamicMaterialInstance* TemplateInstance = Cast<UDynamicMaterialInstance>(Assets[InAssetIndex].GetAsset());
+	TextFilter->SetRawFilterText(InSearchText);
+	AssetView->SetUserSearching(!InSearchText.IsEmpty());
+}
+
+bool SDMMaterialWizard::ShouldFilterOutAsset(const FAssetData& InAsset) const
+{
+	UDynamicMaterialInstance* Instance = Cast<UDynamicMaterialInstance>(InAsset.GetAsset());
 
-	if (!TemplateInstance)
+	if (!Instance)
 	{
-		return FReply::Handled();
+		return true;
 	}
 
-	UDynamicMaterialModel* TemplateModel = Cast<UDynamicMaterialModel>(TemplateInstance->GetMaterialModelBase());
+	UDynamicMaterialModelBase* MaterialModelBase = Instance->GetMaterialModelBase();
 
-	if (!TemplateModel)
+	if (!MaterialModelBase)
 	{
-		return FReply::Handled();
+		return true;
+	}
+
+	return !MaterialModelBase->IsA<UDynamicMaterialModel>();
+}
+
+void SDMMaterialWizard::OnAssetsActivated(TArrayView<const FContentBrowserItem> InSelectedItems, EAssetTypeActivationMethod::Type InActivationMethod)
+{
+	if (InSelectedItems.IsEmpty() || InActivationMethod == EAssetTypeActivationMethod::Previewed)
+	{
+		return;
+	}
+
+	FAssetData AssetData;
+
+	if (!InSelectedItems[0].Legacy_TryGetAssetData(AssetData))
+	{
+		return;
+	}
+
+	UDynamicMaterialInstance* Instance = Cast<UDynamicMaterialInstance>(AssetData.GetAsset());
+
+	if (!Instance)
+	{
+		return;
+	}
+
+	UDynamicMaterialModel* MaterialModel = Cast<UDynamicMaterialModel>(Instance->GetMaterialModelBase());
+
+	if (!MaterialModel)
+	{
+		return;
+	}
+
+	SelectTemplate(MaterialModel);
+}
+
+void SDMMaterialWizard::SelectTemplate(UDynamicMaterialModel* InTemplateModel)
+{
+	TSharedPtr<SDMMaterialDesigner> DesignerWidget = GetDesignerWidget();
+
+	if (!DesignerWidget.IsValid())
+	{
+		return;
 	}
 
 	if (UDynamicMaterialModel* MaterialModel = GetMaterialModel())
 	{
 		if (UDynamicMaterialInstance* Instance = MaterialModel->GetDynamicMaterialInstance())
 		{
-			CreateDynamicMaterialInInstance(TemplateModel, Instance);
+			CreateDynamicMaterialInInstance(InTemplateModel, Instance);
 		}
 		else
 		{
@@ -581,11 +686,11 @@ FReply SDMMaterialWizard::OnTemplateMouseDown(const FGeometry& InGeometry, const
 		{
 			if (UDynamicMaterialInstance* Instance = MaterialObjectProperty->GetMaterial())
 			{
-				CreateDynamicMaterialInInstance(TemplateModel, Instance);
+				CreateDynamicMaterialInInstance(InTemplateModel, Instance);
 			}
 			else
 			{
-				CreateNewDynamicInstanceInActor(TemplateModel, MaterialObjectProperty.GetValue());
+				CreateNewDynamicInstanceInActor(InTemplateModel, MaterialObjectProperty.GetValue());
 			}
 		}
 		else
@@ -601,8 +706,6 @@ FReply SDMMaterialWizard::OnTemplateMouseDown(const FGeometry& InGeometry, const
 
 		DesignerWidget->ShowSelectPrompt();
 	}
-
-	return FReply::Handled();
 }
 
 void SDMMaterialWizard::CreateDynamicMaterialInInstance(UDynamicMaterialModel* InTemplateModel, UDynamicMaterialInstance* InToInstance)
diff --git a/Engine/Plugins/Experimental/DynamicMaterial/Source/DynamicMaterialEditor/Private/UI/Widgets/SDMMaterialWizard.h b/Engine/Plugins/Experimental/DynamicMaterial/Source/DynamicMaterialEditor/Private/UI/Widgets/SDMMaterialWizard.h
--- a/Engine/Plugins/Experimental/DynamicMaterial/Source/DynamicMaterialEditor/Private/UI/Widgets/SDMMaterialWizard.h
+++ b/Engine/Plugins/Experimental/DynamicMaterial/Source/DynamicMaterialEditor/Private/UI/Widgets/SDMMaterialWizard.h
@@ -5,18 +5,24 @@
 #include "Widgets/SCompoundWidget.h"
 
 #include "DMObjectMaterialProperty.h"
+#include "IAssetTypeActions.h"
 #include "Misc/Optional.h"
 #include "Templates/SharedPointer.h"
 #include "UObject/NameTypes.h"
 #include "UObject/WeakObjectPtr.h"
 
+class FAssetTextFilter;
 class FAssetThumbnail;
+class SAssetSearchBox;
+class SAssetView;
 class SBox;
 class SDMMaterialDesigner;
+class SToolInputAssetPicker;
 class SWidget;
 class SWidgetSwitcher;
 class UDynamicMaterialModel;
 enum class ECheckBoxState : uint8;
+struct FContentBrowserItem;
 
 class SDMMaterialWizard : public SCompoundWidget
 {
@@ -40,17 +46,17 @@ public:
 	UDynamicMaterialModel* GetMaterialModel() const;
 
 protected:
-	static TArray<FAssetData> GetTemplateMaterials();
-
 	TWeakPtr<SDMMaterialDesigner> DesignerWidgetWeak;
 	FName CurrentPreset;
 	TSharedPtr<SBox> PresetChannelContainer;
 	TWeakObjectPtr<UDynamicMaterialModel> MaterialModelWeak;
 	TOptional<FDMObjectMaterialProperty> MaterialObjectProperty;
 	TSharedPtr<SWidgetSwitcher> Switcher;
-	TArray<TSharedRef<FAssetThumbnail>> Thumbnails;
-	TArray<FAssetData> Assets;
+	TSharedPtr<SAssetView> AssetView;
+	TSharedPtr<SAssetSearchBox> AssetSearchBox;
+	TSharedPtr<FAssetTextFilter> TextFilter;
 
+	/** Creation of widgets. */
 	TSharedRef<SWidget> CreateLayout();
 	TSharedRef<SWidget> CreateModeSelector();
 	TSharedRef<SWidget> CreateSelectPresetLayout();
@@ -58,8 +64,8 @@ protected:
 	TSharedRef<SWidget> CreateSelectPreset_ChannelList();
 	TSharedRef<SWidget> CreateSelectPreset_AcceptButton();
 	TSharedRef<SWidget> CreateTemplateListLayout();
-	TSharedRef<SWidget> CreateTemplateList_Entry(const FAssetData& InTemplateAsset);
 
+	/** Attributes and Events */
 	ECheckBoxState Preset_GetState(FName InPresetName) const;
 	void Preset_OnChange(ECheckBoxState InState, FName InPresetName);
 
@@ -73,7 +79,18 @@ protected:
 
 	void SetMode(ECheckBoxState InState, int32 InMode);
 
-	FReply OnTemplateMouseDown(const FGeometry& InGeometry, const FPointerEvent& InPointerEvent, int32 InAssetIndex);
+	void OnSearchBoxChanged(const FText& InSearchText);
+	void OnSearchBoxCommitted(const FText& InSearchText, ETextCommit::Type InCommitInfo);
+
+	FText GetSearchText() const;
+	void SetSearchText(const FText& InSearchText);
+
+	bool ShouldFilterOutAsset(const FAssetData& InAsset) const;
+
+	void OnAssetsActivated(TArrayView<const FContentBrowserItem> InSelectedItems, EAssetTypeActivationMethod::Type InActivationMethod);
+
+	/** Operations */
+	void SelectTemplate(UDynamicMaterialModel* InTemplateModel);
 
 	void CreateDynamicMaterialInInstance(UDynamicMaterialModel* InTemplateModel, UDynamicMaterialInstance* InToInstance);
 
