Material Designer
- When using "Save As" the newly created material is now opened.

diff --git a/Engine/Plugins/Experimental/DynamicMaterial/Source/DynamicMaterialEditor/Private/UI/Widgets/Editor/SDMToolBar.cpp b/Engine/Plugins/Experimental/DynamicMaterial/Source/DynamicMaterialEditor/Private/UI/Widgets/Editor/SDMToolBar.cpp
--- a/Engine/Plugins/Experimental/DynamicMaterial/Source/DynamicMaterialEditor/Private/UI/Widgets/Editor/SDMToolBar.cpp
+++ b/Engine/Plugins/Experimental/DynamicMaterial/Source/DynamicMaterialEditor/Private/UI/Widgets/Editor/SDMToolBar.cpp
@@ -531,11 +531,36 @@ FReply SDMToolBar::OnFollowSelectionButtonClicked()
 
 FReply SDMToolBar::OnExportMaterialInstanceButtonClicked()
 {
-	if (UDynamicMaterialModelBase* MaterialModelBase = GetMaterialModelBase())
+	TSharedPtr<SDMMaterialEditor> EditorWidget = GetEditorWidget();
+
+	if (!EditorWidget.IsValid())
+	{
+		return FReply::Handled();
+	}
+
+	TSharedPtr<SDMMaterialDesigner> DesignerWidget = EditorWidget->GetDesignerWidget();
+
+	if (!DesignerWidget.IsValid())
+	{
+		return FReply::Handled();
+	}
+
+	UDynamicMaterialModelBase* MaterialModelBase = GetMaterialModelBase();
+
+	if (!MaterialModelBase)
+	{
+		return FReply::Handled();
+	}
+
+	UDynamicMaterialInstance* NewInstance = UDMMaterialModelFunctionLibrary::ExportMaterialInstance(MaterialModelBase);
+
+	if (!NewInstance)
 	{
-		UDMMaterialModelFunctionLibrary::ExportMaterialInstance(MaterialModelBase);
+		return FReply::Handled();
 	}
 
+	DesignerWidget->OpenMaterialInstance(NewInstance);
+
 	return FReply::Handled();
 }
 
