Material Designer
- The reset to default button will now appear on the channels page.
- The channels page global slider are now full width.

diff --git a/Engine/Plugins/Experimental/DynamicMaterial/Source/DynamicMaterialEditor/Private/UI/Utils/DMWidgetStatics.cpp b/Engine/Plugins/Experimental/DynamicMaterial/Source/DynamicMaterialEditor/Private/UI/Utils/DMWidgetStatics.cpp
--- a/Engine/Plugins/Experimental/DynamicMaterial/Source/DynamicMaterialEditor/Private/UI/Utils/DMWidgetStatics.cpp
+++ b/Engine/Plugins/Experimental/DynamicMaterial/Source/DynamicMaterialEditor/Private/UI/Utils/DMWidgetStatics.cpp
@@ -8,6 +8,8 @@
 #include "Modules/ModuleManager.h"
 #include "PropertyEditorModule.h"
 
+const FLazyName FDMWidgetStatics::PropertyValueWidget = TEXT("SPropertyValueWidget");
+
 FDMWidgetStatics& FDMWidgetStatics::Get()
 {
 	 static FDMWidgetStatics Instance;
@@ -88,6 +90,56 @@ void FDMWidgetStatics::ClearPropertyHandles(const SWidget* InOwningWidget)
 	PropertyHandleMap.Remove(InOwningWidget);
 }
 
+TSharedPtr<SWidget> FDMWidgetStatics::FindWidgetInHierarchy(const TSharedRef<SWidget>& InParent, const FName& InName)
+{
+	if (InParent->GetType() == InName)
+	{
+		return InParent;
+	}
+
+	FChildren* Children = InParent->GetChildren();
+
+	if (!Children)
+	{
+		return nullptr;
+	}
+
+	const int32 ChildNum = Children->Num();
+
+	for (int32 Index = 0; Index < ChildNum; ++Index)
+	{
+		const TSharedRef<SWidget>& Widget = Children->GetChildAt(Index);
+
+		if (Widget->GetType() == InName)
+		{
+			return Widget;
+		}
+	}
+
+	for (int32 Index = 0; Index < ChildNum; ++Index)
+	{
+		if (TSharedPtr<SWidget> FoundChild = FindWidgetInHierarchy(Children->GetChildAt(Index), InName))
+		{
+			return FoundChild;
+		}
+	}
+
+	return nullptr;
+}
+
+TSharedPtr<SWidget> FDMWidgetStatics::GetInnerPropertyValueWidget(const TSharedRef<SWidget>& InWidget)
+{
+	if (FChildren* Children = InWidget->GetChildren())
+	{
+		if (Children->Num() > 0)
+		{
+			return Children->GetChildAt(0);
+		}
+	}
+
+	return nullptr;
+}
+
 FDMPropertyHandle FDMWidgetStatics::CreatePropertyHandle(const void* InOwningWidget, UObject* InObject, FName InPropertyName)
 {
 	FDMPropertyHandle PropertyHandle;
diff --git a/Engine/Plugins/Experimental/DynamicMaterial/Source/DynamicMaterialEditor/Private/UI/Utils/DMWidgetStatics.h b/Engine/Plugins/Experimental/DynamicMaterial/Source/DynamicMaterialEditor/Private/UI/Utils/DMWidgetStatics.h
--- a/Engine/Plugins/Experimental/DynamicMaterial/Source/DynamicMaterialEditor/Private/UI/Utils/DMWidgetStatics.h
+++ b/Engine/Plugins/Experimental/DynamicMaterial/Source/DynamicMaterialEditor/Private/UI/Utils/DMWidgetStatics.h
@@ -15,6 +15,8 @@ class UObject;
 class FDMWidgetStatics
 {
 public:
+	static const FLazyName PropertyValueWidget;
+
 	static FDMWidgetStatics& Get();
 
 	bool GetExpansionState(UObject* InOwner, FName InName, bool& bOutExpanded);
@@ -25,6 +27,10 @@ public:
 
 	void ClearPropertyHandles(const SWidget* InOwningWidget);
 
+	TSharedPtr<SWidget> FindWidgetInHierarchy(const TSharedRef<SWidget>& InParent, const FName& InName);
+
+	TSharedPtr<SWidget> GetInnerPropertyValueWidget(const TSharedRef<SWidget>& InWidget);
+
 private:
 	struct FExpansionItem
 	{
diff --git a/Engine/Plugins/Experimental/DynamicMaterial/Source/DynamicMaterialEditor/Private/UI/Widgets/Editor/SDMMaterialProperties.cpp b/Engine/Plugins/Experimental/DynamicMaterial/Source/DynamicMaterialEditor/Private/UI/Widgets/Editor/SDMMaterialProperties.cpp
--- a/Engine/Plugins/Experimental/DynamicMaterial/Source/DynamicMaterialEditor/Private/UI/Widgets/Editor/SDMMaterialProperties.cpp
+++ b/Engine/Plugins/Experimental/DynamicMaterial/Source/DynamicMaterialEditor/Private/UI/Widgets/Editor/SDMMaterialProperties.cpp
@@ -66,6 +66,11 @@ void SDMMaterialProperties::Validate()
 {
 	if (Content.HasBeenInvalidated())
 	{
+		GlobalItems.Empty();
+		PropertyPreviewContainers.Empty();
+		PropertyEmptyContainers.Empty();
+		PropertyPreviews.Empty();
+		SliderItems.Empty();
 		Content << CreateSlot_Content();
 	}
 }
@@ -329,7 +334,7 @@ TSharedRef<SWidget> SDMMaterialProperties::CreatePropertyRow(UDMMaterialProperty
 
 			+ SHorizontalBox::Slot()
 			.FillWidth(1.f)
-			.HAlign(EHorizontalAlignment::HAlign_Left)
+			.HAlign(EHorizontalAlignment::HAlign_Fill)
 			.VAlign(EVerticalAlignment::VAlign_Fill)
 			.Padding(5.f, 5.f, 0.f, 5.f)
 			[
@@ -625,9 +630,23 @@ TSharedRef<SWidget> SDMMaterialProperties::CreateGlobalSlider(UDMMaterialPropert
 
 	AlphaItem->MakeWidget(nullptr, SharedThis(this));
 
+	SliderItems.Add(AlphaItem);
+
 	TSharedPtr<SWidget> ValueWidget = AlphaItem->GetWidget(ECustomDetailsViewWidgetType::Value);
 	TSharedPtr<SWidget> ExtensionWidget = AlphaItem->GetWidget(ECustomDetailsViewWidgetType::Extensions);
 
+	if (ValueWidget.IsValid())
+	{
+		if (TSharedPtr<SWidget> FoundPropertyValueWidget = FDMWidgetStatics::Get().FindWidgetInHierarchy(ValueWidget.ToSharedRef(), FDMWidgetStatics::PropertyValueWidget))
+		{
+			if (TSharedPtr<SWidget> InnerPropertyWidget = FDMWidgetStatics::Get().GetInnerPropertyValueWidget(FoundPropertyValueWidget.ToSharedRef()))
+			{
+				ValueWidget = InnerPropertyWidget;
+			}
+		}
+	}
+
+
 	const FText PropertyGlobalSliderToolTipFormat = LOCTEXT("PropertyGlobalSliderToolTipFormat", "Change the global {0} value.");
 
 	const FText PropertyGlobalSliderToolTip = FText::Format(
@@ -639,11 +658,11 @@ TSharedRef<SWidget> SDMMaterialProperties::CreateGlobalSlider(UDMMaterialPropert
 		.ToolTipText(PropertyGlobalSliderToolTip)
 		
 		+ SHorizontalBox::Slot()
-		.AutoWidth()
+		.FillWidth(1.f)
 		[
 			SNew(SBox)
 			.HeightOverride(32.f)
-			.VAlign(VAlign_Center)
+			.HAlign(EHorizontalAlignment::HAlign_Fill)
 			.VAlign(EVerticalAlignment::VAlign_Center)
 			[
 				ValueWidget.IsValid() ? ValueWidget.ToSharedRef() : SNullWidget::NullWidget
@@ -651,12 +670,10 @@ TSharedRef<SWidget> SDMMaterialProperties::CreateGlobalSlider(UDMMaterialPropert
 		]
 
 		+ SHorizontalBox::Slot()
-		.FillWidth(1.f)
+		.AutoWidth()
 		[
 			SNew(SBox)
 			.HeightOverride(32.f)
-			.HAlign(HAlign_Left)
-			.VAlign(VAlign_Center)
 			.VAlign(EVerticalAlignment::VAlign_Center)
 			[
 				ExtensionWidget.IsValid() ? ExtensionWidget.ToSharedRef() : SNullWidget::NullWidget
diff --git a/Engine/Plugins/Experimental/DynamicMaterial/Source/DynamicMaterialEditor/Private/UI/Widgets/Editor/SDMMaterialProperties.h b/Engine/Plugins/Experimental/DynamicMaterial/Source/DynamicMaterialEditor/Private/UI/Widgets/Editor/SDMMaterialProperties.h
--- a/Engine/Plugins/Experimental/DynamicMaterial/Source/DynamicMaterialEditor/Private/UI/Widgets/Editor/SDMMaterialProperties.h
+++ b/Engine/Plugins/Experimental/DynamicMaterial/Source/DynamicMaterialEditor/Private/UI/Widgets/Editor/SDMMaterialProperties.h
@@ -42,6 +42,7 @@ protected:
 	TArray<TSharedRef<SBox>> PropertyPreviewContainers;
 	TArray<TSharedRef<SBox>> PropertyEmptyContainers;
 	TArray<TSharedRef<SDMMaterialComponentPreview>> PropertyPreviews;
+	TArray<TSharedRef<ICustomDetailsViewItem>> SliderItems;
 
 	bool bConstructing = false;
 
diff --git a/Engine/Plugins/Experimental/DynamicMaterial/Source/DynamicMaterialEditor/Private/UI/Widgets/Editor/SDMMaterialSlotEditor.cpp b/Engine/Plugins/Experimental/DynamicMaterial/Source/DynamicMaterialEditor/Private/UI/Widgets/Editor/SDMMaterialSlotEditor.cpp
--- a/Engine/Plugins/Experimental/DynamicMaterial/Source/DynamicMaterialEditor/Private/UI/Widgets/Editor/SDMMaterialSlotEditor.cpp
+++ b/Engine/Plugins/Experimental/DynamicMaterial/Source/DynamicMaterialEditor/Private/UI/Widgets/Editor/SDMMaterialSlotEditor.cpp
@@ -53,48 +53,6 @@
 
 #define LOCTEXT_NAMESPACE "SDMMaterialSlotEditor"
 
-namespace UE::DynamicMaterialEditor::Private
-{
-	const FLazyName PropertyValueWidget = TEXT("SPropertyValueWidget");
-
-	TSharedPtr<SWidget> FindWidgetInHierarchy(const TSharedRef<SWidget>& InParent, const FName& InName)
-	{
-		if (InParent->GetType() == InName)
-		{
-			return InParent;
-		}
-
-		FChildren* Children = InParent->GetChildren();
-
-		if (!Children)
-		{
-			return nullptr;
-		}
-
-		const int32 ChildNum = Children->Num();
-
-		for (int32 Index = 0; Index < ChildNum; ++Index)
-		{
-			const TSharedRef<SWidget>& Widget = Children->GetChildAt(Index);
-
-			if (Widget->GetType() == InName)
-			{
-				return Widget;
-			}
-		}
-
-		for (int32 Index = 0; Index < ChildNum; ++Index)
-		{
-			if (TSharedPtr<SWidget> FoundChild = FindWidgetInHierarchy(Children->GetChildAt(Index), InName))
-			{
-				return FoundChild;
-			}
-		}
-
-		return nullptr;
-	}
-}
-
 void SDMMaterialSlotEditor::PrivateRegisterAttributes(FSlateAttributeDescriptor::FInitializer&)
 {
 }
@@ -647,17 +605,11 @@ TSharedRef<SWidget> SDMMaterialSlotEditor::CreateSlot_LayerOpacity()
 
 	if (ValueWidget.IsValid())
 	{
-		if (TSharedPtr<SWidget> FoundPropertyValueWidget = FindWidgetInHierarchy(ValueWidget.ToSharedRef(), PropertyValueWidget))
+		if (TSharedPtr<SWidget> FoundPropertyValueWidget = FDMWidgetStatics::Get().FindWidgetInHierarchy(ValueWidget.ToSharedRef(), FDMWidgetStatics::PropertyValueWidget))
 		{
-			if (FoundPropertyValueWidget.IsValid())
+			if (TSharedPtr<SWidget> InnerPropertyWidget = FDMWidgetStatics::Get().GetInnerPropertyValueWidget(FoundPropertyValueWidget.ToSharedRef()))
 			{
-				if (FChildren* Children = FoundPropertyValueWidget->GetChildren())
-				{
-					if (Children->Num() > 0)
-					{
-						ValueWidget = Children->GetChildAt(0);
-					}
-				}
+				ValueWidget = InnerPropertyWidget;
 			}
 		}
 	}
