Material Designer
- Removed the global apha values from the global settings.

diff --git a/Engine/Plugins/Experimental/DynamicMaterial/Source/DynamicMaterialEditor/Private/UI/PropertyGenerators/DMMaterialModelPropertyRowGenerator.cpp b/Engine/Plugins/Experimental/DynamicMaterial/Source/DynamicMaterialEditor/Private/UI/PropertyGenerators/DMMaterialModelPropertyRowGenerator.cpp
--- a/Engine/Plugins/Experimental/DynamicMaterial/Source/DynamicMaterialEditor/Private/UI/PropertyGenerators/DMMaterialModelPropertyRowGenerator.cpp
+++ b/Engine/Plugins/Experimental/DynamicMaterial/Source/DynamicMaterialEditor/Private/UI/PropertyGenerators/DMMaterialModelPropertyRowGenerator.cpp
@@ -37,43 +37,6 @@ void FDMMaterialModelPropertyRowGenerator::AddMaterialModelProperties(const TSha
 
 	if (UDynamicMaterialModelEditorOnlyData* EditorOnlyData = UDynamicMaterialModelEditorOnlyData::Get(InMaterialModelBase))
 	{
-		if (EditorOnlyData->GetBlendMode() != BLEND_Opaque)
-		{
-			UDMMaterialProperty* OpacityProperty = EditorOnlyData->GetMaterialProperty(EDMMaterialPropertyType::Opacity);
-			UDMMaterialProperty* OpacityMaskProperty = EditorOnlyData->GetMaterialProperty(EDMMaterialPropertyType::OpacityMask);
-
-			const bool bOpacityValid = OpacityProperty && OpacityProperty->IsEnabled() && OpacityProperty->IsValidForModel(*EditorOnlyData)
-				&& EditorOnlyData->GetSlotForMaterialProperty(EDMMaterialPropertyType::Opacity);
-
-			const bool bOpacityMaskValid = OpacityMaskProperty && OpacityMaskProperty->IsEnabled() && OpacityMaskProperty->IsValidForModel(*EditorOnlyData)
-				&& EditorOnlyData->GetSlotForMaterialProperty(EDMMaterialPropertyType::OpacityMask);
-
-			if (bOpacityValid || bOpacityMaskValid)
-			{
-				AddGlobalValue(InGlobalSettingEditorWidget, InMaterialModelBase, InOutPropertyRows,
-					MaterialModel->GetGlobalParameterValue(UDynamicMaterialModel::GlobalOpacityValueName),
-					LOCTEXT("GlobalOpacity", "Global Opacity"));
-			}
-		}
-
-		UE::DynamicMaterial::ForEachMaterialPropertyType(
-			[&InGlobalSettingEditorWidget, InMaterialModelBase, &InOutPropertyRows, EditorOnlyData]
-			(EDMMaterialPropertyType InProperty)
-			{
-				if (InProperty != EDMMaterialPropertyType::Opacity && InProperty != EDMMaterialPropertyType::OpacityMask)
-				{
-					AddGlobalMaterialParameterValue(InProperty, InGlobalSettingEditorWidget, InMaterialModelBase, InOutPropertyRows, EditorOnlyData);
-				}
-
-				return EDMIterationResult::Continue;
-			}
-		);
-	}
-
-	if (UDynamicMaterialModelEditorOnlyData* EditorOnlyData = UDynamicMaterialModelEditorOnlyData::Get(InMaterialModelBase))
-	{
-		// TODO: Add selector for non-dynamic materials to change their preset. It's no longer on the EOD.
-
 		AddVariable(InGlobalSettingEditorWidget, InMaterialModelBase, InOutPropertyRows, EditorOnlyData,
 			GET_MEMBER_NAME_CHECKED(UDynamicMaterialModelEditorOnlyData, Domain));
 
@@ -100,34 +63,6 @@ void FDMMaterialModelPropertyRowGenerator::AddMaterialModelProperties(const TSha
 	}
 }
 
-void FDMMaterialModelPropertyRowGenerator::AddGlobalMaterialParameterValue(EDMMaterialPropertyType InProperty, 
-	const TSharedRef<SDMMaterialGlobalSettingsEditor>& InGlobalSettingEditorWidget, UDynamicMaterialModelBase* InMaterialModelBase, 
-	TArray<FDMPropertyHandle>& InOutPropertyRows, UDynamicMaterialModelEditorOnlyData* InEditorOnlyData)
-{
-	if (InEditorOnlyData->GetSlotForMaterialProperty(InProperty))
-	{
-		if (UDMMaterialProperty* MaterialProperty = InEditorOnlyData->GetMaterialProperty(InProperty))
-		{
-			if (MaterialProperty->IsEnabled() && MaterialProperty->IsValidForModel(*InEditorOnlyData))
-			{
-				if (UDMMaterialValueFloat1* AlphaValue = Cast<UDMMaterialValueFloat1>(MaterialProperty->GetComponent(UDynamicMaterialModelEditorOnlyData::AlphaValueName)))
-				{
-					AddGlobalValue(
-						InGlobalSettingEditorWidget,
-						InMaterialModelBase,
-						InOutPropertyRows,
-						AlphaValue,
-						FText::Format(
-							LOCTEXT("PropertyFormat", "Global {0}"),
-							UE::DynamicMaterialEditor::Private::GetMaterialPropertyLongDisplayName(InProperty)
-						)
-					);
-				}
-			}
-		}
-	}
-}
-
 void FDMMaterialModelPropertyRowGenerator::AddGlobalValue(const TSharedRef<SDMMaterialGlobalSettingsEditor>& InGlobalSettingEditorWidget,
 	UDynamicMaterialModelBase* InMaterialModelBase, TArray<FDMPropertyHandle>& InOutPropertyRows, UDMMaterialComponent* InComponent, 
 	const FText& InNameOverride)
diff --git a/Engine/Plugins/Experimental/DynamicMaterial/Source/DynamicMaterialEditor/Private/UI/PropertyGenerators/DMMaterialModelPropertyRowGenerator.h b/Engine/Plugins/Experimental/DynamicMaterial/Source/DynamicMaterialEditor/Private/UI/PropertyGenerators/DMMaterialModelPropertyRowGenerator.h
--- a/Engine/Plugins/Experimental/DynamicMaterial/Source/DynamicMaterialEditor/Private/UI/PropertyGenerators/DMMaterialModelPropertyRowGenerator.h
+++ b/Engine/Plugins/Experimental/DynamicMaterial/Source/DynamicMaterialEditor/Private/UI/PropertyGenerators/DMMaterialModelPropertyRowGenerator.h
@@ -27,10 +27,6 @@ public:
 	static void AddMaterialModelProperties(const TSharedRef<SDMMaterialGlobalSettingsEditor>& InGlobalSettingEditorWidget, 
 		UDynamicMaterialModelBase* InMaterialModelBase, TArray<FDMPropertyHandle>& InOutPropertyRows);
 
-	static void AddGlobalMaterialParameterValue(EDMMaterialPropertyType InProperty, const TSharedRef<SDMMaterialGlobalSettingsEditor>& InGlobalSettingEditorWidget,
-		UDynamicMaterialModelBase* InMaterialModelBase, TArray<FDMPropertyHandle>& InOutPropertyRows,
-		UDynamicMaterialModelEditorOnlyData* InEditorOnlyData);
-
 	static void AddGlobalValue(const TSharedRef<SDMMaterialGlobalSettingsEditor>& InGlobalSettingEditorWidget,
 		UDynamicMaterialModelBase* InMaterialModelBase, TArray<FDMPropertyHandle>& InOutPropertyRows, UDMMaterialComponent* InComponent,
 		const FText& InNameOverride);
