Material Designer
- Texture set UI texture names are now more visible.

diff --git a/Engine/Plugins/Experimental/DynamicMaterial/Source/DynamicMaterialTextureSetEditor/Private/Widgets/SDMTextureSetBuilderMaterialPropertyCell.cpp b/Engine/Plugins/Experimental/DynamicMaterial/Source/DynamicMaterialTextureSetEditor/Private/Widgets/SDMTextureSetBuilderMaterialPropertyCell.cpp
--- a/Engine/Plugins/Experimental/DynamicMaterial/Source/DynamicMaterialTextureSetEditor/Private/Widgets/SDMTextureSetBuilderMaterialPropertyCell.cpp
+++ b/Engine/Plugins/Experimental/DynamicMaterial/Source/DynamicMaterialTextureSetEditor/Private/Widgets/SDMTextureSetBuilderMaterialPropertyCell.cpp
@@ -13,6 +13,8 @@
 #include "PropertyEditorModule.h"
 #include "SAssetDropTarget.h"
 #include "SDMTextureSetBuilder.h"
+#include "Styling/StyleColors.h"
+#include "Widgets/Colors/SColorBlock.h"
 #include "Widgets/Images/SImage.h"
 #include "Widgets/SBoxPanel.h"
 #include "Widgets/SOverlay.h"
@@ -123,7 +125,10 @@ void SDMTextureSetBuilderMaterialPropertyCell::Construct(const FArguments& InArg
 						SNew(STextBlock)
 						.Text(this, &SDMTextureSetBuilderMaterialPropertyCell::GetTextureName)
 						.WrappingPolicy(ETextWrappingPolicy::AllowPerCharacterWrapping)
-						.WrapTextAt(90.f)
+						.WrapTextAt(110.f)
+						.Font(FAppStyle::GetFontStyle("TinyText"))
+						.HighlightText(this, &SDMTextureSetBuilderMaterialPropertyCell::GetTextureName)
+						.Visibility(this, &SDMTextureSetBuilderMaterialPropertyCell::GetTextureNameVisibility)
 					]
 				]
 			]
@@ -217,4 +222,11 @@ void SDMTextureSetBuilderMaterialPropertyCell::SetMaterialForChannelMask()
 	}
 }
 
+EVisibility SDMTextureSetBuilderMaterialPropertyCell::GetTextureNameVisibility() const
+{
+	return Texture.IsValid()
+		? EVisibility::Visible
+		: EVisibility::Collapsed;
+}
+
 #undef LOCTEXT_NAMESPACE
diff --git a/Engine/Plugins/Experimental/DynamicMaterial/Source/DynamicMaterialTextureSetEditor/Private/Widgets/SDMTextureSetBuilderMaterialPropertyCell.h b/Engine/Plugins/Experimental/DynamicMaterial/Source/DynamicMaterialTextureSetEditor/Private/Widgets/SDMTextureSetBuilderMaterialPropertyCell.h
--- a/Engine/Plugins/Experimental/DynamicMaterial/Source/DynamicMaterialTextureSetEditor/Private/Widgets/SDMTextureSetBuilderMaterialPropertyCell.h
+++ b/Engine/Plugins/Experimental/DynamicMaterial/Source/DynamicMaterialTextureSetEditor/Private/Widgets/SDMTextureSetBuilderMaterialPropertyCell.h
@@ -43,4 +43,6 @@ protected:
 	bool GetPropertyEnabled() const;
 
 	void SetMaterialForChannelMask();
+
+	EVisibility GetTextureNameVisibility() const;
 };
