Material Designer
- The texture set import warning has been shorterned.

diff --git a/Engine/Plugins/Experimental/DynamicMaterial/Source/DynamicMaterialEditor/Private/UI/Widgets/SDMMaterialEditor.cpp b/Engine/Plugins/Experimental/DynamicMaterial/Source/DynamicMaterialEditor/Private/UI/Widgets/SDMMaterialEditor.cpp
--- a/Engine/Plugins/Experimental/DynamicMaterial/Source/DynamicMaterialEditor/Private/UI/Widgets/SDMMaterialEditor.cpp
+++ b/Engine/Plugins/Experimental/DynamicMaterial/Source/DynamicMaterialEditor/Private/UI/Widgets/SDMMaterialEditor.cpp
@@ -892,11 +892,11 @@ void SDMMaterialEditor::HandleDrop_TextureSet(UDMTextureSet* InTextureSet)
 	const EAppReturnType::Type Result = FMessageDialog::Open(
 		EAppMsgType::YesNoCancel,
 		LOCTEXT("ReplaceSlotsTextureSet",
-			"You are about to import a Material Designer Texture Set.\n\n"
-			"Do you want to replace the slot contents?\n"
-			"- Yes: All layers are deleted in the matching slots.\n"
-			"- No: New texture layers are added to the matching slots.\n"
-			"- Cancel: Abort this operation.")
+			"Material Designer Texture Set.\n\n"
+			"Replace Slot?\n"
+			"- Yes: Delete Layers.\n"
+			"- No: Add Layers.\n"
+			"- Cancel")
 	);
 
 	FDMScopedUITransaction Transaction(LOCTEXT("DropTextureSet", "Drop Texture Set"));
