Material Designer
- Typo

diff --git a/Engine/Plugins/Experimental/DynamicMaterial/Source/DynamicMaterialEditor/Private/UI/Widgets/SDMMaterialEditor.cpp b/Engine/Plugins/Experimental/DynamicMaterial/Source/DynamicMaterialEditor/Private/UI/Widgets/SDMMaterialEditor.cpp
--- a/Engine/Plugins/Experimental/DynamicMaterial/Source/DynamicMaterialEditor/Private/UI/Widgets/SDMMaterialEditor.cpp
+++ b/Engine/Plugins/Experimental/DynamicMaterial/Source/DynamicMaterialEditor/Private/UI/Widgets/SDMMaterialEditor.cpp
@@ -893,7 +893,7 @@ void SDMMaterialEditor::HandleDrop_TextureSet(UDMTextureSet* InTextureSet)
 		EAppMsgType::YesNoCancel,
 		LOCTEXT("ReplaceSlotsTextureSet",
 			"Material Designer Texture Set.\n\n"
-			"Replace Slot?\n"
+			"Replace Slots?\n"
 			"- Yes: Delete Layers.\n"
 			"- No: Add Layers.\n"
 			"- Cancel")
