Material Designer
- Duplicating for PIE now instantly compiles materials instead of using the async system.

diff --git a/Engine/Plugins/Experimental/DynamicMaterial/Source/DynamicMaterial/Private/Material/DynamicMaterialInstance.cpp b/Engine/Plugins/Experimental/DynamicMaterial/Source/DynamicMaterial/Private/Material/DynamicMaterialInstance.cpp
--- a/Engine/Plugins/Experimental/DynamicMaterial/Source/DynamicMaterial/Private/Material/DynamicMaterialInstance.cpp
+++ b/Engine/Plugins/Experimental/DynamicMaterial/Source/DynamicMaterial/Private/Material/DynamicMaterialInstance.cpp
@@ -67,7 +67,7 @@ void UDynamicMaterialInstance::PostDuplicate(bool bInDuplicateForPIE)
 		{
 			if (IDynamicMaterialModelEditorOnlyDataInterface* ModelEditorOnlyData = MaterialModel->GetEditorOnlyData())
 			{
-				ModelEditorOnlyData->RequestMaterialBuild();
+				ModelEditorOnlyData->RequestMaterialBuild(/* Immediate */ bInDuplicateForPIE);
 			}
 		}
 	}
diff --git a/Engine/Plugins/Experimental/DynamicMaterial/Source/DynamicMaterial/Public/Model/IDynamicMaterialModelEditorOnlyDataInterface.h b/Engine/Plugins/Experimental/DynamicMaterial/Source/DynamicMaterial/Public/Model/IDynamicMaterialModelEditorOnlyDataInterface.h
--- a/Engine/Plugins/Experimental/DynamicMaterial/Source/DynamicMaterial/Public/Model/IDynamicMaterialModelEditorOnlyDataInterface.h
+++ b/Engine/Plugins/Experimental/DynamicMaterial/Source/DynamicMaterial/Public/Model/IDynamicMaterialModelEditorOnlyDataInterface.h
@@ -50,7 +50,7 @@ public:
 #if WITH_EDITOR
 	/** Called when the model needs to have the material rebuild. */
 	UFUNCTION(BlueprintCallable, Category = "Material Designer")
-	virtual void RequestMaterialBuild() PURE_VIRTUAL(UDynamicMaterialModelEditorOnlyDataInterface::RequestMaterialBuild)
+	virtual void RequestMaterialBuild(bool bInImmediate = false) PURE_VIRTUAL(UDynamicMaterialModelEditorOnlyDataInterface::RequestMaterialBuild)
 
 	/** Called to create the build state for building materials. */
 	virtual TSharedRef<IDMMaterialBuildStateInterface> CreateBuildStateInterface(UMaterial* InMaterialToBuild) const = 0;
diff --git a/Engine/Plugins/Experimental/DynamicMaterial/Source/DynamicMaterialEditor/Private/Model/DynamicMaterialModelEditorOnlyData.cpp b/Engine/Plugins/Experimental/DynamicMaterial/Source/DynamicMaterialEditor/Private/Model/DynamicMaterialModelEditorOnlyData.cpp
--- a/Engine/Plugins/Experimental/DynamicMaterial/Source/DynamicMaterialEditor/Private/Model/DynamicMaterialModelEditorOnlyData.cpp
+++ b/Engine/Plugins/Experimental/DynamicMaterial/Source/DynamicMaterialEditor/Private/Model/DynamicMaterialModelEditorOnlyData.cpp
@@ -571,11 +571,18 @@ void UDynamicMaterialModelEditorOnlyData::BuildMaterial(bool bInDirtyAssets)
 	OnMaterialBuiltDelegate.Broadcast(MaterialModel);
 }
 
-void UDynamicMaterialModelEditorOnlyData::RequestMaterialBuild()
+void UDynamicMaterialModelEditorOnlyData::RequestMaterialBuild(bool bInImmediate)
 {
 	if (!HasAnyFlags(RF_ClassDefaultObject | RF_ArchetypeObject))
 	{
-		FDynamicMaterialEditorModule::Get().AddBuildRequest(this, /* Dirty Packages */ !UE::GetIsEditorLoadingPackage());
+		if (bInImmediate)
+		{
+			BuildMaterial(/* Dirty Assets */ false);
+		}
+		else
+		{
+			FDynamicMaterialEditorModule::Get().AddBuildRequest(this, /* Dirty Packages */ !UE::GetIsEditorLoadingPackage());
+		}
 	}
 }
 
diff --git a/Engine/Plugins/Experimental/DynamicMaterial/Source/DynamicMaterialEditor/Public/Model/DynamicMaterialModelEditorOnlyData.h b/Engine/Plugins/Experimental/DynamicMaterial/Source/DynamicMaterialEditor/Public/Model/DynamicMaterialModelEditorOnlyData.h
--- a/Engine/Plugins/Experimental/DynamicMaterial/Source/DynamicMaterialEditor/Public/Model/DynamicMaterialModelEditorOnlyData.h
+++ b/Engine/Plugins/Experimental/DynamicMaterial/Source/DynamicMaterialEditor/Public/Model/DynamicMaterialModelEditorOnlyData.h
@@ -243,7 +243,7 @@ public:
 
 	//~ Begin IDynamicMaterialModelEditorOnlyDataInterface
 	DYNAMICMATERIALEDITOR_API virtual void PostEditorDuplicate() override;
-	DYNAMICMATERIALEDITOR_API virtual void RequestMaterialBuild() override;
+	DYNAMICMATERIALEDITOR_API virtual void RequestMaterialBuild(bool bInImmediate = false) override;
 	DYNAMICMATERIALEDITOR_API virtual void OnValueListUpdate() override;
 	DYNAMICMATERIALEDITOR_API virtual void OnValueUpdated(UDMMaterialValue* InValue, EDMUpdateType InUpdateType) override;
 	DYNAMICMATERIALEDITOR_API virtual void OnTextureUVUpdated(UDMTextureUV* InTextureUV) override;
