Material Designer
- Default displacement texture is not 0.5,0.5,0.5,0.5.
- Displacement now scales from 0.5 with the global alpha slider.

#virtualized

diff --git a/Engine/Plugins/Experimental/DynamicMaterial/Source/DynamicMaterialEditor/Private/Components/MaterialProperties/DMMPDisplacement.cpp b/Engine/Plugins/Experimental/DynamicMaterial/Source/DynamicMaterialEditor/Private/Components/MaterialProperties/DMMPDisplacement.cpp
--- a/Engine/Plugins/Experimental/DynamicMaterial/Source/DynamicMaterialEditor/Private/Components/MaterialProperties/DMMPDisplacement.cpp
+++ b/Engine/Plugins/Experimental/DynamicMaterial/Source/DynamicMaterialEditor/Private/Components/MaterialProperties/DMMPDisplacement.cpp
@@ -1,7 +1,26 @@
 // Copyright Epic Games, Inc. All Rights Reserved.
 
 #include "Components/MaterialProperties/DMMPDisplacement.h"
+
+#include "Components/MaterialValues/DMMaterialValueFloat1.h"
+#include "Materials/MaterialExpressionMaterialFunctionCall.h"
+#include "Model/DMMaterialBuildState.h"
 #include "Model/DynamicMaterialModelEditorOnlyData.h"
+#include "Utils/DMMaterialFunctionLibrary.h"
+#include "Utils/DMUtils.h"
+
+namespace UE::DynamicMaterialEditor::Private
+{
+	UMaterialFunctionInterface* GetDisplacementAlpha()
+	{
+		static UMaterialFunctionInterface* DisplacementAlpha = FDMMaterialFunctionLibrary::Get().GetFunction(
+			"DisplacementAlpha",
+			TEXT("/Script/Engine.MaterialFunction'/DynamicMaterial/MaterialFunctions/MF_DM_DisplacementAlpha.MF_DM_DisplacementAlpha'")
+		);
+
+		return DisplacementAlpha;
+	}
+}
 
 UDMMaterialPropertyDisplacement::UDMMaterialPropertyDisplacement()
 	: UDMMaterialProperty(
@@ -15,3 +34,58 @@ UMaterialExpression* UDMMaterialPropertyDisplacement::GetDefaultInput(
 {
 	return CreateConstant(InBuildState, FVector::ZeroVector);
 }
+
+void UDMMaterialPropertyDisplacement::AddAlphaMultiplier(const TSharedRef<FDMMaterialBuildState>& InBuildState) const
+{
+	UDMMaterialValueFloat1* AlphaValue = GetTypedComponent<UDMMaterialValueFloat1>(UDynamicMaterialModelEditorOnlyData::AlphaValueName);
+
+	if (!AlphaValue)
+	{
+		return;
+	}
+
+	FExpressionInput* PropertyInputExpression = InBuildState->GetMaterialProperty(MaterialProperty);
+
+	if (!PropertyInputExpression || !PropertyInputExpression->Expression)
+	{
+		return;
+	}
+
+	AlphaValue->GenerateExpression(InBuildState);
+
+	UMaterialExpression* GlobalOpacityExpression = InBuildState->GetLastValueExpression(AlphaValue);
+
+	if (!GlobalOpacityExpression)
+	{
+		return;
+	}
+
+	using namespace UE::DynamicMaterialEditor::Private;
+
+	UMaterialFunctionInterface* DisplacementAlphaFunction = GetDisplacementAlpha();
+
+	if (!DisplacementAlphaFunction)
+	{
+		return;
+	}
+
+	UMaterialExpressionMaterialFunctionCall* DisplacementAlphaExpression = InBuildState->GetBuildUtils().CreateExpression<UMaterialExpressionMaterialFunctionCall>(UE_DM_NodeComment_Default);
+	DisplacementAlphaExpression->SetMaterialFunction(DisplacementAlphaFunction);
+	DisplacementAlphaExpression->UpdateFromFunctionResource();
+
+	FExpressionInput& ValueInput = DisplacementAlphaExpression->FunctionInputs[0].Input;
+	ValueInput.Expression = PropertyInputExpression->Expression;
+	ValueInput.Mask = PropertyInputExpression->Mask;
+	ValueInput.MaskR = PropertyInputExpression->MaskR;
+	ValueInput.MaskG = PropertyInputExpression->MaskG;
+	ValueInput.MaskB = PropertyInputExpression->MaskB;
+	ValueInput.MaskA = PropertyInputExpression->MaskA;
+	ValueInput.OutputIndex = PropertyInputExpression->OutputIndex;
+
+	FExpressionInput& AlphaInput = DisplacementAlphaExpression->FunctionInputs[1].Input;
+	AlphaInput.Expression = GlobalOpacityExpression;
+	AlphaInput.SetMask(1, 1, 0, 0, 0);
+	AlphaInput.OutputIndex = 0;
+
+	PropertyInputExpression->Expression = DisplacementAlphaExpression;
+}
diff --git a/Engine/Plugins/Experimental/DynamicMaterial/Source/DynamicMaterialEditor/Private/Components/MaterialProperties/DMMPDisplacement.h b/Engine/Plugins/Experimental/DynamicMaterial/Source/DynamicMaterialEditor/Private/Components/MaterialProperties/DMMPDisplacement.h
--- a/Engine/Plugins/Experimental/DynamicMaterial/Source/DynamicMaterialEditor/Private/Components/MaterialProperties/DMMPDisplacement.h
+++ b/Engine/Plugins/Experimental/DynamicMaterial/Source/DynamicMaterialEditor/Private/Components/MaterialProperties/DMMPDisplacement.h
@@ -17,5 +17,6 @@ public:
 
 	//~ Begin UDMMaterialProperty
 	virtual UMaterialExpression* GetDefaultInput(const TSharedRef<FDMMaterialBuildState>& InBuildState) const override;
+	virtual void AddAlphaMultiplier(const TSharedRef<FDMMaterialBuildState>& InBuildState) const override;
 	//~ End UDMMaterialProperty
 };
