Material Designer
- Instead of exporting to an image file, the material snapshot now exports to a texture asset.

diff --git a/Engine/Plugins/Experimental/DynamicMaterial/Source/DynamicMaterialEditor/Private/UI/Menus/DMToolBarMenus.cpp b/Engine/Plugins/Experimental/DynamicMaterial/Source/DynamicMaterialEditor/Private/UI/Menus/DMToolBarMenus.cpp
--- a/Engine/Plugins/Experimental/DynamicMaterial/Source/DynamicMaterialEditor/Private/UI/Menus/DMToolBarMenus.cpp
+++ b/Engine/Plugins/Experimental/DynamicMaterial/Source/DynamicMaterialEditor/Private/UI/Menus/DMToolBarMenus.cpp
@@ -4,7 +4,6 @@
 
 #include "AssetToolsModule.h"
 #include "ContentBrowserModule.h"
-#include "DesktopPlatformModule.h"
 #include "DynamicMaterialEditorModule.h"
 #include "DynamicMaterialEditorSettings.h"
 #include "Engine/Engine.h"
@@ -12,7 +11,6 @@
 #include "EngineAnalytics.h"
 #include "Framework/Application/SlateApplication.h"
 #include "IContentBrowserSingleton.h"
-#include "IDesktopPlatform.h"
 #include "ISinglePropertyView.h"
 #include "Material/DynamicMaterialInstance.h"
 #include "Materials/Material.h"
@@ -409,27 +407,26 @@ void FDMToolBarMenus::SnapshotMaterial(TWeakObjectPtr<UDynamicMaterialModelBase>
 		return;
 	}
 
-	TArray<FString> OutFilenames;
+	// Choose asset location
+	IContentBrowserSingleton& ContentBrowser = FModuleManager::LoadModuleChecked<FContentBrowserModule>("ContentBrowser").Get();
+	const FContentBrowserItemPath CurrentPath = ContentBrowser.GetCurrentPath();
+	const FString PathStr = CurrentPath.HasInternalPath() ? CurrentPath.GetInternalPathString() : "/Game";
 
-	if (IDesktopPlatform* DesktopPlatform = FDesktopPlatformModule::Get())
-	{
-		DesktopPlatform->SaveFileDialog(
-			FSlateApplication::Get().FindBestParentWindowHandleForDialogs(nullptr),
-			LOCTEXT("SaveSnapshotAs", "Save Snapshot As").ToString(),
-			FPaths::ProjectSavedDir(),
-			Material->GetName() + "_Snapshot_" + FString::FromInt(InTextureSize.X) + "x" + FString::FromInt(InTextureSize.Y),
-			TEXT("HDR File (*.hdr)|*.hdr|EXR File (*.exr)|*.exr|PNG File (*.png)|*.png"),
-			EFileDialogFlags::None,
-			OutFilenames
-		);
-	}
+	FSaveAssetDialogConfig SaveAssetDialogConfig;
+	SaveAssetDialogConfig.DialogTitleOverride = LOCTEXT("SaveAssetDialogTitle", "Export Material To");
+	SaveAssetDialogConfig.DefaultPath = PathStr;
+	SaveAssetDialogConfig.ExistingAssetPolicy = ESaveAssetDialogExistingAssetPolicy::Disallow;
+	SaveAssetDialogConfig.DefaultAssetName = TEXT("T_MD_") + Material->GetName();
 
-	if (OutFilenames.IsEmpty())
+	FContentBrowserModule& ContentBrowserModule = FModuleManager::LoadModuleChecked<FContentBrowserModule>("ContentBrowser");
+	FString SaveObjectPath = ContentBrowserModule.Get().CreateModalSaveAssetDialog(SaveAssetDialogConfig);
+
+	if (SaveObjectPath.IsEmpty())
 	{
 		return;
 	}
 
-	FDMMaterialShapshotLibrary::SnapshotMaterial(Material, InTextureSize, OutFilenames[0]);
+	FDMMaterialShapshotLibrary::SnapshotMaterial(Material, InTextureSize, SaveObjectPath);
 
 	if (FEngineAnalytics::IsAvailable())
 	{
diff --git a/Engine/Plugins/Experimental/DynamicMaterial/Source/DynamicMaterialEditor/Private/Utils/DMMaterialSnapshotLibrary.cpp b/Engine/Plugins/Experimental/DynamicMaterial/Source/DynamicMaterialEditor/Private/Utils/DMMaterialSnapshotLibrary.cpp
--- a/Engine/Plugins/Experimental/DynamicMaterial/Source/DynamicMaterialEditor/Private/Utils/DMMaterialSnapshotLibrary.cpp
+++ b/Engine/Plugins/Experimental/DynamicMaterial/Source/DynamicMaterialEditor/Private/Utils/DMMaterialSnapshotLibrary.cpp
@@ -1,12 +1,14 @@
 // Copyright Epic Games, Inc. All Rights Reserved.
 
 #include "Utils/DMMaterialSnapshotLibrary.h"
+
+#include "AssetRegistry/AssetRegistryModule.h"
 #include "DMAlphaOneMinusPS.h"
 #include "DynamicMaterialEditorModule.h"
 #include "Engine/Canvas.h"
 #include "Engine/Engine.h"
+#include "Engine/Texture2D.h"
 #include "Framework/Notifications/NotificationManager.h"
-#include "ImageUtils.h"
 #include "Kismet/KismetRenderingLibrary.h"
 #include "Materials/MaterialInterface.h"
 #include "RenderGraphBuilder.h"
@@ -19,123 +21,123 @@
 
 #define LOCTEXT_NAMESPACE "DMMaterialShapshotLibrary"
 
-namespace UE::DynamicMaterial::Private
+void FDMMaterialShapshotLibrary::RenderMaterialToRenderTarget(UMaterialInterface* InMaterial, UTextureRenderTarget2D* InRenderTarget)
 {
-	void RenderMaterialToRenderTarget(UMaterialInterface* InMaterial, UTextureRenderTarget2D* InRenderTarget)
-	{
-		InMaterial->EnsureIsComplete();
+	InMaterial->EnsureIsComplete();
 
-		FTextureRenderTargetResource* RenderTargetResource = InRenderTarget->GameThread_GetRenderTargetResource();
+	FTextureRenderTargetResource* RenderTargetResource = InRenderTarget->GameThread_GetRenderTargetResource();
 
-		UCanvas* Canvas = NewObject<UCanvas>(GetTransientPackage());
+	UCanvas* Canvas = NewObject<UCanvas>(GetTransientPackage());
 
-		FCanvas RenderCanvas(
-			RenderTargetResource,
-			nullptr,
-			FGameTime::CreateUndilated(0, 0),
-			GEngine->GetDefaultWorldFeatureLevel()
-		);
+	FCanvas RenderCanvas(
+		RenderTargetResource,
+		nullptr,
+		FGameTime::CreateUndilated(0, 0),
+		GEngine->GetDefaultWorldFeatureLevel()
+	);
 
-		Canvas->Init(InRenderTarget->SizeX, InRenderTarget->SizeY, nullptr, &RenderCanvas);
+	Canvas->Init(InRenderTarget->SizeX, InRenderTarget->SizeY, nullptr, &RenderCanvas);
 
-		{
-			RHI_BREADCRUMB_EVENT_GAMETHREAD("DrawMaterialToRenderTarget: %s", InRenderTarget->GetFName());
+	{
+		RHI_BREADCRUMB_EVENT_GAMETHREAD("DrawMaterialToRenderTarget: %s", InRenderTarget->GetFName());
 
-			ENQUEUE_RENDER_COMMAND(FlushDeferredResourceUpdateCommand)(
-				[RenderTargetResource](FRHICommandListImmediate& RHICmdList)
-				{
-					RenderTargetResource->FlushDeferredResourceUpdate(RHICmdList);
-				});
+		ENQUEUE_RENDER_COMMAND(FlushDeferredResourceUpdateCommand)(
+			[RenderTargetResource](FRHICommandListImmediate& RHICmdList)
+			{
+				RenderTargetResource->FlushDeferredResourceUpdate(RHICmdList);
+			});
 
-			Canvas->K2_DrawMaterial(InMaterial, FVector2D(0, 0), FVector2D(InRenderTarget->SizeX, InRenderTarget->SizeY), FVector2D(0, 0));
+		Canvas->K2_DrawMaterial(InMaterial, FVector2D(0, 0), FVector2D(InRenderTarget->SizeX, InRenderTarget->SizeY), FVector2D(0, 0));
 
-			RenderCanvas.Flush_GameThread();
-			Canvas->Canvas = nullptr;
+		RenderCanvas.Flush_GameThread();
+		Canvas->Canvas = nullptr;
 
-			//UpdateResourceImmediate must be called here to ensure mips are generated.
-			InRenderTarget->UpdateResourceImmediate(false);
+		//UpdateResourceImmediate must be called here to ensure mips are generated.
+		InRenderTarget->UpdateResourceImmediate(false);
 
-			ENQUEUE_RENDER_COMMAND(ResetSceneTextureExtentHistory)(
-				[RenderTargetResource](FRHICommandListImmediate& RHICmdList)
-				{
-					RenderTargetResource->ResetSceneTextureExtentsHistory();
-				});
-		}
+		ENQUEUE_RENDER_COMMAND(ResetSceneTextureExtentHistory)(
+			[RenderTargetResource](FRHICommandListImmediate& RHICmdList)
+			{
+				RenderTargetResource->ResetSceneTextureExtentsHistory();
+			});
 	}
+}
 
-	UTextureRenderTarget2D* CreateSnapshotRenderTarget(const FIntPoint& InTextureSize)
-	{
-		UTextureRenderTarget2D* RenderTarget = NewObject<UTextureRenderTarget2D>(GetTransientPackage());
-		check(RenderTarget);
-		RenderTarget->RenderTargetFormat = RTF_RGBA32f;
-		RenderTarget->ClearColor = FLinearColor::Black;
-		RenderTarget->bAutoGenerateMips = false;
-		RenderTarget->bCanCreateUAV = false;
-		RenderTarget->InitAutoFormat(InTextureSize.X, InTextureSize.Y);
-		RenderTarget->UpdateResourceImmediate(true);
-
-		return RenderTarget;
-	}
+UTextureRenderTarget2D* FDMMaterialShapshotLibrary::CreateSnapshotRenderTarget(const FIntPoint& InTextureSize)
+{
+	UTextureRenderTarget2D* RenderTarget = NewObject<UTextureRenderTarget2D>(GetTransientPackage());
+	check(RenderTarget);
+	RenderTarget->RenderTargetFormat = RTF_RGBA32f;
+	RenderTarget->ClearColor = FLinearColor::Black;
+	RenderTarget->bAutoGenerateMips = false;
+	RenderTarget->bCanCreateUAV = false;
+	RenderTarget->InitAutoFormat(InTextureSize.X, InTextureSize.Y);
+	RenderTarget->UpdateResourceImmediate(true);
+
+	return RenderTarget;
+}
 
-	void ApplyAlphaOneMinusShader(FRHICommandListImmediate& InRHICmdList, FTextureRenderTargetResource* InSourceTextureResource, FTextureRenderTargetResource* InDestTargetResource)
-	{
-		FRDGBuilder GraphBuilder(InRHICmdList);
+void FDMMaterialShapshotLibrary::ApplyAlphaOneMinusShader(FRHICommandListImmediate& InRHICmdList, FTextureRenderTargetResource* InSourceTextureResource, FTextureRenderTargetResource* InDestTargetResource)
+{
+	FRDGBuilder GraphBuilder(InRHICmdList);
 
-		const FTextureRHIRef& SourceTexture = InSourceTextureResource->GetRenderTargetTexture();
-		const FTextureRHIRef& DestTarget = InDestTargetResource->GetRenderTargetTexture();
+	const FTextureRHIRef& SourceTexture = InSourceTextureResource->GetRenderTargetTexture();
+	const FTextureRHIRef& DestTarget = InDestTargetResource->GetRenderTargetTexture();
 
-		FRDGTexture* InputTexture = GraphBuilder.RegisterExternalTexture(CreateRenderTarget(SourceTexture, TEXT("SourceTexture")));
-		FRDGTexture* OutputTexture = GraphBuilder.RegisterExternalTexture(CreateRenderTarget(DestTarget, TEXT("DestTarget")));
-		const FGlobalShaderMap* GlobalShaderMap = GetGlobalShaderMap(GMaxRHIFeatureLevel);
+	FRDGTexture* InputTexture = GraphBuilder.RegisterExternalTexture(CreateRenderTarget(SourceTexture, TEXT("SourceTexture")));
+	FRDGTexture* OutputTexture = GraphBuilder.RegisterExternalTexture(CreateRenderTarget(DestTarget, TEXT("DestTarget")));
+	const FGlobalShaderMap* GlobalShaderMap = GetGlobalShaderMap(GMaxRHIFeatureLevel);
 
-		// The formats or sizes differ to pixel shader stuff
-		// Configure source/output viewport to get the right UV scaling from source texture to output texture
-		const FScreenPassTextureViewport InputViewport(InputTexture);
-		const FScreenPassTextureViewport OutputViewport(OutputTexture);
+	// The formats or sizes differ to pixel shader stuff
+	// Configure source/output viewport to get the right UV scaling from source texture to output texture
+	const FScreenPassTextureViewport InputViewport(InputTexture);
+	const FScreenPassTextureViewport OutputViewport(OutputTexture);
 
-		const TShaderMapRef<FScreenPassVS> VertexShader(GlobalShaderMap);
+	const TShaderMapRef<FScreenPassVS> VertexShader(GlobalShaderMap);
 
-		// Rectangle area to use from source
-		const FIntRect ViewRect(FIntPoint(0, 0), InputTexture->Desc.Extent);
+	// Rectangle area to use from source
+	const FIntRect ViewRect(FIntPoint(0, 0), InputTexture->Desc.Extent);
 
-		//Dummy ViewFamily/ViewInfo created to use built in Draw Screen/Texture Pass
-		const FSceneViewFamily ViewFamily(FSceneViewFamily::ConstructionValues(nullptr, nullptr, FEngineShowFlags(ESFIM_Game))
-			.SetTime(FGameTime()));
-		FSceneViewInitOptions ViewInitOptions;
-		ViewInitOptions.ViewFamily = &ViewFamily;
-		ViewInitOptions.SetViewRectangle(ViewRect);
-		ViewInitOptions.ViewOrigin = FVector::ZeroVector;
-		ViewInitOptions.ViewRotationMatrix = FMatrix::Identity;
-		ViewInitOptions.ProjectionMatrix = FMatrix::Identity;
-		const FSceneView ViewInfo(ViewInitOptions);
+	//Dummy ViewFamily/ViewInfo created to use built in Draw Screen/Texture Pass
+	const FSceneViewFamily ViewFamily(FSceneViewFamily::ConstructionValues(nullptr, nullptr, FEngineShowFlags(ESFIM_Game))
+		.SetTime(FGameTime()));
+	FSceneViewInitOptions ViewInitOptions;
+	ViewInitOptions.ViewFamily = &ViewFamily;
+	ViewInitOptions.SetViewRectangle(ViewRect);
+	ViewInitOptions.ViewOrigin = FVector::ZeroVector;
+	ViewInitOptions.ViewRotationMatrix = FMatrix::Identity;
+	ViewInitOptions.ProjectionMatrix = FMatrix::Identity;
+	const FSceneView ViewInfo(ViewInitOptions);
 
-		const TShaderMapRef<FDMAlphaOneMinusPS> PixelShader(GlobalShaderMap);
-		FDMAlphaOneMinusPS::FParameters* Parameters = PixelShader->AllocateAndSetParameters(GraphBuilder, InputTexture, OutputTexture);
-		AddDrawScreenPass(GraphBuilder, RDG_EVENT_NAME("ApplyAlphaOneMinusShader"), ViewInfo, OutputViewport, InputViewport, VertexShader, PixelShader, Parameters);
+	const TShaderMapRef<FDMAlphaOneMinusPS> PixelShader(GlobalShaderMap);
+	FDMAlphaOneMinusPS::FParameters* Parameters = PixelShader->AllocateAndSetParameters(GraphBuilder, InputTexture, OutputTexture);
+	AddDrawScreenPass(GraphBuilder, RDG_EVENT_NAME("ApplyAlphaOneMinusShader"), ViewInfo, OutputViewport, InputViewport, VertexShader, PixelShader, Parameters);
 
-		GraphBuilder.Execute();
-	}
+	GraphBuilder.Execute();
+}
 
-	UTextureRenderTarget2D* ApplyAlphaOneMinusShader(UTextureRenderTarget2D* InRenderTarget)
-	{
-		FTextureRenderTargetResource* Source = InRenderTarget->GameThread_GetRenderTargetResource();
+UTextureRenderTarget2D* FDMMaterialShapshotLibrary::ApplyAlphaOneMinusShader(UTextureRenderTarget2D* InRenderTarget)
+{
+	FTextureRenderTargetResource* Source = InRenderTarget->GameThread_GetRenderTargetResource();
 
-		UTextureRenderTarget2D* OutRenderTarget = CreateSnapshotRenderTarget({InRenderTarget->SizeX, InRenderTarget->SizeY});
-		FTextureRenderTargetResource* Target = OutRenderTarget->GameThread_GetRenderTargetResource();
+	UTextureRenderTarget2D* OutRenderTarget = CreateSnapshotRenderTarget({InRenderTarget->SizeX, InRenderTarget->SizeY});
+	FTextureRenderTargetResource* Target = OutRenderTarget->GameThread_GetRenderTargetResource();
 
-		ENQUEUE_RENDER_COMMAND(ApplyAlphaOneMinusCommand)(
-			[Source, Target](FRHICommandListImmediate& InRHICmdList)
-			{
-				ApplyAlphaOneMinusShader(InRHICmdList, Source, Target);
-			});
+	ENQUEUE_RENDER_COMMAND(ApplyAlphaOneMinusCommand)(
+		[Source, Target](FRHICommandListImmediate& InRHICmdList)
+		{
+			ApplyAlphaOneMinusShader(InRHICmdList, Source, Target);
+		});
 
-		return OutRenderTarget;
-	}
+	return OutRenderTarget;
 }
 
 bool FDMMaterialShapshotLibrary::SnapshotMaterial(UMaterialInterface* InMaterial, const FIntPoint& InTextureSize, const FString& InSavePath)
 {
-	UTextureRenderTarget2D* RenderTarget = UE::DynamicMaterial::Private::CreateSnapshotRenderTarget(InTextureSize);
+	const FString PackageName = FPaths::GetBaseFilename(*InSavePath, false);
+	const FString AssetName = FPaths::GetBaseFilename(*InSavePath, true);
+
+	UTextureRenderTarget2D* RenderTarget = CreateSnapshotRenderTarget(InTextureSize);
 
 	if (!RenderTarget)
 	{
@@ -143,19 +145,36 @@ bool FDMMaterialShapshotLibrary::SnapshotMaterial(UMaterialInterface* InMaterial
 		return false;
 	}
 
-	UE::DynamicMaterial::Private::RenderMaterialToRenderTarget(InMaterial, RenderTarget);
+	RenderMaterialToRenderTarget(InMaterial, RenderTarget);
 
 	// Render target has inverted alpha - fix that.
-	UTextureRenderTarget2D* FixedAlphaRenderTarget = UE::DynamicMaterial::Private::ApplyAlphaOneMinusShader(RenderTarget);
+	UTextureRenderTarget2D* FixedAlphaRenderTarget = ApplyAlphaOneMinusShader(RenderTarget);
 
-	FImage Image;
-	FImageUtils::GetRenderTargetImage(FixedAlphaRenderTarget, Image);
+	//FixedAlphaRenderTarget->ConstructTexture2D
+	UTexture2D* Texture = FixedAlphaRenderTarget->ConstructTexture2D(GetTransientPackage(), AssetName, RF_Public | RF_Standalone);
 
-	FImageUtils::SaveImageByExtension(*InSavePath, Image);
+	if (!Texture)
+	{
+		UE_LOG(LogDynamicMaterialEditor, Warning, TEXT("Unable to texture for material snapshot."));
+		return false;
+	}
 
 	RenderTarget->ReleaseResource();
 	FixedAlphaRenderTarget->ReleaseResource();
 
+	// Create package and save texture.
+	UPackage* Package = CreatePackage(*PackageName);
+
+	if (!Package)
+	{
+		UE_LOG(LogDynamicMaterialEditor, Warning, TEXT("Unable to create package for material snapshot."));
+		return false;
+	}
+
+	Texture->Rename(*AssetName, Package, UE::DynamicMaterial::RenameFlags);
+
+	FAssetRegistryModule::AssetCreated(Package);
+
 	// Make a toast!
 	const FString HyperLinkText = FPaths::ConvertRelativePathToFull(InSavePath);
 
diff --git a/Engine/Plugins/Experimental/DynamicMaterial/Source/DynamicMaterialEditor/Private/Utils/DMMaterialSnapshotLibrary.h b/Engine/Plugins/Experimental/DynamicMaterial/Source/DynamicMaterialEditor/Private/Utils/DMMaterialSnapshotLibrary.h
--- a/Engine/Plugins/Experimental/DynamicMaterial/Source/DynamicMaterialEditor/Private/Utils/DMMaterialSnapshotLibrary.h
+++ b/Engine/Plugins/Experimental/DynamicMaterial/Source/DynamicMaterialEditor/Private/Utils/DMMaterialSnapshotLibrary.h
@@ -6,9 +6,22 @@
 #include "Math/MathFwd.h"
 
 class UMaterialInterface;
+class UTextureRenderTarget2D;
+class FRHICommandListImmediate;
+class FTextureRenderTargetResource;
 
 class FDMMaterialShapshotLibrary
 {
 public:
 	static bool SnapshotMaterial(UMaterialInterface* InMaterial, const FIntPoint& InTextureSize, const FString& InSavePath);
+
+private:
+	static void RenderMaterialToRenderTarget(UMaterialInterface* InMaterial, UTextureRenderTarget2D* InRenderTarget);
+
+	static UTextureRenderTarget2D* CreateSnapshotRenderTarget(const FIntPoint& InTextureSize);
+
+	static void ApplyAlphaOneMinusShader(FRHICommandListImmediate& InRHICmdList, FTextureRenderTargetResource* InSourceTextureResource,
+		FTextureRenderTargetResource* InDestTargetResource);
+
+	static UTextureRenderTarget2D* ApplyAlphaOneMinusShader(UTextureRenderTarget2D* InRenderTarget);
 };
