Material Designer
- Fixed the wizard being able to base a new material off itself.
- Fixed the wizard beinga ble to base a new material off dynamic material deigner assets.
- Fixed a crash when the editor was closed with the wizard open.

diff --git a/Engine/Plugins/Experimental/DynamicMaterial/Source/DynamicMaterialEditor/Private/UI/Widgets/SDMMaterialWizard.cpp b/Engine/Plugins/Experimental/DynamicMaterial/Source/DynamicMaterialEditor/Private/UI/Widgets/SDMMaterialWizard.cpp
--- a/Engine/Plugins/Experimental/DynamicMaterial/Source/DynamicMaterialEditor/Private/UI/Widgets/SDMMaterialWizard.cpp
+++ b/Engine/Plugins/Experimental/DynamicMaterial/Source/DynamicMaterialEditor/Private/UI/Widgets/SDMMaterialWizard.cpp
@@ -50,6 +50,8 @@ void SDMMaterialWizard::PrivateRegisterAttributes(FSlateAttributeDescriptor::FIn
 
 SDMMaterialWizard::~SDMMaterialWizard()
 {
+	FCoreDelegates::OnEnginePreExit.RemoveAll(this);
+
 	if (!FDynamicMaterialModule::AreUObjectsSafe())
 	{
 		return;
@@ -72,6 +74,8 @@ void SDMMaterialWizard::Construct(const FArguments& InArgs, const TSharedRef<SDM
 
 	SetCanTick(false);
 
+	FCoreDelegates::OnEnginePreExit.AddSP(this, &SDMMaterialWizard::OnEnginePreExit);
+
 	if (MaterialObjectProperty.IsSet())
 	{
 		if (UDynamicMaterialModelBase* MaterialModelBase = MaterialObjectProperty.GetValue().GetMaterialModelBase())
@@ -610,21 +614,49 @@ void SDMMaterialWizard::SetSearchText(const FText& InSearchText)
 
 bool SDMMaterialWizard::ShouldFilterOutAsset(const FAssetData& InAsset) const
 {
-	UDynamicMaterialInstance* Instance = Cast<UDynamicMaterialInstance>(InAsset.GetAsset());
+	UDynamicMaterialInstance* MaterialInstance = Cast<UDynamicMaterialInstance>(InAsset.GetAsset());
 
-	if (!Instance)
+	if (!MaterialInstance)
 	{
 		return true;
 	}
 
-	UDynamicMaterialModelBase* MaterialModelBase = Instance->GetMaterialModelBase();
+	UDynamicMaterialModelBase* AssetMaterialModelBase = MaterialInstance->GetMaterialModelBase();
+
+	if (!AssetMaterialModelBase)
+	{
+		return true;
+	}
 
-	if (!MaterialModelBase)
+	// Only non-dynamic models can be used as a basis.
+	if (!AssetMaterialModelBase->IsA<UDynamicMaterialModel>())
 	{
 		return true;
 	}
 
-	return !MaterialModelBase->IsA<UDynamicMaterialModel>();
+	if (UDynamicMaterialModel* MaterialModel = MaterialModelWeak.Get())
+	{
+		// Can't use it off ourselves
+		if (AssetMaterialModelBase == MaterialModel)
+		{
+			return true;
+		}
+	}
+
+	UDynamicMaterialModelEditorOnlyData* EditorOnlyData = UDynamicMaterialModelEditorOnlyData::Get(AssetMaterialModelBase);
+
+	if (!EditorOnlyData)
+	{
+		return true;
+	}
+
+	// Can't base it off things which also need a wizard.
+	if (EditorOnlyData->NeedsWizard())
+	{
+		return true;
+	}
+
+	return false;
 }
 
 void SDMMaterialWizard::OnAssetsActivated(TArrayView<const FContentBrowserItem> InSelectedItems, EAssetTypeActivationMethod::Type InActivationMethod)
@@ -658,6 +690,16 @@ void SDMMaterialWizard::OnAssetsActivated(TArrayView<const FContentBrowserItem>
 	SelectTemplate(MaterialModel);
 }
 
+void SDMMaterialWizard::OnEnginePreExit()
+{
+	TextFilter.Reset();
+	AssetSearchBox.Reset();
+	AssetView.Reset();
+	Switcher.Reset();
+
+	ChildSlot.DetachWidget();
+}
+
 void SDMMaterialWizard::SelectTemplate(UDynamicMaterialModel* InTemplateModel)
 {
 	TSharedPtr<SDMMaterialDesigner> DesignerWidget = GetDesignerWidget();
diff --git a/Engine/Plugins/Experimental/DynamicMaterial/Source/DynamicMaterialEditor/Private/UI/Widgets/SDMMaterialWizard.h b/Engine/Plugins/Experimental/DynamicMaterial/Source/DynamicMaterialEditor/Private/UI/Widgets/SDMMaterialWizard.h
--- a/Engine/Plugins/Experimental/DynamicMaterial/Source/DynamicMaterialEditor/Private/UI/Widgets/SDMMaterialWizard.h
+++ b/Engine/Plugins/Experimental/DynamicMaterial/Source/DynamicMaterialEditor/Private/UI/Widgets/SDMMaterialWizard.h
@@ -89,6 +89,8 @@ protected:
 
 	void OnAssetsActivated(TArrayView<const FContentBrowserItem> InSelectedItems, EAssetTypeActivationMethod::Type InActivationMethod);
 
+	void OnEnginePreExit();
+
 	/** Operations */
 	void SelectTemplate(UDynamicMaterialModel* InTemplateModel);
 
