Material Designer
- Added custom thumbnails to Dynamic Material Instances.
- Added custom thumbnail scene and renderer.
- Added ability to customise the default DMI mesh and orientation.

diff --git a/Engine/Plugins/Experimental/DynamicMaterial/Source/DynamicMaterialEditor/Private/DynamicMaterialEditorCommands.cpp b/Engine/Plugins/Experimental/DynamicMaterial/Source/DynamicMaterialEditor/Private/DynamicMaterialEditorCommands.cpp
--- a/Engine/Plugins/Experimental/DynamicMaterial/Source/DynamicMaterialEditor/Private/DynamicMaterialEditorCommands.cpp
+++ b/Engine/Plugins/Experimental/DynamicMaterial/Source/DynamicMaterialEditor/Private/DynamicMaterialEditorCommands.cpp
@@ -58,9 +58,9 @@ void FDynamicMaterialEditorCommands::RegisterCommands()
 		, EUserInterfaceActionType::Button
 		, FInputChord(EKeys::RightBracket, EModifierKey::Alt));
 
-	UI_COMMAND(SetShaderBallPreview
-		, "Shader Ball"
-		, "Sets the preview mesh to the Shader Ball."
+	UI_COMMAND(SetCustomPreviewMesh
+		, "Custom Mesh"
+		, "Sets the preview mesh to a custom mesh specified in the settings."
 		, EUserInterfaceActionType::ToggleButton
 		, FInputChord());
 
diff --git a/Engine/Plugins/Experimental/DynamicMaterial/Source/DynamicMaterialEditor/Private/DynamicMaterialEditorCommands.h b/Engine/Plugins/Experimental/DynamicMaterial/Source/DynamicMaterialEditor/Private/DynamicMaterialEditorCommands.h
--- a/Engine/Plugins/Experimental/DynamicMaterial/Source/DynamicMaterialEditor/Private/DynamicMaterialEditorCommands.h
+++ b/Engine/Plugins/Experimental/DynamicMaterial/Source/DynamicMaterialEditor/Private/DynamicMaterialEditorCommands.h
@@ -21,7 +21,7 @@ public:
 	TSharedPtr<FUICommandInfo> MoveLayerUp;
 	TSharedPtr<FUICommandInfo> MoveLayerDown;
 
-	TSharedPtr<FUICommandInfo> SetShaderBallPreview;
+	TSharedPtr<FUICommandInfo> SetCustomPreviewMesh;
 
 	TSharedPtr<FUICommandInfo> NavigateForward;
 	TSharedPtr<FUICommandInfo> NavigateBack;
diff --git a/Engine/Plugins/Experimental/DynamicMaterial/Source/DynamicMaterialEditor/Private/DynamicMaterialEditorModule.cpp b/Engine/Plugins/Experimental/DynamicMaterial/Source/DynamicMaterialEditor/Private/DynamicMaterialEditorModule.cpp
--- a/Engine/Plugins/Experimental/DynamicMaterial/Source/DynamicMaterialEditor/Private/DynamicMaterialEditorModule.cpp
+++ b/Engine/Plugins/Experimental/DynamicMaterial/Source/DynamicMaterialEditor/Private/DynamicMaterialEditorModule.cpp
@@ -34,6 +34,7 @@
 #include "Model/DynamicMaterialModelEditorOnlyData.h"
 #include "Modules/ModuleManager.h"
 #include "PropertyEditorModule.h"
+#include "ThumbnailRendering/ThumbnailManager.h"
 #include "UI/PropertyGenerators/DMComponentPropertyRowGenerator.h"
 #include "UI/PropertyGenerators/DMInputThroughputPropertyRowGenerator.h"
 #include "UI/PropertyGenerators/DMMaterialEffectFunctionPropertyRowGenerator.h"
@@ -44,6 +45,7 @@
 #include "UI/PropertyGenerators/DMTextureUVDynamicPropertyRowGenerator.h"
 #include "UI/PropertyGenerators/DMTextureUVPropertyRowGenerator.h"
 #include "UI/PropertyGenerators/DMThroughputPropertyRowGenerator.h"
+#include "UI/Utils/DynamicMaterialInstanceThumbnailRenderer.h"
 #include "UI/Widgets/SDMMaterialDesigner.h"
 
 DEFINE_LOG_CATEGORY(LogDynamicMaterialEditor);
@@ -258,6 +260,8 @@ void FDynamicMaterialEditorModule::StartupModule()
 		});
 
 	FDMValueDetailsRowExtensions::Get().RegisterRowExtensions();
+
+	UThumbnailManager::Get().RegisterCustomRenderer(UDynamicMaterialInstance::StaticClass(), UDynamicMaterialInstanceThumbnailRenderer::StaticClass());
 }
 
 void FDynamicMaterialEditorModule::ShutdownModule()
@@ -265,10 +269,15 @@ void FDynamicMaterialEditorModule::ShutdownModule()
 	FDynamicMaterialEditorCommands::Unregister();
 	FDMContentBrowserIntegration::Disintegrate();
 
-	if (FDynamicMaterialModule::AreUObjectsSafe() && FModuleManager::Get().IsModuleLoaded("PropertyEditor"))
+	if (FDynamicMaterialModule::AreUObjectsSafe())
 	{
-		FPropertyEditorModule& PropertyModule = FModuleManager::GetModuleChecked<FPropertyEditorModule>("PropertyEditor");
-		PropertyModule.UnregisterCustomPropertyTypeLayout(UDynamicMaterialModelEditorOnlyData::StaticClass()->GetFName());
+		if (FModuleManager::Get().IsModuleLoaded("PropertyEditor"))
+		{
+			FPropertyEditorModule& PropertyModule = FModuleManager::GetModuleChecked<FPropertyEditorModule>("PropertyEditor");
+			PropertyModule.UnregisterCustomPropertyTypeLayout(UDynamicMaterialModelEditorOnlyData::StaticClass()->GetFName());
+		}
+
+		UThumbnailManager::Get().UnregisterCustomRenderer(UDynamicMaterialInstance::StaticClass());
 	}
 
 	using namespace UE::DynamicMaterialEditor::Private;
diff --git a/Engine/Plugins/Experimental/DynamicMaterial/Source/DynamicMaterialEditor/Private/DynamicMaterialEditorSettings.cpp b/Engine/Plugins/Experimental/DynamicMaterial/Source/DynamicMaterialEditor/Private/DynamicMaterialEditorSettings.cpp
--- a/Engine/Plugins/Experimental/DynamicMaterial/Source/DynamicMaterialEditor/Private/DynamicMaterialEditorSettings.cpp
+++ b/Engine/Plugins/Experimental/DynamicMaterial/Source/DynamicMaterialEditor/Private/DynamicMaterialEditorSettings.cpp
@@ -139,6 +139,8 @@ UDynamicMaterialEditorSettings::UDynamicMaterialEditorSettings()
 
 	DefaultMask = TSoftObjectPtr<UTexture>(FSoftObjectPath(TEXT("/Script/Engine.Texture2D'/Engine/EngineResources/WhiteSquareTexture.WhiteSquareTexture'")));
 
+	CustomPreviewMesh = TSoftObjectPtr<UStaticMesh>(FSoftObjectPath(TEXT("/Script/Engine.StaticMesh'/Engine/EngineMeshes/SM_MatPreviewMesh_01.SM_MatPreviewMesh_01'")));
+
 	FDMMaterialChannelListPreset Opaque;
 	Opaque.Name = TEXT("Opaque");
 	Opaque.bBaseColor = true;
diff --git a/Engine/Plugins/Experimental/DynamicMaterial/Source/DynamicMaterialEditor/Private/DynamicMaterialEditorSettings.h b/Engine/Plugins/Experimental/DynamicMaterial/Source/DynamicMaterialEditor/Private/DynamicMaterialEditorSettings.h
--- a/Engine/Plugins/Experimental/DynamicMaterial/Source/DynamicMaterialEditor/Private/DynamicMaterialEditorSettings.h
+++ b/Engine/Plugins/Experimental/DynamicMaterial/Source/DynamicMaterialEditor/Private/DynamicMaterialEditorSettings.h
@@ -135,7 +135,7 @@ enum class EDMMaterialPreviewMesh : uint8
 	Cube,
 	Sphere,
 	Cylinder,
-	ShaderBall
+	Custom
 };
 
 UENUM(BlueprintType)
@@ -154,6 +154,24 @@ enum class EDMMaterialEditorLayout : uint8
 	Last = TopSlim
 };
 
+USTRUCT()
+struct FDMContentBrowserThumbnailSettings
+{
+	GENERATED_BODY()
+
+	UPROPERTY(Config, EditAnywhere, Category = "Preview")
+	EDMMaterialPreviewMesh PreviewMesh = EDMMaterialPreviewMesh::Custom;
+
+	UPROPERTY(Config, EditAnywhere, Category = "Preview", meta = (EditCondition = "PreviewMesh == EDMMaterialPreviewMesh::Custom", EditConditionHides))
+	float CustomMeshOrbitPitch = -30.f;
+
+	UPROPERTY(Config, EditAnywhere, Category = "Preview", meta = (EditCondition = "PreviewMesh == EDMMaterialPreviewMesh::Custom", EditConditionHides))
+	float CustomMeshOrbitYaw = 152.f;
+
+	UPROPERTY(Config, EditAnywhere, Category = "Preview", meta = (EditCondition = "PreviewMesh == EDMMaterialPreviewMesh::Custom", EditConditionHides))
+	float CustomMeshZoom = -409.f;
+};
+
 /**
  * Material Designer Settings
  */
@@ -203,6 +221,12 @@ public:
 	UPROPERTY(Config, EditAnywhere, BlueprintReadWrite, Category = "Preview")
 	EDMMaterialPreviewMesh PreviewMesh;
 
+	UPROPERTY(Config, EditAnywhere, BlueprintReadWrite, Category = "Preview")
+	TSoftObjectPtr<UStaticMesh> CustomPreviewMesh;
+
+	UPROPERTY(Config, EditAnywhere, Category = "Preview")
+	FDMContentBrowserThumbnailSettings ContentBrowserThumbnail;
+
 	UPROPERTY(Config, EditAnywhere, BlueprintReadWrite, Category = "Preview")
 	bool bShowPreviewBackground;
 
diff --git a/Engine/Plugins/Experimental/DynamicMaterial/Source/DynamicMaterialEditor/Private/Material/AssetDefinition_DynamicMaterialInstance.cpp b/Engine/Plugins/Experimental/DynamicMaterial/Source/DynamicMaterialEditor/Private/Material/AssetDefinition_DynamicMaterialInstance.cpp
--- a/Engine/Plugins/Experimental/DynamicMaterial/Source/DynamicMaterialEditor/Private/Material/AssetDefinition_DynamicMaterialInstance.cpp
+++ b/Engine/Plugins/Experimental/DynamicMaterial/Source/DynamicMaterialEditor/Private/Material/AssetDefinition_DynamicMaterialInstance.cpp
@@ -2,9 +2,11 @@
 
 #include "AssetDefinition_DynamicMaterialInstance.h"
 #include "AssetToolsModule.h"
+#include "DynamicMaterialEditorSettings.h"
 #include "IAssetTools.h"
 #include "Material/DynamicMaterialInstance.h"
 #include "Model/DynamicMaterialModel.h"
+#include "ThumbnailRendering/SceneThumbnailInfoWithPrimitive.h"
 
 #define LOCTEXT_NAMESPACE "AssetDefinition_DynamicMaterialInstance"
 
@@ -34,6 +36,25 @@ TConstArrayView<FAssetCategoryPath> UAssetDefinition_DynamicMaterialInstance::Ge
 	return Categories;
 }
 
+UThumbnailInfo* UAssetDefinition_DynamicMaterialInstance::LoadThumbnailInfo(const FAssetData& InAsset) const
+{
+	const UDynamicMaterialEditorSettings* Settings = GetDefault<UDynamicMaterialEditorSettings>();
+
+	if (!Settings)
+	{
+		return nullptr;
+	}
+
+	UDynamicMaterialInstance* MaterialInstance = Cast<UDynamicMaterialInstance>(InAsset.GetAsset());
+
+	if (!MaterialInstance)
+	{
+		return nullptr;
+	}
+
+	return UE::Editor::FindOrCreateThumbnailInfo<USceneThumbnailInfoWithPrimitive>(MaterialInstance);
+}
+
 EAssetCommandResult UAssetDefinition_DynamicMaterialInstance::OpenAssets(const FAssetOpenArgs& InOpenArgs) const
 {
 	TArray<UObject*> MaterialModels;
diff --git a/Engine/Plugins/Experimental/DynamicMaterial/Source/DynamicMaterialEditor/Private/Material/AssetDefinition_DynamicMaterialInstance.h b/Engine/Plugins/Experimental/DynamicMaterial/Source/DynamicMaterialEditor/Private/Material/AssetDefinition_DynamicMaterialInstance.h
--- a/Engine/Plugins/Experimental/DynamicMaterial/Source/DynamicMaterialEditor/Private/Material/AssetDefinition_DynamicMaterialInstance.h
+++ b/Engine/Plugins/Experimental/DynamicMaterial/Source/DynamicMaterialEditor/Private/Material/AssetDefinition_DynamicMaterialInstance.h
@@ -19,6 +19,7 @@ public:
 	virtual TSoftClassPtr<UObject> GetAssetClass() const override;
 	virtual FLinearColor GetAssetColor() const override;
 	virtual TConstArrayView<FAssetCategoryPath> GetAssetCategories() const override;
+	virtual UThumbnailInfo* LoadThumbnailInfo(const FAssetData& InAsset) const override;
 	virtual EAssetCommandResult OpenAssets(const FAssetOpenArgs& InOpenArgs) const override;
 	//~ End UAssetDefinition End
 };
diff --git a/Engine/Plugins/Experimental/DynamicMaterial/Source/DynamicMaterialEditor/Private/UI/Utils/DynamicMaterialInstanceThumbnailRenderer.cpp b/Engine/Plugins/Experimental/DynamicMaterial/Source/DynamicMaterialEditor/Private/UI/Utils/DynamicMaterialInstanceThumbnailRenderer.cpp
new file mode 100644
--- /dev/null
+++ b/Engine/Plugins/Experimental/DynamicMaterial/Source/DynamicMaterialEditor/Private/UI/Utils/DynamicMaterialInstanceThumbnailRenderer.cpp
@@ -0,0 +1,69 @@
+// Copyright Epic Games, Inc. All Rights Reserved.
+
+#include "UI/Utils/DynamicMaterialInstanceThumbnailRenderer.h"
+
+#include "Engine/Texture2D.h"
+#include "Material/DynamicMaterialInstance.h"
+#include "Materials/Material.h"
+#include "Materials/MaterialInterface.h"
+#include "Misc/App.h"
+#include "SceneInterface.h"
+#include "SceneView.h"
+#include "ShowFlags.h"
+#include "ThumbnailHelpers.h"
+#include "ThumbnailRendering/ThumbnailManager.h"
+#include "UI/Utils/DynamicMaterialInstanceThumbnailScene.h"
+#include "Widgets/Images/SImage.h"
+
+void UDynamicMaterialInstanceThumbnailRenderer::Draw(UObject* InObject, int32 InX, int32 InY, uint32 InWidth, uint32 InHeight, 
+	FRenderTarget* InRenderTarget, FCanvas* InCanvas, bool bInAdditionalViewFamily)
+{
+	UDynamicMaterialInstance* MaterialInstance = Cast<UDynamicMaterialInstance>(InObject);
+
+	if (IsValid(MaterialInstance))
+	{
+		if (ThumbnailScene == nullptr || ensure(ThumbnailScene->GetWorld() != nullptr) == false)
+		{
+			if (ThumbnailScene)
+			{
+				FlushRenderingCommands();
+				delete ThumbnailScene;
+			}
+
+			ThumbnailScene = new FDynamicMaterialInstanceThumbnailScene();
+		}
+
+		ThumbnailScene->SetDynamicMaterialInstance(MaterialInstance);
+		ThumbnailScene->GetScene()->UpdateSpeedTreeWind(0.0);
+
+		FSceneViewFamilyContext ViewFamily(FSceneViewFamily::ConstructionValues(InRenderTarget, ThumbnailScene->GetScene(), FEngineShowFlags(ESFIM_Game))
+			.SetTime(UThumbnailRenderer::GetTime())
+			.SetAdditionalViewFamily(bInAdditionalViewFamily));
+
+		ViewFamily.EngineShowFlags.DisableAdvancedFeatures();
+		ViewFamily.EngineShowFlags.SetSeparateTranslucency(false);
+		ViewFamily.EngineShowFlags.MotionBlur = 0;
+		ViewFamily.EngineShowFlags.AntiAliasing = 0;
+
+		RenderViewFamily(InCanvas, &ViewFamily, ThumbnailScene->CreateView(&ViewFamily, InX, InY, InWidth, InHeight));
+
+		ThumbnailScene->SetDynamicMaterialInstance(nullptr);
+		ThumbnailScene->SetStaticMesh(nullptr);
+	}
+}
+
+bool UDynamicMaterialInstanceThumbnailRenderer::CanVisualizeAsset(UObject* InObject)
+{
+	return IsValid(InObject) && InObject->IsA<UDynamicMaterialInstance>();
+}
+
+void UDynamicMaterialInstanceThumbnailRenderer::BeginDestroy()
+{
+	if (ThumbnailScene != nullptr)
+	{
+		delete ThumbnailScene;
+		ThumbnailScene = nullptr;
+	}
+
+	Super::BeginDestroy();
+}
diff --git a/Engine/Plugins/Experimental/DynamicMaterial/Source/DynamicMaterialEditor/Private/UI/Utils/DynamicMaterialInstanceThumbnailRenderer.h b/Engine/Plugins/Experimental/DynamicMaterial/Source/DynamicMaterialEditor/Private/UI/Utils/DynamicMaterialInstanceThumbnailRenderer.h
new file mode 100644
--- /dev/null
+++ b/Engine/Plugins/Experimental/DynamicMaterial/Source/DynamicMaterialEditor/Private/UI/Utils/DynamicMaterialInstanceThumbnailRenderer.h
@@ -0,0 +1,34 @@
+// Copyright Epic Games, Inc. All Rights Reserved.
+
+#pragma once
+
+#include "ThumbnailRendering/DefaultSizedThumbnailRenderer.h"
+
+#include "DynamicMaterialInstanceThumbnailRenderer.generated.h"
+
+class FCanvas;
+class FDynamicMaterialInstanceThumbnailScene;
+class FRenderTarget;
+
+/**
+ * This thumbnail renderer displays a given Material Designer Instance.
+ */
+UCLASS(config = Editor)
+class UDynamicMaterialInstanceThumbnailRenderer : public UDefaultSizedThumbnailRenderer
+{
+	GENERATED_BODY()
+
+	//~ Begin UThumbnailRenderer
+	virtual void Draw(UObject* InObject, int32 InX, int32 InY, uint32 InWidth, uint32 InHeight, FRenderTarget* InRenderTarget, 
+		FCanvas* InCanvas, bool bInAdditionalViewFamily) override;
+	virtual bool CanVisualizeAsset(UObject* InObject) override;
+	//~ End UThumbnailRenderer
+
+	//~ Begin UObject
+	virtual void BeginDestroy() override;
+	//~ End UObject
+
+private:
+	FDynamicMaterialInstanceThumbnailScene* ThumbnailScene;
+};
+
diff --git a/Engine/Plugins/Experimental/DynamicMaterial/Source/DynamicMaterialEditor/Private/UI/Utils/DynamicMaterialInstanceThumbnailScene.cpp b/Engine/Plugins/Experimental/DynamicMaterial/Source/DynamicMaterialEditor/Private/UI/Utils/DynamicMaterialInstanceThumbnailScene.cpp
new file mode 100644
--- /dev/null
+++ b/Engine/Plugins/Experimental/DynamicMaterial/Source/DynamicMaterialEditor/Private/UI/Utils/DynamicMaterialInstanceThumbnailScene.cpp
@@ -0,0 +1,221 @@
+// Copyright Epic Games, Inc. All Rights Reserved.
+
+#include "UI/Utils/DynamicMaterialInstanceThumbnailScene.h"
+
+#include "Components/StaticMeshComponent.h"
+#include "DynamicMaterialEditorSettings.h"
+#include "Editor/UnrealEdEngine.h"
+#include "Engine/StaticMesh.h"
+#include "Engine/StaticMeshActor.h"
+#include "Material/DynamicMaterialInstance.h"
+#include "ThumbnailRendering/SceneThumbnailInfo.h"
+#include "ThumbnailRendering/SceneThumbnailInfoWithPrimitive.h"
+#include "UnrealEdGlobals.h"
+
+FDynamicMaterialInstanceThumbnailScene::FDynamicMaterialInstanceThumbnailScene()
+	: FThumbnailPreviewScene()
+{
+	bForceAllUsedMipsResident = false;
+
+	UWorld* World = GetWorld();
+	check(World);
+
+	FActorSpawnParameters SpawnInfo;
+	SpawnInfo.SpawnCollisionHandlingOverride = ESpawnActorCollisionHandlingMethod::AlwaysSpawn;
+	SpawnInfo.bNoFail = true;
+	SpawnInfo.ObjectFlags = RF_Transient;
+	PreviewActor = World->SpawnActor<AStaticMeshActor>(SpawnInfo);
+
+	PreviewActor->GetStaticMeshComponent()->SetCanEverAffectNavigation(false);
+	PreviewActor->GetStaticMeshComponent()->SetMobility(EComponentMobility::Movable);
+	PreviewActor->SetActorEnableCollision(false);
+}
+
+void FDynamicMaterialInstanceThumbnailScene::SetStaticMesh(UStaticMesh* InStaticMesh)
+{
+	check(PreviewActor);
+
+	UStaticMeshComponent* StaticMeshComponent = PreviewActor->GetStaticMeshComponent();
+	check(StaticMeshComponent);
+
+	StaticMeshComponent->SetStaticMesh(InStaticMesh);
+
+	if (InStaticMesh)
+	{
+		FTransform MeshTransform = FTransform::Identity;
+
+		PreviewActor->SetActorLocation(FVector(0, 0, 0), /* Sweep */ false);
+
+		// Force LOD 0
+		StaticMeshComponent->ForcedLodModel = 1;
+
+		StaticMeshComponent->UpdateBounds();
+
+		// Center the mesh at the world origin then offset to put it on top of the plane
+		const float BoundsZOffset = GetBoundsZOffset(StaticMeshComponent->Bounds);
+
+		PreviewActor->SetActorLocation(
+			-StaticMeshComponent->Bounds.Origin + FVector(0, 0, BoundsZOffset),
+			/* Sweep */ false
+		);
+	}
+
+	StaticMeshComponent->RecreateRenderState_Concurrent();
+}
+
+void FDynamicMaterialInstanceThumbnailScene::SetDynamicMaterialInstance(UDynamicMaterialInstance* InMaterialInstance)
+{
+	check(PreviewActor);
+
+	MaterialInstance = InMaterialInstance;
+
+	UStaticMeshComponent* StaticMeshComponent = PreviewActor->GetStaticMeshComponent();
+	check(StaticMeshComponent);
+
+	if (!InMaterialInstance)
+	{
+		StaticMeshComponent->OverrideMaterials.Empty();
+		return;
+	}
+
+	const UDynamicMaterialEditorSettings* Settings = GetDefault<UDynamicMaterialEditorSettings>();
+	check(Settings);
+
+	USceneThumbnailInfoWithPrimitive* ThumbnailInfo = Cast<USceneThumbnailInfoWithPrimitive>(InMaterialInstance->ThumbnailInfo);
+	EThumbnailPrimType PrimitiveType = TPT_Plane;
+
+	if (ThumbnailInfo && ThumbnailInfo->bUserModifiedShape)
+	{
+		PrimitiveType = ThumbnailInfo->PrimitiveType;
+	}
+	else
+	{
+		switch (Settings->ContentBrowserThumbnail.PreviewMesh)
+		{
+			default:
+			case EDMMaterialPreviewMesh::Plane:
+				PrimitiveType = TPT_Plane;
+				break;
+
+			case EDMMaterialPreviewMesh::Cube:
+				PrimitiveType = TPT_Cube;
+				break;
+
+			case EDMMaterialPreviewMesh::Sphere:
+				PrimitiveType = TPT_Sphere;
+				break;
+
+			case EDMMaterialPreviewMesh::Cylinder:
+				PrimitiveType = TPT_Cylinder;
+				break;
+
+			case EDMMaterialPreviewMesh::Custom:
+				PrimitiveType = TPT_None;
+				break;
+		}
+	}
+
+	switch (PrimitiveType)
+	{
+		default:
+		case TPT_Plane:
+			SetStaticMesh(GUnrealEd->GetThumbnailManager()->EditorPlane);
+			break;
+
+		case TPT_Cube:
+			SetStaticMesh(GUnrealEd->GetThumbnailManager()->EditorCube);
+			break;
+
+		case TPT_Sphere:
+			SetStaticMesh(GUnrealEd->GetThumbnailManager()->EditorSphere);
+			break;
+
+		case TPT_Cylinder:
+			SetStaticMesh(GUnrealEd->GetThumbnailManager()->EditorCylinder);
+			break;
+
+		case TPT_None:
+		{
+			SetStaticMesh(Settings->CustomPreviewMesh.LoadSynchronous());
+			break;
+		}
+	}
+
+	UStaticMesh* StaticMesh = StaticMeshComponent->GetStaticMesh();
+
+	if (!StaticMesh)
+	{
+		StaticMeshComponent->OverrideMaterials.Empty();
+		return;
+	}
+
+	const int32 MaterialCount = StaticMesh->GetStaticMaterials().Num();
+
+	StaticMeshComponent->OverrideMaterials.SetNumUninitialized(MaterialCount);
+
+	for (int32 Index = 0; Index < MaterialCount; ++Index)
+	{
+		StaticMeshComponent->OverrideMaterials[Index] = InMaterialInstance;
+	}
+
+	StaticMeshComponent->MarkRenderStateDirty();
+}
+
+void FDynamicMaterialInstanceThumbnailScene::GetViewMatrixParameters(const float InFOVDegrees, FVector& OutOrigin, float& OutOrbitPitch, 
+	float& OutOrbitYaw, float& OutOrbitZoom) const
+{
+	check(MaterialInstance);
+	check(PreviewActor);
+	check(PreviewActor->GetStaticMeshComponent());
+	check(PreviewActor->GetStaticMeshComponent()->GetStaticMesh());
+
+	const float HalfFOVRadians = FMath::DegreesToRadians<float>(InFOVDegrees) * 0.5f;
+	// Add extra size to view slightly outside of the sphere to compensate for perspective
+	const float HalfMeshSize = static_cast<float>(PreviewActor->GetStaticMeshComponent()->Bounds.SphereRadius * 1.15);
+	const float BoundsZOffset = GetBoundsZOffset(PreviewActor->GetStaticMeshComponent()->Bounds);
+	const float TargetDistance = HalfMeshSize / FMath::Tan(HalfFOVRadians);
+
+	USceneThumbnailInfoWithPrimitive* ThumbnailInfo = Cast<USceneThumbnailInfoWithPrimitive>(MaterialInstance->ThumbnailInfo);
+
+	OutOrigin = FVector(0, 0, -BoundsZOffset);
+
+	if (ThumbnailInfo)
+	{
+		if (TargetDistance + ThumbnailInfo->OrbitZoom < 0)
+		{
+			ThumbnailInfo->OrbitZoom = -TargetDistance;
+		}
+
+		OutOrbitPitch = ThumbnailInfo->OrbitPitch;
+		OutOrbitYaw = ThumbnailInfo->OrbitYaw;
+		OutOrbitZoom = TargetDistance + ThumbnailInfo->OrbitZoom;
+	}
+	else
+	{
+		const UDynamicMaterialEditorSettings* Settings = GetDefault<UDynamicMaterialEditorSettings>();
+		check(Settings);
+
+		switch (Settings->ContentBrowserThumbnail.PreviewMesh)
+		{
+			case EDMMaterialPreviewMesh::Plane:
+			case EDMMaterialPreviewMesh::Cylinder:
+			case EDMMaterialPreviewMesh::Sphere:
+				OutOrbitPitch = 0.f;
+				OutOrbitYaw = 0.f;
+				OutOrbitZoom = TargetDistance;
+				break;
+
+			case EDMMaterialPreviewMesh::Cube:
+				OutOrbitPitch = -30.f;
+				OutOrbitYaw = 30.f;
+				OutOrbitZoom = TargetDistance;
+				break;
+
+			case EDMMaterialPreviewMesh::Custom:
+				OutOrbitPitch = -30.f;
+				OutOrbitYaw = 152.f;
+				OutOrbitZoom = TargetDistance - 409.f;
+				break;
+		}
+	}
+}
\ No newline at end of file
diff --git a/Engine/Plugins/Experimental/DynamicMaterial/Source/DynamicMaterialEditor/Private/UI/Utils/DynamicMaterialInstanceThumbnailScene.h b/Engine/Plugins/Experimental/DynamicMaterial/Source/DynamicMaterialEditor/Private/UI/Utils/DynamicMaterialInstanceThumbnailScene.h
new file mode 100644
--- /dev/null
+++ b/Engine/Plugins/Experimental/DynamicMaterial/Source/DynamicMaterialEditor/Private/UI/Utils/DynamicMaterialInstanceThumbnailScene.h
@@ -0,0 +1,32 @@
+// Copyright Epic Games, Inc. All Rights Reserved.
+
+#pragma once
+
+#include "ThumbnailHelpers.h"
+
+#include "Math/MathFwd.h"
+
+class AStaticMeshActor;
+class UDynamicMaterialInstance;
+
+class FDynamicMaterialInstanceThumbnailScene : public FThumbnailPreviewScene
+{
+public:
+	FDynamicMaterialInstanceThumbnailScene();
+
+	/** Sets the static mesh to render with. */
+	void SetStaticMesh(UStaticMesh* InStaticMesh);
+
+	/** Sets the Material Instance used on the mesh. */
+	void SetDynamicMaterialInstance(UDynamicMaterialInstance* InMaterialInstance);
+
+protected:
+	//~ Begin FThumbnailPreviewScene
+	virtual void GetViewMatrixParameters(const float InFOVDegrees, FVector& OutOrigin, float& OutOrbitPitch, 
+		float& OutOrbitYaw, float& OutOrbitZoom) const override;
+	//~ End FThumbnailPreviewScene
+
+private:
+	AStaticMeshActor* PreviewActor;
+	UDynamicMaterialInstance* MaterialInstance;
+};
diff --git a/Engine/Plugins/Experimental/DynamicMaterial/Source/DynamicMaterialEditor/Private/UI/Widgets/Editor/SDMMaterialPreview.cpp b/Engine/Plugins/Experimental/DynamicMaterial/Source/DynamicMaterialEditor/Private/UI/Widgets/Editor/SDMMaterialPreview.cpp
--- a/Engine/Plugins/Experimental/DynamicMaterial/Source/DynamicMaterialEditor/Private/UI/Widgets/Editor/SDMMaterialPreview.cpp
+++ b/Engine/Plugins/Experimental/DynamicMaterial/Source/DynamicMaterialEditor/Private/UI/Widgets/Editor/SDMMaterialPreview.cpp
@@ -30,15 +30,6 @@
 
 #define LOCTEXT_NAMESPACE "SDMMaterialPreview"
 
-namespace UE::DynamicMaterialEditor::Private
-{
-	UStaticMesh* GetShaderBallMesh()
-	{
-		TSoftObjectPtr<UStaticMesh> ShaderBallSoft = TSoftObjectPtr<UStaticMesh>(FSoftObjectPath(TEXT("/Script/Engine.StaticMesh'/Engine/EngineMeshes/SM_MatPreviewMesh_01.SM_MatPreviewMesh_01'")));
-		return ShaderBallSoft.LoadSynchronous();
-	}
-}
-
 void SDMMaterialPreview::PrivateRegisterAttributes(FSlateAttributeDescriptor::FInitializer&)
 {
 }
@@ -222,7 +213,7 @@ void SDMMaterialPreview::SetPreviewType(EDMMaterialPreviewMesh InPrimitiveType)
 		case EDMMaterialPreviewMesh::Cube: Primitive = GUnrealEd->GetThumbnailManager()->EditorCube; break;
 		case EDMMaterialPreviewMesh::Sphere: Primitive = GUnrealEd->GetThumbnailManager()->EditorSphere; break;
 		case EDMMaterialPreviewMesh::Cylinder: Primitive = GUnrealEd->GetThumbnailManager()->EditorCylinder; break;
-		case EDMMaterialPreviewMesh::ShaderBall: Primitive = GetShaderBallMesh(); break;
+		case EDMMaterialPreviewMesh::Custom: Primitive = Settings->CustomPreviewMesh.LoadSynchronous(); break;
 		default: return;
 	}
 
@@ -252,6 +243,13 @@ void SDMMaterialPreview::SetPreviewAsset(UObject* InAsset)
 		return;
 	}
 
+	UDynamicMaterialEditorSettings* Settings = UDynamicMaterialEditorSettings::Get();
+
+	if (!Settings)
+	{
+		return;
+	}
+
 	// Unregister the current component
 	if (PreviewMeshComponent != nullptr)
 	{
@@ -269,8 +267,6 @@ void SDMMaterialPreview::SetPreviewAsset(UObject* InAsset)
 
 	using namespace UE::DynamicMaterialEditor::Private;
 
-	UStaticMesh* ShaderBallMesh = GetShaderBallMesh();
-
 	// Update the rotation of the plane mesh so that it is front facing to the viewport camera's default forward view.
 	if (StaticMesh == GUnrealEd->GetThumbnailManager()->EditorPlane)
 	{
@@ -283,7 +279,7 @@ void SDMMaterialPreview::SetPreviewAsset(UObject* InAsset)
 		Transform.SetRotation(FQuat(PlaneRotation));
 		Transform.SetScale3D(FVector(0.75));
 	}
-	else if (StaticMesh == ShaderBallMesh)
+	else if (StaticMesh == Settings->CustomPreviewMesh)
 	{
 		const FRotator PlaneRotation(0.0f, -90.0f, 0.0f);
 		Transform.SetRotation(FQuat(PlaneRotation));
@@ -472,7 +468,7 @@ TSharedRef<SWidget> SDMMaterialPreview::GenerateToolbarMenu()
 		MeshSection.AddEntry(FToolMenuEntry::InitMenuEntry(MaterialEditorCommands.SetCylinderPreview));
 
 		MeshSection.AddEntry(FToolMenuEntry::InitMenuEntry(
-			DMEditorCommands.SetShaderBallPreview, 
+			DMEditorCommands.SetCustomPreviewMesh, 
 			TAttribute<FText>(),
 			TAttribute<FText>(),
 			FSlateIconFinder::FindIcon(TEXT("GraphEditor.SpawnActor_16x"))
@@ -591,10 +587,10 @@ void SDMMaterialPreview::BindCommands()
 		FGetActionCheckState::CreateSP(this, &SDMMaterialPreview::IsPreviewTypeSet, EDMMaterialPreviewMesh::Cylinder));
 
 	CommandList->MapAction(
-		DMEditorCommands.SetShaderBallPreview,
-		FExecuteAction::CreateSP(this, &SDMMaterialPreview::SetPreviewType, EDMMaterialPreviewMesh::ShaderBall),
+		DMEditorCommands.SetCustomPreviewMesh,
+		FExecuteAction::CreateSP(this, &SDMMaterialPreview::SetPreviewType, EDMMaterialPreviewMesh::Custom),
 		FCanExecuteAction(),
-		FGetActionCheckState::CreateSP(this, &SDMMaterialPreview::IsPreviewTypeSet, EDMMaterialPreviewMesh::ShaderBall));
+		FGetActionCheckState::CreateSP(this, &SDMMaterialPreview::IsPreviewTypeSet, EDMMaterialPreviewMesh::Custom));
 
 	CommandList->MapAction(
 		MaterialEditorCommands.TogglePreviewBackground,
