Material Designer
- Slot preview materials now generate correctly.

diff --git a/Engine/Plugins/Experimental/DynamicMaterial/Source/DynamicMaterialEditor/Private/Components/DMMaterialProperty.cpp b/Engine/Plugins/Experimental/DynamicMaterial/Source/DynamicMaterialEditor/Private/Components/DMMaterialProperty.cpp
--- a/Engine/Plugins/Experimental/DynamicMaterial/Source/DynamicMaterialEditor/Private/Components/DMMaterialProperty.cpp
+++ b/Engine/Plugins/Experimental/DynamicMaterial/Source/DynamicMaterialEditor/Private/Components/DMMaterialProperty.cpp
@@ -344,7 +344,6 @@ void UDMMaterialProperty::GenerateExpressions(const TSharedRef<FDMMaterialBuildS
 	MaterialPropertyPtr->Expression = nullptr;
 	MaterialPropertyPtr->OutputIndex = 0;
 
-	UMaterialExpression* LastPropertyExpression = nullptr;
 	UDMMaterialSlot* Slot = EditorOnlyData->GetSlotForMaterialProperty(MaterialProperty);
 
 	if (!Slot || Slot->GetLayers().IsEmpty())
@@ -361,7 +360,7 @@ void UDMMaterialProperty::GenerateExpressions(const TSharedRef<FDMMaterialBuildS
 		return;
 	}
 
-	LastPropertyExpression = InBuildState->GetLastSlotPropertyExpression(Slot, MaterialProperty);
+	UMaterialExpression* LastPropertyExpression = InBuildState->GetLastSlotPropertyExpression(Slot, MaterialProperty);
 
 	if (!LastPropertyExpression)
 	{
@@ -612,7 +611,7 @@ void UDMMaterialProperty::GeneratePreviewMaterial(UMaterial* InPreviewMaterial)
 	InPreviewMaterial->BlendMode = EBlendMode::BLEND_Translucent;
 
 	TSharedRef<FDMMaterialBuildState> BuildState = EditorOnlyData->CreateBuildState(InPreviewMaterial);
-	BuildState->SetPreviewMaterial();
+	BuildState->SetPreviewObject(this);
 
 	UE_LOG(LogDynamicMaterialEditor, Display, TEXT("Building Material Designer Property Preview (%s)..."), *GetName());
 
diff --git a/Engine/Plugins/Experimental/DynamicMaterial/Source/DynamicMaterialEditor/Private/Components/DMMaterialStage.cpp b/Engine/Plugins/Experimental/DynamicMaterial/Source/DynamicMaterialEditor/Private/Components/DMMaterialStage.cpp
--- a/Engine/Plugins/Experimental/DynamicMaterial/Source/DynamicMaterialEditor/Private/Components/DMMaterialStage.cpp
+++ b/Engine/Plugins/Experimental/DynamicMaterial/Source/DynamicMaterialEditor/Private/Components/DMMaterialStage.cpp
@@ -516,7 +516,7 @@ void UDMMaterialStage::GenerateExpressions(const TSharedRef<FDMMaterialBuildStat
 		StageExpressions.Append(InBuildState->GetStageSourceExpressions(Source));		
 	}
 
-	if (InBuildState->IsPreviewMaterial())
+	if (InBuildState->GetPreviewObject() == this)
 	{
 		int32 OutputChannel = FDMMaterialStageConnectorChannel::WHOLE_CHANNEL;
 		int32 OutputIndex = 0;
diff --git a/Engine/Plugins/Experimental/DynamicMaterial/Source/DynamicMaterialEditor/Private/Components/DMMaterialStageBlend.cpp b/Engine/Plugins/Experimental/DynamicMaterial/Source/DynamicMaterialEditor/Private/Components/DMMaterialStageBlend.cpp
--- a/Engine/Plugins/Experimental/DynamicMaterial/Source/DynamicMaterialEditor/Private/Components/DMMaterialStageBlend.cpp
+++ b/Engine/Plugins/Experimental/DynamicMaterial/Source/DynamicMaterialEditor/Private/Components/DMMaterialStageBlend.cpp
@@ -341,7 +341,7 @@ bool UDMMaterialStageBlend::GenerateStagePreviewMaterial(UDMMaterialStage* InSta
 	}
 
 	TSharedRef<FDMMaterialBuildState> BuildState = ModelEditorOnlyData->CreateBuildState(InPreviewMaterial);
-	BuildState->SetPreviewMaterial();
+	BuildState->SetPreviewObject(InStage);
 
 	UDMMaterialStageSource* PreviewSource = GetInputB();
 
@@ -750,7 +750,7 @@ void UDMMaterialStageBlend::GeneratePreviewMaterial(UMaterial* InPreviewMaterial
 	}
  
 	TSharedRef<FDMMaterialBuildState> BuildState = ModelEditorOnlyData->CreateBuildState(InPreviewMaterial);
-	BuildState->SetPreviewMaterial();
+	BuildState->SetPreviewObject(this);
  
 	if (!bHasStageInput || Inputs.IsEmpty())
 	{
diff --git a/Engine/Plugins/Experimental/DynamicMaterial/Source/DynamicMaterialEditor/Private/Components/DMMaterialStageInput.cpp b/Engine/Plugins/Experimental/DynamicMaterial/Source/DynamicMaterialEditor/Private/Components/DMMaterialStageInput.cpp
--- a/Engine/Plugins/Experimental/DynamicMaterial/Source/DynamicMaterialEditor/Private/Components/DMMaterialStageInput.cpp
+++ b/Engine/Plugins/Experimental/DynamicMaterial/Source/DynamicMaterialEditor/Private/Components/DMMaterialStageInput.cpp
@@ -70,7 +70,7 @@ void UDMMaterialStageInput::GeneratePreviewMaterial(UMaterial* InPreviewMaterial
 	check(ModelEditorOnlyData);
 
 	TSharedRef<FDMMaterialBuildState> BuildState = ModelEditorOnlyData->CreateBuildState(InPreviewMaterial);
-	BuildState->SetPreviewMaterial();
+	BuildState->SetPreviewObject(this);
 
 	GenerateExpressions(BuildState);
 	UMaterialExpression* StageSourceExpression = BuildState->GetLastStageSourceExpression(this);
diff --git a/Engine/Plugins/Experimental/DynamicMaterial/Source/DynamicMaterialEditor/Private/Components/DMMaterialStageSource.cpp b/Engine/Plugins/Experimental/DynamicMaterial/Source/DynamicMaterialEditor/Private/Components/DMMaterialStageSource.cpp
--- a/Engine/Plugins/Experimental/DynamicMaterial/Source/DynamicMaterialEditor/Private/Components/DMMaterialStageSource.cpp
+++ b/Engine/Plugins/Experimental/DynamicMaterial/Source/DynamicMaterialEditor/Private/Components/DMMaterialStageSource.cpp
@@ -123,7 +123,7 @@ void UDMMaterialStageSource::GeneratePreviewMaterial(UMaterial* InPreviewMateria
 	check(ModelEditorOnlyData);
 
 	TSharedRef<FDMMaterialBuildState> BuildState = ModelEditorOnlyData->CreateBuildState(InPreviewMaterial);
-	BuildState->SetPreviewMaterial();
+	BuildState->SetPreviewObject(this);
 
 	GenerateExpressions(BuildState);
 	UMaterialExpression* StageSourceExpression = BuildState->GetLastStageSourceExpression(this);
@@ -152,7 +152,7 @@ bool UDMMaterialStageSource::GenerateStagePreviewMaterial(UDMMaterialStage* InSt
 	check(ModelEditorOnlyData);
 
 	TSharedRef<FDMMaterialBuildState> BuildState = ModelEditorOnlyData->CreateBuildState(InPreviewMaterial);
-	BuildState->SetPreviewMaterial();
+	BuildState->SetPreviewObject(InStage);
 
 	UDMMaterialStageSource* PreviewSource = InStage->GetSource();
 
diff --git a/Engine/Plugins/Experimental/DynamicMaterial/Source/DynamicMaterialEditor/Private/Components/DMMaterialStageThroughput.cpp b/Engine/Plugins/Experimental/DynamicMaterial/Source/DynamicMaterialEditor/Private/Components/DMMaterialStageThroughput.cpp
--- a/Engine/Plugins/Experimental/DynamicMaterial/Source/DynamicMaterialEditor/Private/Components/DMMaterialStageThroughput.cpp
+++ b/Engine/Plugins/Experimental/DynamicMaterial/Source/DynamicMaterialEditor/Private/Components/DMMaterialStageThroughput.cpp
@@ -917,7 +917,7 @@ void UDMMaterialStageThroughput::GeneratePreviewMaterial(UMaterial* InPreviewMat
 	}
 
 	TSharedRef<FDMMaterialBuildState> BuildState = ModelEditorOnlyData->CreateBuildState(InPreviewMaterial);
-	BuildState->SetPreviewMaterial();
+	BuildState->SetPreviewObject(this);
 
 	if (!bHasStageInput || Inputs.IsEmpty())
 	{
diff --git a/Engine/Plugins/Experimental/DynamicMaterial/Source/DynamicMaterialEditor/Private/Components/DMMaterialStageThroughputLayerBlend.cpp b/Engine/Plugins/Experimental/DynamicMaterial/Source/DynamicMaterialEditor/Private/Components/DMMaterialStageThroughputLayerBlend.cpp
--- a/Engine/Plugins/Experimental/DynamicMaterial/Source/DynamicMaterialEditor/Private/Components/DMMaterialStageThroughputLayerBlend.cpp
+++ b/Engine/Plugins/Experimental/DynamicMaterial/Source/DynamicMaterialEditor/Private/Components/DMMaterialStageThroughputLayerBlend.cpp
@@ -296,7 +296,9 @@ void UDMMaterialStageThroughputLayerBlend::GenerateExpressions(const TSharedRef<
 		return;
 	}
 
-	if (InBuildState->IsPreviewMaterial())
+	UObject* PreviewObject = InBuildState->GetPreviewObject();
+
+	if (PreviewObject == this || PreviewObject == GetStage())
 	{
 		GeneratePreviewExpressions(InBuildState);
 	}
@@ -581,7 +583,7 @@ bool UDMMaterialStageThroughputLayerBlend::GenerateStagePreviewMaterial(UDMMater
 	}
 
 	TSharedRef<FDMMaterialBuildState> BuildState = ModelEditorOnlyData->CreateBuildState(InPreviewMaterial);
-	BuildState->SetPreviewMaterial();
+	BuildState->SetPreviewObject(InStage);
 
 	UDMMaterialStageSource* PreviewSource = GetInputMask();
 
@@ -954,7 +956,7 @@ void UDMMaterialStageThroughputLayerBlend::GeneratePreviewMaterial(UMaterial* In
 	}
  
 	TSharedRef<FDMMaterialBuildState> BuildState = ModelEditorOnlyData->CreateBuildState(InPreviewMaterial);
-	BuildState->SetPreviewMaterial();
+	BuildState->SetPreviewObject(this);
 
 	if (ChannelInputs.IsEmpty())
 	{
diff --git a/Engine/Plugins/Experimental/DynamicMaterial/Source/DynamicMaterialEditor/Private/Components/MaterialStageExpressions/DMMSEMathBase.cpp b/Engine/Plugins/Experimental/DynamicMaterial/Source/DynamicMaterialEditor/Private/Components/MaterialStageExpressions/DMMSEMathBase.cpp
--- a/Engine/Plugins/Experimental/DynamicMaterial/Source/DynamicMaterialEditor/Private/Components/MaterialStageExpressions/DMMSEMathBase.cpp
+++ b/Engine/Plugins/Experimental/DynamicMaterial/Source/DynamicMaterialEditor/Private/Components/MaterialStageExpressions/DMMSEMathBase.cpp
@@ -397,7 +397,7 @@ void UDMMaterialStageExpressionMathBase::GeneratePreviewMaterial(UMaterial* InPr
 	}
 
 	TSharedRef<FDMMaterialBuildState> BuildState = ModelEditorOnlyData->CreateBuildState(InPreviewMaterial);
-	BuildState->SetPreviewMaterial();
+	BuildState->SetPreviewObject(this);
 
 	if (!bHasStageInput || Inputs.IsEmpty())
 	{
diff --git a/Engine/Plugins/Experimental/DynamicMaterial/Source/DynamicMaterialEditor/Private/Model/DMMaterialBuildState.cpp b/Engine/Plugins/Experimental/DynamicMaterial/Source/DynamicMaterialEditor/Private/Model/DMMaterialBuildState.cpp
--- a/Engine/Plugins/Experimental/DynamicMaterial/Source/DynamicMaterialEditor/Private/Model/DMMaterialBuildState.cpp
+++ b/Engine/Plugins/Experimental/DynamicMaterial/Source/DynamicMaterialEditor/Private/Model/DMMaterialBuildState.cpp
@@ -26,7 +26,6 @@ FDMMaterialBuildState::FDMMaterialBuildState(UMaterial* InDynamicMaterial, UDyna
 	, MaterialModel(InMaterialModel)
 	, bDirtyAssets(bInDirtyAssets)
 	, bIgnoreUVs(false)
-	, bIsPreviewMaterial(false)
 	, Utils(MakeShared<FDMMaterialBuildUtils>(*this))
 {
 	check(InDynamicMaterial);
@@ -75,9 +74,9 @@ void FDMMaterialBuildState::SetIgnoreUVs()
 	bIgnoreUVs = true;
 }
 
-void FDMMaterialBuildState::SetPreviewMaterial()
+void FDMMaterialBuildState::SetPreviewObject(UObject* InObject)
 {
-	bIsPreviewMaterial = true;
+	PreviewObject = InObject;
 	DynamicMaterial->MaterialDomain = EMaterialDomain::MD_UI;
 
 	if (UDynamicMaterialEditorSettings::Get()->bPreviewImagesUseTextureUVs == false)
diff --git a/Engine/Plugins/Experimental/DynamicMaterial/Source/DynamicMaterialEditor/Private/Model/DMMaterialBuildUtils.cpp b/Engine/Plugins/Experimental/DynamicMaterial/Source/DynamicMaterialEditor/Private/Model/DMMaterialBuildUtils.cpp
--- a/Engine/Plugins/Experimental/DynamicMaterial/Source/DynamicMaterialEditor/Private/Model/DMMaterialBuildUtils.cpp
+++ b/Engine/Plugins/Experimental/DynamicMaterial/Source/DynamicMaterialEditor/Private/Model/DMMaterialBuildUtils.cpp
@@ -93,7 +93,7 @@ UMaterialExpression* FDMMaterialBuildUtils::CreateExpressionParameter(TSubclassO
 				}
 
 				// Error - fall back to hidden group
-				if (!BuildState.IsPreviewMaterial())
+				if (!BuildState.GetPreviewObject())
 				{
 					UE_LOG(LogDynamicMaterialEditor, Error, TEXT("Missing property for group type."));
 				}
diff --git a/Engine/Plugins/Experimental/DynamicMaterial/Source/DynamicMaterialEditor/Private/UI/Widgets/Visualizers/SDMMaterialComponentPreview.cpp b/Engine/Plugins/Experimental/DynamicMaterial/Source/DynamicMaterialEditor/Private/UI/Widgets/Visualizers/SDMMaterialComponentPreview.cpp
--- a/Engine/Plugins/Experimental/DynamicMaterial/Source/DynamicMaterialEditor/Private/UI/Widgets/Visualizers/SDMMaterialComponentPreview.cpp
+++ b/Engine/Plugins/Experimental/DynamicMaterial/Source/DynamicMaterialEditor/Private/UI/Widgets/Visualizers/SDMMaterialComponentPreview.cpp
@@ -98,6 +98,21 @@ void SDMMaterialComponentPreview::SetPreviewSize(const FVector2D& InSize)
 	}
 }
 
+UDMMaterialComponent* SDMMaterialComponentPreview::GetComponent() const
+{
+	return ComponentWeak.Get();
+}
+
+UMaterial* SDMMaterialComponentPreview::GetPreviewMaterial() const
+{
+	return PreviewMaterialBaseWeak.Get();
+}
+
+UMaterialInstanceDynamic* SDMMaterialComponentPreview::GetPreviewMaterialDynamic() const
+{
+	return PreviewMaterialDynamicWeak.Get();
+}
+
 void SDMMaterialComponentPreview::Tick(const FGeometry& AllottedGeometry, const double InCurrentTime, const float InDeltaTime)
 {
 	SCompoundWidget::Tick(AllottedGeometry, InCurrentTime, InDeltaTime);
diff --git a/Engine/Plugins/Experimental/DynamicMaterial/Source/DynamicMaterialEditor/Private/UI/Widgets/Visualizers/SDMMaterialComponentPreview.h b/Engine/Plugins/Experimental/DynamicMaterial/Source/DynamicMaterialEditor/Private/UI/Widgets/Visualizers/SDMMaterialComponentPreview.h
--- a/Engine/Plugins/Experimental/DynamicMaterial/Source/DynamicMaterialEditor/Private/UI/Widgets/Visualizers/SDMMaterialComponentPreview.h
+++ b/Engine/Plugins/Experimental/DynamicMaterial/Source/DynamicMaterialEditor/Private/UI/Widgets/Visualizers/SDMMaterialComponentPreview.h
@@ -40,6 +40,12 @@ public:
 
 	void SetPreviewSize(const FVector2D& InSize);
 
+	UDMMaterialComponent* GetComponent() const;
+
+	UMaterial* GetPreviewMaterial() const;
+
+	UMaterialInstanceDynamic* GetPreviewMaterialDynamic() const;
+
 	//~ Begin SWidget
 	virtual void Tick(const FGeometry& InAllottedGeometry, const double InCurrentTime, const float InDeltaTime) override;
 	//~ End SWidget
diff --git a/Engine/Plugins/Experimental/DynamicMaterial/Source/DynamicMaterialEditor/Public/Model/DMMaterialBuildState.h b/Engine/Plugins/Experimental/DynamicMaterial/Source/DynamicMaterialEditor/Public/Model/DMMaterialBuildState.h
--- a/Engine/Plugins/Experimental/DynamicMaterial/Source/DynamicMaterialEditor/Public/Model/DMMaterialBuildState.h
+++ b/Engine/Plugins/Experimental/DynamicMaterial/Source/DynamicMaterialEditor/Public/Model/DMMaterialBuildState.h
@@ -57,9 +57,9 @@ struct FDMMaterialBuildState : public IDMMaterialBuildStateInterface, public TSh
 	DYNAMICMATERIALEDITOR_API void SetIgnoreUVs();
 
 	/** Whether the build state is building a full material or a preview. */
-	bool IsPreviewMaterial() const { return bIsPreviewMaterial; }
+	UObject* GetPreviewObject() const { return PreviewObject; }
 
-	DYNAMICMATERIALEDITOR_API void SetPreviewMaterial();
+	DYNAMICMATERIALEDITOR_API void SetPreviewObject(UObject* InObject);
 
 	/** A handy set of tools for creating material expressions. */
 	DYNAMICMATERIALEDITOR_API virtual IDMMaterialBuildUtilsInterface& GetBuildUtils() const override;
@@ -158,7 +158,7 @@ private:
 	const UDMMaterialProperty* CurrentProperty = nullptr;
 	bool bDirtyAssets;
 	bool bIgnoreUVs;
-	bool bIsPreviewMaterial;
+	UObject* PreviewObject = nullptr;
 	TSharedRef<FDMMaterialBuildUtils> Utils;
 
 	TMap<const UDMMaterialValue*, TArray<UMaterialExpression*>> Values;
