Material Designer
- Changed default layout to Top Vertical
- Reordered layout list.
- Removed all layouts but Left, Top Vertical and Top Slim.
- Renamed Top Vertical to Top.

diff --git a/Engine/Plugins/Experimental/DynamicMaterial/Source/DynamicMaterialEditor/Private/DynamicMaterialEditorSettings.cpp b/Engine/Plugins/Experimental/DynamicMaterial/Source/DynamicMaterialEditor/Private/DynamicMaterialEditorSettings.cpp
--- a/Engine/Plugins/Experimental/DynamicMaterial/Source/DynamicMaterialEditor/Private/DynamicMaterialEditorSettings.cpp
+++ b/Engine/Plugins/Experimental/DynamicMaterial/Source/DynamicMaterialEditor/Private/DynamicMaterialEditorSettings.cpp
@@ -306,7 +306,7 @@ void UDynamicMaterialEditorSettings::OpenEditorSettingsWindow() const
 
 void UDynamicMaterialEditorSettings::ResetAllLayoutSettings()
 {
-	Layout = EDMMaterialEditorLayout::Left;
+	Layout = EDMMaterialEditorLayout::Top;
 	bUseFullChannelNamesInTopSlimLayout = false;
 	SplitterLocation = 0.5;
 	PreviewSplitterLocation = 0.333;
diff --git a/Engine/Plugins/Experimental/DynamicMaterial/Source/DynamicMaterialEditor/Private/DynamicMaterialEditorSettings.h b/Engine/Plugins/Experimental/DynamicMaterial/Source/DynamicMaterialEditor/Private/DynamicMaterialEditorSettings.h
--- a/Engine/Plugins/Experimental/DynamicMaterial/Source/DynamicMaterialEditor/Private/DynamicMaterialEditorSettings.h
+++ b/Engine/Plugins/Experimental/DynamicMaterial/Source/DynamicMaterialEditor/Private/DynamicMaterialEditorSettings.h
@@ -141,17 +141,12 @@ enum class EDMMaterialPreviewMesh : uint8
 UENUM(BlueprintType)
 enum class EDMMaterialEditorLayout : uint8
 {
-	Left,
-	LeftAutoHide,
-	LeftSlim,
-	TopVertical,
-	TopVerticalAutoHide,
-	TopHorizontal,
-	TopHorizontalAutoHide,
+	Top,
 	TopSlim,
+	Left,
 
-	First = Left,
-	Last = TopSlim
+	First = Top,
+	Last = Left
 };
 
 USTRUCT()
diff --git a/Engine/Plugins/Experimental/DynamicMaterial/Source/DynamicMaterialEditor/Private/UI/Menus/DMToolBarMenus.cpp b/Engine/Plugins/Experimental/DynamicMaterial/Source/DynamicMaterialEditor/Private/UI/Menus/DMToolBarMenus.cpp
--- a/Engine/Plugins/Experimental/DynamicMaterial/Source/DynamicMaterialEditor/Private/UI/Menus/DMToolBarMenus.cpp
+++ b/Engine/Plugins/Experimental/DynamicMaterial/Source/DynamicMaterialEditor/Private/UI/Menus/DMToolBarMenus.cpp
@@ -210,18 +210,6 @@ void FDMToolBarMenus::AddEditorLayoutSection(UToolMenu* InMenu)
 	{
 		FUIAction Action;
 
-		switch (Layout)
-		{
-			case EDMMaterialEditorLayout::LeftAutoHide:
-			case EDMMaterialEditorLayout::TopHorizontalAutoHide:
-			case EDMMaterialEditorLayout::TopVerticalAutoHide:
-				Action.CanExecuteAction = FCanExecuteAction::CreateLambda([]() { return false; });
-				break;
-
-			default:
-				break;
-		}
-
 		Action.GetActionCheckState = FGetActionCheckState::CreateLambda(
 			[Layout]()
 			{
diff --git a/Engine/Plugins/Experimental/DynamicMaterial/Source/DynamicMaterialEditor/Private/UI/Widgets/SDMMaterialDesigner.cpp b/Engine/Plugins/Experimental/DynamicMaterial/Source/DynamicMaterialEditor/Private/UI/Widgets/SDMMaterialDesigner.cpp
--- a/Engine/Plugins/Experimental/DynamicMaterial/Source/DynamicMaterialEditor/Private/UI/Widgets/SDMMaterialDesigner.cpp
+++ b/Engine/Plugins/Experimental/DynamicMaterial/Source/DynamicMaterialEditor/Private/UI/Widgets/SDMMaterialDesigner.cpp
@@ -549,30 +549,11 @@ bool SDMMaterialDesigner::SetEditorLayout(EDMMaterialEditorLayout InLayout, UDyn
 				.MaterialModelBase(InMaterialModelBase);
 			break;
 
-		case EDMMaterialEditorLayout::LeftAutoHide:
-			break;
-
-		case EDMMaterialEditorLayout::LeftSlim:
-			NewEditor = SNew(SDMMaterialEditor_LeftSlim, SharedThis(this))
-				.MaterialModelBase(InMaterialModelBase);
-			break;
-
-		case EDMMaterialEditorLayout::TopHorizontal:
-			NewEditor = SNew(SDMMaterialEditor_TopHorizontal, SharedThis(this))
-				.MaterialModelBase(InMaterialModelBase);
-			break;
-
-		case EDMMaterialEditorLayout::TopHorizontalAutoHide:
-			break;
-
-		case EDMMaterialEditorLayout::TopVertical:
+		case EDMMaterialEditorLayout::Top:
 			NewEditor = SNew(SDMMaterialEditor_TopVertical, SharedThis(this))
 				.MaterialModelBase(InMaterialModelBase);
 			break;
 
-		case EDMMaterialEditorLayout::TopVerticalAutoHide:
-			break;
-
 		case EDMMaterialEditorLayout::TopSlim:
 			NewEditor = SNew(SDMMaterialEditor_TopSlim, SharedThis(this))
 				.MaterialModelBase(InMaterialModelBase);
@@ -603,30 +584,11 @@ bool SDMMaterialDesigner::SetEditorLayout(EDMMaterialEditorLayout InLayout, cons
 				.MaterialProperty(InObjectMaterialProperty);
 			break;
 
-		case EDMMaterialEditorLayout::LeftAutoHide:
-			break;
-
-		case EDMMaterialEditorLayout::LeftSlim:
-			NewEditor = SNew(SDMMaterialEditor_LeftSlim, SharedThis(this))
-				.MaterialProperty(InObjectMaterialProperty);
-			break;
-
-		case EDMMaterialEditorLayout::TopHorizontal:
-			NewEditor = SNew(SDMMaterialEditor_TopHorizontal, SharedThis(this))
-				.MaterialProperty(InObjectMaterialProperty);
-			break;
-
-		case EDMMaterialEditorLayout::TopHorizontalAutoHide:
-			break;
-
-		case EDMMaterialEditorLayout::TopVertical:
+		case EDMMaterialEditorLayout::Top:
 			NewEditor = SNew(SDMMaterialEditor_TopVertical, SharedThis(this))
 				.MaterialProperty(InObjectMaterialProperty);
 			break;
 
-		case EDMMaterialEditorLayout::TopVerticalAutoHide:
-			break;
-
 		case EDMMaterialEditorLayout::TopSlim:
 			NewEditor = SNew(SDMMaterialEditor_TopSlim, SharedThis(this))
 				.MaterialProperty(InObjectMaterialProperty);
