Material Designer
- You can now drag and drop textures onto the active channels on the channels page.

diff --git a/Engine/Plugins/Experimental/Avalanche/Source/AvalancheEditor/Private/MaterialDesigner/AvaMaterialDesignerTextureAssetFactory.cpp b/Engine/Plugins/Experimental/Avalanche/Source/AvalancheEditor/Private/MaterialDesigner/AvaMaterialDesignerTextureAssetFactory.cpp
--- a/Engine/Plugins/Experimental/Avalanche/Source/AvalancheEditor/Private/MaterialDesigner/AvaMaterialDesignerTextureAssetFactory.cpp
+++ b/Engine/Plugins/Experimental/Avalanche/Source/AvalancheEditor/Private/MaterialDesigner/AvaMaterialDesignerTextureAssetFactory.cpp
@@ -127,12 +127,6 @@ AActor* UAvaMaterialDesignerTextureAssetFactory::SpawnActor(UObject* InAsset, UL
 	EditorOnlyData->SetChannelListPreset(TEXT("Emissive"));
 	EditorOnlyData->OnWizardComplete();
 
-	ON_SCOPE_EXIT
-	{
-		const IDynamicMaterialEditorModule& MaterialDesignerModule = IDynamicMaterialEditorModule::Get();
-		MaterialDesignerModule.OpenMaterialInstance(NewInstance, InLevel->GetWorld(), /* Invoke Tab */ true);
-	};
-
 	const UDMMaterialSlot* Slot = EditorOnlyData->GetSlotForMaterialProperty(EDMMaterialPropertyType::EmissiveColor);
 
 	if (!Slot)
diff --git a/Engine/Plugins/Experimental/DynamicMaterial/Source/DynamicMaterialEditor/Private/UI/Widgets/Editor/SDMMaterialProperties.cpp b/Engine/Plugins/Experimental/DynamicMaterial/Source/DynamicMaterialEditor/Private/UI/Widgets/Editor/SDMMaterialProperties.cpp
--- a/Engine/Plugins/Experimental/DynamicMaterial/Source/DynamicMaterialEditor/Private/UI/Widgets/Editor/SDMMaterialProperties.cpp
+++ b/Engine/Plugins/Experimental/DynamicMaterial/Source/DynamicMaterialEditor/Private/UI/Widgets/Editor/SDMMaterialProperties.cpp
@@ -16,14 +16,18 @@
 #include "Items/ICustomDetailsViewCustomCategoryItem.h"
 #include "Items/ICustomDetailsViewCustomItem.h"
 #include "Items/ICustomDetailsViewItem.h"
+#include "Misc/MessageDialog.h"
+#include "Model/DynamicMaterialModel.h"
 #include "Model/DynamicMaterialModelBase.h"
 #include "Model/DynamicMaterialModelDynamic.h"
 #include "Model/DynamicMaterialModelEditorOnlyData.h"
+#include "SAssetDropTarget.h"
 #include "Styling/StyleColors.h"
 #include "UI/Utils/DMWidgetStatics.h"
 #include "UI/Widgets/Editor/SDMMaterialPropertySelector.h"
 #include "UI/Widgets/SDMMaterialEditor.h"
 #include "UI/Widgets/Visualizers/SDMMaterialComponentPreview.h"
+#include "Utils/DMMaterialSlotFunctionLibrary.h"
 #include "Utils/DMPrivate.h"
 #include "Widgets/Input/SCheckBox.h"
 #include "Widgets/Layout/SBox.h"
@@ -270,8 +274,8 @@ TSharedRef<SWidget> SDMMaterialProperties::CreatePropertyRow(UDMMaterialProperty
 {
 	// There are all ensured to be valid by the caller of this method
 	TSharedPtr<SDMMaterialEditor> EditorWidget = EditorWidgetWeak.Pin();
-	UDynamicMaterialModel* MaterialModel = EditorWidget->GetMaterialModel();
-	UDynamicMaterialModelEditorOnlyData* EditorOnlyData = UDynamicMaterialModelEditorOnlyData::Get(MaterialModel);
+	UDynamicMaterialModelBase* MaterialModelBase = EditorWidget->GetMaterialModelBase();
+	UDynamicMaterialModelEditorOnlyData* EditorOnlyData = UDynamicMaterialModelEditorOnlyData::Get(MaterialModelBase);
 	UDynamicMaterialEditorSettings* Settings = UDynamicMaterialEditorSettings::Get();
 	const EDMMaterialPropertyType MaterialProperty = InProperty->GetMaterialProperty();
 
@@ -322,55 +326,67 @@ TSharedRef<SWidget> SDMMaterialProperties::CreatePropertyRow(UDMMaterialProperty
 
 	// If the property is just disabled, leave the slider widget blank.
 
-	return SNew(SHorizontalBox)
+	TSharedRef<SHorizontalBox> PropertyWidget = SNew(SHorizontalBox)
 
-			+ SHorizontalBox::Slot()
-			.AutoWidth()
-			.VAlign(EVerticalAlignment::VAlign_Center)
-			.Padding(0.f, 5.f, 0.f, 5.f)
-			[
-				PreviewWidgetContainer
-			]
+		+ SHorizontalBox::Slot()
+		.AutoWidth()
+		.VAlign(EVerticalAlignment::VAlign_Center)
+		.Padding(0.f, 5.f, 0.f, 5.f)
+		[
+			PreviewWidgetContainer
+		]
 
-			+ SHorizontalBox::Slot()
-			.FillWidth(1.f)
-			.HAlign(EHorizontalAlignment::HAlign_Fill)
-			.VAlign(EVerticalAlignment::VAlign_Fill)
-			.Padding(5.f, 5.f, 0.f, 5.f)
+		+ SHorizontalBox::Slot()
+		.FillWidth(1.f)
+		.HAlign(EHorizontalAlignment::HAlign_Fill)
+		.VAlign(EVerticalAlignment::VAlign_Fill)
+		.Padding(5.f, 5.f, 0.f, 5.f)
+		[
+			SNew(SVerticalBox)
+
+			+ SVerticalBox::Slot()
+			.AutoHeight()
 			[
-				SNew(SVerticalBox)
+				SNew(SHorizontalBox)
 
-				+ SVerticalBox::Slot()
-				.AutoHeight()
+				+ SHorizontalBox::Slot()
+				.AutoWidth()
+				.HAlign(EHorizontalAlignment::HAlign_Left)
+				.VAlign(EVerticalAlignment::VAlign_Center)
 				[
-					SNew(SHorizontalBox)
-
-					+ SHorizontalBox::Slot()
-					.AutoWidth()
-					.HAlign(EHorizontalAlignment::HAlign_Left)
-					.VAlign(EVerticalAlignment::VAlign_Center)
-					[
-						CreateSlot_EnabledButton(MaterialProperty)
-					]
-
-					+ SHorizontalBox::Slot()
-					.AutoWidth()
-					.HAlign(EHorizontalAlignment::HAlign_Left)
-					.VAlign(EVerticalAlignment::VAlign_Center)
-					.Padding(5.f, 0.f, 0.f, 0.f)
-					[
-						PropertyName
-					]
+					CreateSlot_EnabledButton(MaterialProperty)
 				]
 
-				+ SVerticalBox::Slot()
-				.FillHeight(1.f)
-				.HAlign(EHorizontalAlignment::HAlign_Fill)
+				+ SHorizontalBox::Slot()
+				.AutoWidth()
+				.HAlign(EHorizontalAlignment::HAlign_Left)
 				.VAlign(EVerticalAlignment::VAlign_Center)
+				.Padding(5.f, 0.f, 0.f, 0.f)
 				[
-					Slider
+					PropertyName
 				]
-			];
+			]
+
+			+ SVerticalBox::Slot()
+			.FillHeight(1.f)
+			.HAlign(EHorizontalAlignment::HAlign_Fill)
+			.VAlign(EVerticalAlignment::VAlign_Center)
+			[
+				Slider
+			]
+		];
+
+	if (!bPropertyEnabled || !MaterialModelBase->IsA<UDynamicMaterialModel>())
+	{
+		return PropertyWidget;
+	}
+
+	return SNew(SAssetDropTarget)
+		.OnAreAssetsAcceptableForDrop(this, &SDMMaterialProperties::OnAssetDraggedOver, MaterialProperty)
+		.OnAssetsDropped(this, &SDMMaterialProperties::OnAssetsDropped, MaterialProperty)
+		[
+			PropertyWidget
+		];
 }
 
 TSharedRef<SWidget> SDMMaterialProperties::CreateSlot_EnabledButton(EDMMaterialPropertyType InMaterialProperty)
@@ -743,4 +759,146 @@ void SDMMaterialProperties::OnSettingsUpdated(const FPropertyChangedEvent& InPro
 	}
 }
 
+bool SDMMaterialProperties::OnAssetDraggedOver(TArrayView<FAssetData> InAssets, EDMMaterialPropertyType InMaterialProperty)
+{
+	TSharedPtr<SDMMaterialEditor> EditorWidget = EditorWidgetWeak.Pin();
+
+	if (!EditorWidget.IsValid())
+	{
+		return false;
+	}
+
+	UDynamicMaterialModel* MaterialModel = EditorWidget->GetMaterialModel();
+
+	if (!MaterialModel)
+	{
+		return false;
+	}
+
+	UDynamicMaterialModelEditorOnlyData* EditorOnlyData = UDynamicMaterialModelEditorOnlyData::Get(MaterialModel);
+
+	if (!EditorOnlyData)
+	{
+		return false;
+	}
+
+	UDMMaterialProperty* Property = EditorOnlyData->GetMaterialProperty(InMaterialProperty);
+
+	if (!Property || !Property->IsEnabled())
+	{
+		return false;
+	}
+
+	const TArray<UClass*> AllowedClasses = {
+		UTexture::StaticClass()
+	};
+
+	TArray<FAssetData> Textures;
+
+	for (const FAssetData& Asset : InAssets)
+	{
+		UClass* AssetClass = Asset.GetClass(EResolveClass::Yes);
+
+		if (!AssetClass)
+		{
+			continue;
+		}
+
+		for (UClass* AllowedClass : AllowedClasses)
+		{
+			if (AssetClass->IsChildOf(AllowedClass))
+			{
+				Textures.Add(Asset);
+			}
+		}
+	}
+
+	return Textures.Num() == 1;
+}
+
+void SDMMaterialProperties::OnAssetsDropped(const FDragDropEvent& InDragDropEvent, TArrayView<FAssetData> InAssets, 
+	EDMMaterialPropertyType InMaterialProperty)
+{
+	for (const FAssetData& Asset : InAssets)
+	{
+		UClass* AssetClass = Asset.GetClass(EResolveClass::Yes);
+
+		if (!AssetClass)
+		{
+			continue;
+		}
+
+		if (AssetClass->IsChildOf(UTexture::StaticClass()))
+		{
+			HandleDrop_Texture(Cast<UTexture>(Asset.GetAsset()), InMaterialProperty);
+			break;
+		}
+	}
+}
+
+void SDMMaterialProperties::HandleDrop_Texture(UTexture* InTexture, EDMMaterialPropertyType InMaterialProperty)
+{
+	TSharedPtr<SDMMaterialEditor> EditorWidget = EditorWidgetWeak.Pin();
+
+	if (!EditorWidget.IsValid())
+	{
+		return;
+	}
+
+	UDynamicMaterialModel* MaterialModel = EditorWidget->GetMaterialModel();
+
+	if (!MaterialModel)
+	{
+		return;
+	}
+
+	UDynamicMaterialModelEditorOnlyData* EditorOnlyData = UDynamicMaterialModelEditorOnlyData::Get(MaterialModel);
+
+	if (!EditorOnlyData)
+	{
+		return;
+	}
+
+	UDMMaterialProperty* Property = EditorOnlyData->GetMaterialProperty(InMaterialProperty);
+
+	if (!Property || !Property->IsEnabled())
+	{
+		return;
+	}
+
+	UDMMaterialSlot* Slot = EditorOnlyData->GetSlotForMaterialProperty(InMaterialProperty);
+
+	if (!Slot)
+	{
+		return;
+	}
+
+	const EAppReturnType::Type Result = FMessageDialog::Open(
+		EAppMsgType::YesNoCancel,
+		LOCTEXT("ReplaceSlotsTextureSet",
+			"Material Designer Channel.\n\n"
+			"Replace Slot?\n\n"
+			"- Yes: Delete Layers.\n"
+			"- No: Add Layer.\n"
+			"- Cancel")
+	);
+
+	FDMScopedUITransaction Transaction(LOCTEXT("DropTexture", "Drop Texture On Channel"));
+
+	switch (Result)
+	{
+		case EAppReturnType::Yes:
+			UDMMaterialSlotFunctionLibrary::AddTextureLayer(Slot, InTexture, InMaterialProperty, /* Replace Slot */ true);
+			break;
+
+		case EAppReturnType::No:
+			UDMMaterialSlotFunctionLibrary::AddTextureLayer(Slot, InTexture, InMaterialProperty, /* Replace Slot */ false);
+			break;
+
+		default:
+			Transaction.Transaction.Cancel();
+			break;
+	}
+}
+
 #undef LOCTEXT_NAMESPACE
diff --git a/Engine/Plugins/Experimental/DynamicMaterial/Source/DynamicMaterialEditor/Private/UI/Widgets/Editor/SDMMaterialProperties.h b/Engine/Plugins/Experimental/DynamicMaterial/Source/DynamicMaterialEditor/Private/UI/Widgets/Editor/SDMMaterialProperties.h
--- a/Engine/Plugins/Experimental/DynamicMaterial/Source/DynamicMaterialEditor/Private/UI/Widgets/Editor/SDMMaterialProperties.h
+++ b/Engine/Plugins/Experimental/DynamicMaterial/Source/DynamicMaterialEditor/Private/UI/Widgets/Editor/SDMMaterialProperties.h
@@ -15,6 +15,7 @@ class SDMMaterialEditor;
 class SVerticalBox;
 class UDMMaterialProperty;
 class UDMMaterialSlot;
+class UTexture;
 enum class ECheckBoxState : uint8;
 enum class EDMMaterialPropertyType : uint8;
 struct FPropertyChangedEvent;
@@ -68,4 +69,10 @@ protected:
 	void OnExpansionStateChanged(const TSharedRef<ICustomDetailsViewItem>& InItem, bool bInExpansionState);
 
 	void OnSettingsUpdated(const FPropertyChangedEvent& InPropertyChangedEvent);
+
+	bool OnAssetDraggedOver(TArrayView<FAssetData> InAssets, EDMMaterialPropertyType InMaterialProperty);
+
+	void OnAssetsDropped(const FDragDropEvent& InDragDropEvent, TArrayView<FAssetData> InAssets, EDMMaterialPropertyType InMaterialProperty);
+
+	void HandleDrop_Texture(UTexture* InTexture, EDMMaterialPropertyType InMaterialProperty);
 };
diff --git a/Engine/Plugins/Experimental/DynamicMaterial/Source/DynamicMaterialEditor/Private/UI/Widgets/Editor/SlotEditor/SDMMaterialStage.cpp b/Engine/Plugins/Experimental/DynamicMaterial/Source/DynamicMaterialEditor/Private/UI/Widgets/Editor/SlotEditor/SDMMaterialStage.cpp
--- a/Engine/Plugins/Experimental/DynamicMaterial/Source/DynamicMaterialEditor/Private/UI/Widgets/Editor/SlotEditor/SDMMaterialStage.cpp
+++ b/Engine/Plugins/Experimental/DynamicMaterial/Source/DynamicMaterialEditor/Private/UI/Widgets/Editor/SlotEditor/SDMMaterialStage.cpp
@@ -92,11 +92,9 @@ void SDMMaterialStage::Construct(const FArguments& InArgs, const TSharedRef<SDMM
 		return;
 	}
 
-	TSharedPtr<SAssetDropTarget> DropTarget;
-
 	ChildSlot
 	[
-		SAssignNew(DropTarget, SAssetDropTarget)
+		SNew(SAssetDropTarget)
 		.OnAreAssetsAcceptableForDrop(this, &SDMMaterialStage::OnAssetDraggedOver)
 		.OnAssetsDropped(this, &SDMMaterialStage::OnAssetsDropped)
 		[
diff --git a/Engine/Plugins/Experimental/DynamicMaterial/Source/DynamicMaterialEditor/Private/UI/Widgets/SDMMaterialEditor.cpp b/Engine/Plugins/Experimental/DynamicMaterial/Source/DynamicMaterialEditor/Private/UI/Widgets/SDMMaterialEditor.cpp
--- a/Engine/Plugins/Experimental/DynamicMaterial/Source/DynamicMaterialEditor/Private/UI/Widgets/SDMMaterialEditor.cpp
+++ b/Engine/Plugins/Experimental/DynamicMaterial/Source/DynamicMaterialEditor/Private/UI/Widgets/SDMMaterialEditor.cpp
@@ -893,7 +893,7 @@ void SDMMaterialEditor::HandleDrop_TextureSet(UDMTextureSet* InTextureSet)
 		EAppMsgType::YesNoCancel,
 		LOCTEXT("ReplaceSlotsTextureSet",
 			"Material Designer Texture Set.\n\n"
-			"Replace Slots?\n"
+			"Replace Slots?\n\n"
 			"- Yes: Delete Layers.\n"
 			"- No: Add Layers.\n"
 			"- Cancel")
diff --git a/Engine/Plugins/Experimental/DynamicMaterial/Source/DynamicMaterialEditor/Private/Utils/DMMaterialSlotFunctionLibrary.cpp b/Engine/Plugins/Experimental/DynamicMaterial/Source/DynamicMaterialEditor/Private/Utils/DMMaterialSlotFunctionLibrary.cpp
--- a/Engine/Plugins/Experimental/DynamicMaterial/Source/DynamicMaterialEditor/Private/Utils/DMMaterialSlotFunctionLibrary.cpp
+++ b/Engine/Plugins/Experimental/DynamicMaterial/Source/DynamicMaterialEditor/Private/Utils/DMMaterialSlotFunctionLibrary.cpp
@@ -10,6 +10,7 @@
 #include "Components/DMMaterialStageFunction.h"
 #include "Components/DMMaterialStageGradient.h"
 #include "Components/DMMaterialStageThroughputLayerBlend.h"
+#include "Components/DMMaterialSubStage.h"
 #include "Components/DMMaterialValue.h"
 #include "Components/MaterialStageBlends/DMMSBNormal.h"
 #include "Components/MaterialStageExpressions/DMMSESceneTexture.h"
@@ -20,6 +21,7 @@
 #include "Components/MaterialStageInputs/DMMSISlot.h"
 #include "Components/MaterialStageInputs/DMMSITextureUV.h"
 #include "Components/MaterialStageInputs/DMMSIValue.h"
+#include "Components/MaterialValues/DMMaterialValueTexture.h"
 #include "DMPrivate.h"
 #include "Model/DynamicMaterialModelEditorOnlyData.h"
 #include "Utils/DMMaterialStageFunctionLibrary.h"
@@ -121,6 +123,93 @@ UDMMaterialLayerObject* UDMMaterialSlotFunctionLibrary::AddNewLayer(UDMMaterialS
 	return Layer;
 }
 
+UDMMaterialLayerObject* UDMMaterialSlotFunctionLibrary::AddTextureLayer(UDMMaterialSlot* InSlot, UTexture* InTexture, 
+	EDMMaterialPropertyType InPropertyType, bool bInReplaceSlot)
+{
+	if (GUndo)
+	{
+		InSlot->Modify();
+	}
+
+	UDMMaterialLayerObject* Layer = nullptr;
+
+	const FDMUpdateGuard Guard;
+
+	Layer = InSlot->AddDefaultLayer(InPropertyType);
+
+	if (!ensure(Layer))
+	{
+		return nullptr;
+	}
+
+	bool bMadeChange = true;
+
+	UDMMaterialStage* Stage = Layer->GetStage(EDMMaterialLayerStage::Base);
+
+	if (!ensure(Stage))
+	{
+		return nullptr;
+	}
+
+	UDMMaterialStageInputExpression* NewExpression = UDMMaterialStageInputExpression::ChangeStageInput_Expression(
+		Stage,
+		UDMMaterialStageExpressionTextureSample::StaticClass(),
+		UDMMaterialStageBlend::InputB,
+		FDMMaterialStageConnectorChannel::WHOLE_CHANNEL,
+		0,
+		FDMMaterialStageConnectorChannel::WHOLE_CHANNEL
+	);
+
+	if (!ensure(NewExpression))
+	{
+		return nullptr;
+	}
+
+	UDMMaterialSubStage* SubStage = NewExpression->GetSubStage();
+
+	if (!ensure(SubStage))
+	{
+		return nullptr;
+	}
+
+	UDMMaterialStageInputValue* InputValue = UDMMaterialStageInputValue::ChangeStageInput_NewLocalValue(
+		SubStage,
+		0,
+		FDMMaterialStageConnectorChannel::WHOLE_CHANNEL,
+		EDMValueType::VT_Texture,
+		FDMMaterialStageConnectorChannel::WHOLE_CHANNEL
+	);
+
+	if (ensure(InputValue))
+	{
+		UDMMaterialValueTexture* InputTexture = Cast<UDMMaterialValueTexture>(InputValue->GetValue());
+
+		if (ensure(InputTexture))
+		{
+			InputTexture->SetValue(InTexture);
+		}
+	}
+
+	if (bInReplaceSlot)
+	{
+		for (int32 Index = InSlot->GetLayers().Num() - 1; Index >= 0; --Index)
+		{
+			UDMMaterialLayerObject* LayerIter = InSlot->GetLayer(Index);
+
+			if (!LayerIter || LayerIter->GetStage(EDMMaterialLayerStage::Base) == Stage)
+			{
+				continue;
+			}
+
+			InSlot->RemoveLayer(LayerIter);
+		}
+	}
+
+	Layer->Update(Layer, EDMUpdateType::Structure);
+
+	return Layer;
+}
+
 UDMMaterialLayerObject* UDMMaterialSlotFunctionLibrary::AddNewLayer_NewLocalValue(UDMMaterialSlot* InSlot, EDMValueType InValueType)
 {
 	return AddNewLayer_NewLocalValue(
diff --git a/Engine/Plugins/Experimental/DynamicMaterial/Source/DynamicMaterialEditor/Public/Utils/DMMaterialSlotFunctionLibrary.h b/Engine/Plugins/Experimental/DynamicMaterial/Source/DynamicMaterialEditor/Public/Utils/DMMaterialSlotFunctionLibrary.h
--- a/Engine/Plugins/Experimental/DynamicMaterial/Source/DynamicMaterialEditor/Public/Utils/DMMaterialSlotFunctionLibrary.h
+++ b/Engine/Plugins/Experimental/DynamicMaterial/Source/DynamicMaterialEditor/Public/Utils/DMMaterialSlotFunctionLibrary.h
@@ -18,6 +18,7 @@ class UDMMaterialStageGradient;
 class UDMMaterialValue;
 class UDMRenderTargetRenderer;
 class UMaterialFunctionInterface;
+class UTexture;
 enum class EDMMaterialLayerStage : uint8;
 enum class EDMMaterialPropertyType : uint8;
 enum class EDMValueType : uint8;
@@ -67,4 +68,7 @@ public:
 
 	DYNAMICMATERIALEDITOR_API static UDMMaterialLayerObject* AddNewLayer(UDMMaterialSlot* InSlot, UDMMaterialStage* InNewBaseStage = nullptr, 
 		UDMMaterialStage* InNewMaskStage = nullptr);
+
+	DYNAMICMATERIALEDITOR_API static UDMMaterialLayerObject* AddTextureLayer(UDMMaterialSlot* InSlot, UTexture* InTexture, 
+		EDMMaterialPropertyType InPropertyType, bool bInReplaceSlot);
 };
