Material Designer
- Materials will no longer compile when a texture base stage is updated and triggers an update of the mask stage texture... unless that mask stage is not already a texture.

diff --git a/Engine/Plugins/Experimental/DynamicMaterial/Source/DynamicMaterialEditor/Private/Components/MaterialStageExpressions/DMMSETextureSampleBase.cpp b/Engine/Plugins/Experimental/DynamicMaterial/Source/DynamicMaterialEditor/Private/Components/MaterialStageExpressions/DMMSETextureSampleBase.cpp
--- a/Engine/Plugins/Experimental/DynamicMaterial/Source/DynamicMaterialEditor/Private/Components/MaterialStageExpressions/DMMSETextureSampleBase.cpp
+++ b/Engine/Plugins/Experimental/DynamicMaterial/Source/DynamicMaterialEditor/Private/Components/MaterialStageExpressions/DMMSETextureSampleBase.cpp
@@ -404,16 +404,16 @@ void UDMMaterialStageExpressionTextureSampleBase::UpdateMask()
 								check(NewLayerBlendTextureValue);
 
 								LayerBlendTextureInputValue->SetValue(NewLayerBlendTextureValue);
-							}
 
-							if (GUndo)
-							{
-								MaskStage->Modify();
-							}
+								if (GUndo)
+								{
+									MaskStage->Modify();
+								}
 
-							// Set output to alpha
-							MaskStage->UpdateInputMap(2, MaskStageInputConnections[2].Channels[0].SourceIndex,
-								FDMMaterialStageConnectorChannel::WHOLE_CHANNEL, 1, FDMMaterialStageConnectorChannel::WHOLE_CHANNEL, EDMMaterialPropertyType::None);
+								// Set output to alpha
+								MaskStage->UpdateInputMap(2, MaskStageInputConnections[2].Channels[0].SourceIndex,
+									FDMMaterialStageConnectorChannel::WHOLE_CHANNEL, 1, FDMMaterialStageConnectorChannel::WHOLE_CHANNEL, EDMMaterialPropertyType::None);
+							}
 
 							return;
 						}
