Material Designer
- When editing a dynamic material, the world "Dynamic" is now overlaid on the material preview.

diff --git a/Engine/Plugins/Experimental/DynamicMaterial/Source/DynamicMaterialEditor/Private/UI/Widgets/Editor/EditorLayouts/SDMMaterialEditor_Left.cpp b/Engine/Plugins/Experimental/DynamicMaterial/Source/DynamicMaterialEditor/Private/UI/Widgets/Editor/EditorLayouts/SDMMaterialEditor_Left.cpp
--- a/Engine/Plugins/Experimental/DynamicMaterial/Source/DynamicMaterialEditor/Private/UI/Widgets/Editor/EditorLayouts/SDMMaterialEditor_Left.cpp
+++ b/Engine/Plugins/Experimental/DynamicMaterial/Source/DynamicMaterialEditor/Private/UI/Widgets/Editor/EditorLayouts/SDMMaterialEditor_Left.cpp
@@ -51,7 +51,7 @@ TSharedRef<SWidget> SDMMaterialEditor_Left::CreateSlot_Left()
 			]
 		];
 
-	MaterialPreviewSlot = TDMWidgetSlot<SDMMaterialPreview>(MaterialPreviewSlotPtr, CreateSlot_Preview());
+	MaterialPreviewSlot = TDMWidgetSlot<SWidget>(MaterialPreviewSlotPtr, CreateSlot_Preview());
 	PropertySelectorSlot = TDMWidgetSlot<SDMMaterialPropertySelector>(PropertySelectorSlotPtr, CreateSlot_PropertySelector());
 
 	return NewLeft;
diff --git a/Engine/Plugins/Experimental/DynamicMaterial/Source/DynamicMaterialEditor/Private/UI/Widgets/Editor/EditorLayouts/SDMMaterialEditor_LeftBase.cpp b/Engine/Plugins/Experimental/DynamicMaterial/Source/DynamicMaterialEditor/Private/UI/Widgets/Editor/EditorLayouts/SDMMaterialEditor_LeftBase.cpp
--- a/Engine/Plugins/Experimental/DynamicMaterial/Source/DynamicMaterialEditor/Private/UI/Widgets/Editor/EditorLayouts/SDMMaterialEditor_LeftBase.cpp
+++ b/Engine/Plugins/Experimental/DynamicMaterial/Source/DynamicMaterialEditor/Private/UI/Widgets/Editor/EditorLayouts/SDMMaterialEditor_LeftBase.cpp
@@ -156,7 +156,7 @@ TSharedRef<SWidget> SDMMaterialEditor_LeftBase::CreateSlot_Left()
 			]
 		];
 
-	MaterialPreviewSlot = TDMWidgetSlot<SDMMaterialPreview>(MaterialPreviewSlotPtr, CreateSlot_Preview());
+	MaterialPreviewSlot = TDMWidgetSlot<SWidget>(MaterialPreviewSlotPtr, CreateSlot_Preview());
 	PropertySelectorSlot = TDMWidgetSlot<SDMMaterialPropertySelector>(PropertySelectorSlotPtr, CreateSlot_PropertySelector());
 
 	return NewLeft;
diff --git a/Engine/Plugins/Experimental/DynamicMaterial/Source/DynamicMaterialEditor/Private/UI/Widgets/Editor/EditorLayouts/SDMMaterialEditor_TopBase.cpp b/Engine/Plugins/Experimental/DynamicMaterial/Source/DynamicMaterialEditor/Private/UI/Widgets/Editor/EditorLayouts/SDMMaterialEditor_TopBase.cpp
--- a/Engine/Plugins/Experimental/DynamicMaterial/Source/DynamicMaterialEditor/Private/UI/Widgets/Editor/EditorLayouts/SDMMaterialEditor_TopBase.cpp
+++ b/Engine/Plugins/Experimental/DynamicMaterial/Source/DynamicMaterialEditor/Private/UI/Widgets/Editor/EditorLayouts/SDMMaterialEditor_TopBase.cpp
@@ -154,7 +154,7 @@ TSharedRef<SWidget> SDMMaterialEditor_TopBase::CreateSlot_Top()
 			]
 		];
 
-	MaterialPreviewSlot = TDMWidgetSlot<SDMMaterialPreview>(MaterialPreviewSlotPtr, CreateSlot_Preview());
+	MaterialPreviewSlot = TDMWidgetSlot<SWidget>(MaterialPreviewSlotPtr, CreateSlot_Preview());
 	PropertySelectorSlot = TDMWidgetSlot<SDMMaterialPropertySelector>(PropertySelectorSlotPtr, CreateSlot_PropertySelector());
 
 	return NewTop;
diff --git a/Engine/Plugins/Experimental/DynamicMaterial/Source/DynamicMaterialEditor/Private/UI/Widgets/Editor/EditorLayouts/SDMMaterialEditor_TopHorizontal.cpp b/Engine/Plugins/Experimental/DynamicMaterial/Source/DynamicMaterialEditor/Private/UI/Widgets/Editor/EditorLayouts/SDMMaterialEditor_TopHorizontal.cpp
--- a/Engine/Plugins/Experimental/DynamicMaterial/Source/DynamicMaterialEditor/Private/UI/Widgets/Editor/EditorLayouts/SDMMaterialEditor_TopHorizontal.cpp
+++ b/Engine/Plugins/Experimental/DynamicMaterial/Source/DynamicMaterialEditor/Private/UI/Widgets/Editor/EditorLayouts/SDMMaterialEditor_TopHorizontal.cpp
@@ -50,7 +50,7 @@ TSharedRef<SWidget> SDMMaterialEditor_TopHorizontal::CreateSlot_Top()
 			]
 		];
 
-	MaterialPreviewSlot = TDMWidgetSlot<SDMMaterialPreview>(MaterialPreviewSlotPtr, CreateSlot_Preview());
+	MaterialPreviewSlot = TDMWidgetSlot<SWidget>(MaterialPreviewSlotPtr, CreateSlot_Preview());
 	PropertySelectorSlot = TDMWidgetSlot<SDMMaterialPropertySelector>(PropertySelectorSlotPtr, CreateSlot_PropertySelector());
 
 	return NewLeft;
diff --git a/Engine/Plugins/Experimental/DynamicMaterial/Source/DynamicMaterialEditor/Private/UI/Widgets/Editor/EditorLayouts/SDMMaterialEditor_TopVertical.cpp b/Engine/Plugins/Experimental/DynamicMaterial/Source/DynamicMaterialEditor/Private/UI/Widgets/Editor/EditorLayouts/SDMMaterialEditor_TopVertical.cpp
--- a/Engine/Plugins/Experimental/DynamicMaterial/Source/DynamicMaterialEditor/Private/UI/Widgets/Editor/EditorLayouts/SDMMaterialEditor_TopVertical.cpp
+++ b/Engine/Plugins/Experimental/DynamicMaterial/Source/DynamicMaterialEditor/Private/UI/Widgets/Editor/EditorLayouts/SDMMaterialEditor_TopVertical.cpp
@@ -57,7 +57,7 @@ TSharedRef<SWidget> SDMMaterialEditor_TopVertical::CreateSlot_Top()
 			]
 		];
 
-	MaterialPreviewSlot = TDMWidgetSlot<SDMMaterialPreview>(MaterialPreviewSlotPtr, CreateSlot_Preview());
+	MaterialPreviewSlot = TDMWidgetSlot<SWidget>(MaterialPreviewSlotPtr, CreateSlot_Preview());
 	PropertySelectorSlot = TDMWidgetSlot<SDMMaterialPropertySelector>(PropertySelectorSlotPtr, CreateSlot_PropertySelector());
 
 	return NewLeft;
diff --git a/Engine/Plugins/Experimental/DynamicMaterial/Source/DynamicMaterialEditor/Private/UI/Widgets/SDMMaterialEditor.cpp b/Engine/Plugins/Experimental/DynamicMaterial/Source/DynamicMaterialEditor/Private/UI/Widgets/SDMMaterialEditor.cpp
--- a/Engine/Plugins/Experimental/DynamicMaterial/Source/DynamicMaterialEditor/Private/UI/Widgets/SDMMaterialEditor.cpp
+++ b/Engine/Plugins/Experimental/DynamicMaterial/Source/DynamicMaterialEditor/Private/UI/Widgets/SDMMaterialEditor.cpp
@@ -40,6 +40,7 @@
 #include "Widgets/SBoxPanel.h"
 #include "Widgets/SNullWidget.h"
 #include "Widgets/SToolTip.h"
+#include "Widgets/Text/STextBlock.h"
 
 #define LOCTEXT_NAMESPACE "SDMMaterialEditor"
 
@@ -424,7 +425,7 @@ void SDMMaterialEditor::OpenMaterialPreviewTab()
 
 	TSharedRef<SBox> Wrapper = SNew(SBox);
 
-	MaterialPreviewTabSlot = TDMWidgetSlot<SDMMaterialPreview>(
+	MaterialPreviewTabSlot = TDMWidgetSlot<SWidget>(
 		Wrapper, 
 		0, 
 		SNew(SDMMaterialPreview, SharedThis(this), MaterialModelBase)
@@ -477,7 +478,7 @@ TSharedPtr<IToolTip> SDMMaterialEditor::GetMaterialPreviewToolTip()
 			}
 		));
 
-	MaterialPreviewToolTipSlot = TDMWidgetSlot<SDMMaterialPreview>(
+	MaterialPreviewToolTipSlot = TDMWidgetSlot<SWidget>(
 		Wrapper,
 		0,
 		SNew(SDMMaterialPreview, SharedThis(this), MaterialModelBase)
@@ -1022,8 +1023,28 @@ TSharedRef<SDMMaterialProperties> SDMMaterialEditor::CreateSlot_MaterialProperti
 	return SNew(SDMMaterialProperties, SharedThis(this));
 }
 
-TSharedRef<SDMMaterialPreview> SDMMaterialEditor::CreateSlot_Preview()
+TSharedRef<SWidget> SDMMaterialEditor::CreateSlot_Preview()
 {
+	if (IsDynamicModel())
+	{
+		return SNew(SOverlay)
+			+ SOverlay::Slot()
+			[
+				SNew(SDMMaterialPreview, SharedThis(this), GetMaterialModelBase())
+			]
+			+ SOverlay::Slot()
+			.HAlign(EHorizontalAlignment::HAlign_Left)
+			.VAlign(EVerticalAlignment::VAlign_Bottom)
+			.Padding(3.f, 2.f)
+			[
+				SNew(STextBlock)
+					.Font(FAppStyle::GetFontStyle("TinyText"))
+					.Text(LOCTEXT("DynamicMaterial", "Dynamic"))
+					.ShadowColorAndOpacity(FLinearColor::Black)
+					.ShadowOffset(FVector2D(1.0))
+			];
+	}
+
 	return SNew(SDMMaterialPreview, SharedThis(this), GetMaterialModelBase());
 }
 
diff --git a/Engine/Plugins/Experimental/DynamicMaterial/Source/DynamicMaterialEditor/Private/UI/Widgets/SDMMaterialEditor.h b/Engine/Plugins/Experimental/DynamicMaterial/Source/DynamicMaterialEditor/Private/UI/Widgets/SDMMaterialEditor.h
--- a/Engine/Plugins/Experimental/DynamicMaterial/Source/DynamicMaterialEditor/Private/UI/Widgets/SDMMaterialEditor.h
+++ b/Engine/Plugins/Experimental/DynamicMaterial/Source/DynamicMaterialEditor/Private/UI/Widgets/SDMMaterialEditor.h
@@ -164,7 +164,7 @@ protected:
 	TDMWidgetSlot<SWidget> ContentSlot;
 	TDMWidgetSlot<SDMToolBar> ToolBarSlot;
 	TDMWidgetSlot<SWidget> MainSlot;
-	TDMWidgetSlot<SDMMaterialPreview> MaterialPreviewSlot;
+	TDMWidgetSlot<SWidget> MaterialPreviewSlot;
 	TDMWidgetSlot<SDMMaterialPropertySelector> PropertySelectorSlot;
 	TDMWidgetSlot<SDMMaterialGlobalSettingsEditor> GlobalSettingsEditorSlot;
 	TDMWidgetSlot<SDMMaterialProperties> MaterialPropertiesSlot;
@@ -179,9 +179,9 @@ protected:
 	TSharedRef<FUICommandList> CommandList;
 	TSharedRef<FDMPreviewMaterialManager> PreviewMaterialManager;
 	TSharedPtr<SDockTab> MaterialPreviewTab;
-	TDMWidgetSlot<SDMMaterialPreview> MaterialPreviewTabSlot;
+	TDMWidgetSlot<SWidget> MaterialPreviewTabSlot;
 	TSharedPtr<IToolTip> MaterialPreviewToolTip;
-	TDMWidgetSlot<SDMMaterialPreview> MaterialPreviewToolTipSlot;
+	TDMWidgetSlot<SWidget> MaterialPreviewToolTipSlot;
 
 	EDMMaterialEditorMode EditMode;
 	EDMMaterialPropertyType SelectedMaterialProperty;
@@ -233,7 +233,7 @@ protected:
 
 	TSharedRef<SDMMaterialProperties> CreateSlot_MaterialProperties();
 
-	TSharedRef<SDMMaterialPreview> CreateSlot_Preview();
+	TSharedRef<SWidget> CreateSlot_Preview();
 
 	TSharedRef<SDMMaterialPropertySelector> CreateSlot_PropertySelector();
 
