Material Designer
- The material preview popout tab no longer has an option to open the material preview popout tab.

diff --git a/Engine/Plugins/Experimental/DynamicMaterial/Source/DynamicMaterialEditor/Private/UI/Widgets/Editor/SDMMaterialPreview.cpp b/Engine/Plugins/Experimental/DynamicMaterial/Source/DynamicMaterialEditor/Private/UI/Widgets/Editor/SDMMaterialPreview.cpp
--- a/Engine/Plugins/Experimental/DynamicMaterial/Source/DynamicMaterialEditor/Private/UI/Widgets/Editor/SDMMaterialPreview.cpp
+++ b/Engine/Plugins/Experimental/DynamicMaterial/Source/DynamicMaterialEditor/Private/UI/Widgets/Editor/SDMMaterialPreview.cpp
@@ -477,18 +477,21 @@ TSharedRef<SWidget> SDMMaterialPreview::GenerateToolbarMenu()
 		FToolMenuSection& SettingsSection = Menu->AddSection(TEXT("Settings"), LOCTEXT("Settings", "Settings"));
 		SettingsSection.AddEntry(FToolMenuEntry::InitMenuEntry(MaterialEditorCommands.TogglePreviewBackground));
 
-		FToolMenuSection& ActionsSection = Menu->AddSection(TEXT("Actions"), LOCTEXT("Actions", "Actions"));
-
-		FUIAction OpenPreviewTabAction;
-		OpenPreviewTabAction.ExecuteAction.BindSP(this, &SDMMaterialPreview::OpenMaterialPreviewTab);
-
-		ActionsSection.AddEntry(FToolMenuEntry::InitMenuEntry(
-			TEXT("PopoutMaterialPreviewTab"),
-			LOCTEXT("OpenPreview", "Open Preview"),
-			LOCTEXT("OpenPreviewToolTip", "Open a tab with a preview of the material."),
-			TAttribute<FSlateIcon>(),
-			OpenPreviewTabAction
-		));		
+		if (!bIsPopout)
+		{
+			FToolMenuSection& ActionsSection = Menu->AddSection(TEXT("Actions"), LOCTEXT("Actions", "Actions"));
+
+			FUIAction OpenPreviewTabAction;
+			OpenPreviewTabAction.ExecuteAction.BindSP(this, &SDMMaterialPreview::OpenMaterialPreviewTab);
+
+			ActionsSection.AddEntry(FToolMenuEntry::InitMenuEntry(
+				TEXT("PopoutMaterialPreviewTab"),
+				LOCTEXT("OpenPreview", "Open Preview"),
+				LOCTEXT("OpenPreviewToolTip", "Open a tab with a preview of the material."),
+				TAttribute<FSlateIcon>(),
+				OpenPreviewTabAction
+			));
+		}
 	}
 
 	FToolMenuContext Context;
diff --git a/Engine/Plugins/Experimental/DynamicMaterial/Source/DynamicMaterialEditor/Private/UI/Widgets/Editor/SDMMaterialPreview.h b/Engine/Plugins/Experimental/DynamicMaterial/Source/DynamicMaterialEditor/Private/UI/Widgets/Editor/SDMMaterialPreview.h
--- a/Engine/Plugins/Experimental/DynamicMaterial/Source/DynamicMaterialEditor/Private/UI/Widgets/Editor/SDMMaterialPreview.h
+++ b/Engine/Plugins/Experimental/DynamicMaterial/Source/DynamicMaterialEditor/Private/UI/Widgets/Editor/SDMMaterialPreview.h
@@ -27,14 +27,17 @@ class SDMMaterialPreview : public SEditorViewport, public FGCObject
 
 	SLATE_BEGIN_ARGS(SDMMaterialPreview)
 		: _ShowMenu(true)
+		, _IsPopout(false)
 		{}
 		SLATE_ARGUMENT(bool, ShowMenu)
+		SLATE_ARGUMENT(bool, IsPopout)
 	SLATE_END_ARGS()
 
 public:
 	virtual ~SDMMaterialPreview() override;
 
-	void Construct(const FArguments& InArgs, const TSharedRef<SDMMaterialEditor>& InEditorWidget, UDynamicMaterialModelBase* InMaterialModelBase);
+	void Construct(const FArguments& InArgs, const TSharedRef<SDMMaterialEditor>& InEditorWidget, 
+		UDynamicMaterialModelBase* InMaterialModelBase);
 
 	//~ Begin SEditorViewport
 	virtual TSharedPtr<SWidget> MakeViewportToolbar() override;
@@ -48,6 +51,7 @@ public:
 protected:
 	TWeakPtr<SDMMaterialEditor> EditorWidgetWeak;
 	bool bShowMenu;
+	bool bIsPopout;
 
 	TSharedPtr<FDMMaterialPreviewViewportClient> EditorViewportClient;
 	TSharedPtr<FAdvancedPreviewScene> PreviewScene;
diff --git a/Engine/Plugins/Experimental/DynamicMaterial/Source/DynamicMaterialEditor/Private/UI/Widgets/SDMMaterialEditor.cpp b/Engine/Plugins/Experimental/DynamicMaterial/Source/DynamicMaterialEditor/Private/UI/Widgets/SDMMaterialEditor.cpp
--- a/Engine/Plugins/Experimental/DynamicMaterial/Source/DynamicMaterialEditor/Private/UI/Widgets/SDMMaterialEditor.cpp
+++ b/Engine/Plugins/Experimental/DynamicMaterial/Source/DynamicMaterialEditor/Private/UI/Widgets/SDMMaterialEditor.cpp
@@ -429,6 +429,7 @@ void SDMMaterialEditor::OpenMaterialPreviewTab()
 		Wrapper, 
 		0, 
 		SNew(SDMMaterialPreview, SharedThis(this), MaterialModelBase)
+		.IsPopout(true)
 	);
 
 	MaterialPreviewTab->SetContent(Wrapper);
@@ -1038,10 +1039,10 @@ TSharedRef<SWidget> SDMMaterialEditor::CreateSlot_Preview()
 			.Padding(3.f, 2.f)
 			[
 				SNew(STextBlock)
-					.Font(FAppStyle::GetFontStyle("TinyText"))
-					.Text(LOCTEXT("DynamicMaterial", "Dynamic"))
-					.ShadowColorAndOpacity(FLinearColor::Black)
-					.ShadowOffset(FVector2D(1.0))
+				.Font(FAppStyle::GetFontStyle("TinyText"))
+				.Text(LOCTEXT("DynamicMaterial", "Dynamic"))
+				.ShadowColorAndOpacity(FLinearColor::Black)
+				.ShadowOffset(FVector2D(1.0))
 			];
 	}
 
