Material Designer
- Horde localization issue 573472

diff --git a/Engine/Plugins/Experimental/DynamicMaterial/Source/DynamicMaterialEditor/Private/UI/Menus/DMToolBarMenus.cpp b/Engine/Plugins/Experimental/DynamicMaterial/Source/DynamicMaterialEditor/Private/UI/Menus/DMToolBarMenus.cpp
--- a/Engine/Plugins/Experimental/DynamicMaterial/Source/DynamicMaterialEditor/Private/UI/Menus/DMToolBarMenus.cpp
+++ b/Engine/Plugins/Experimental/DynamicMaterial/Source/DynamicMaterialEditor/Private/UI/Menus/DMToolBarMenus.cpp
@@ -401,7 +401,7 @@ void FDMToolBarMenus::SnapshotMaterial(TWeakObjectPtr<UDynamicMaterialModelBase>
 	const FString PathStr = CurrentPath.HasInternalPath() ? CurrentPath.GetInternalPathString() : "/Game";
 
 	FSaveAssetDialogConfig SaveAssetDialogConfig;
-	SaveAssetDialogConfig.DialogTitleOverride = LOCTEXT("SaveAssetDialogTitle", "Export Material To");
+	SaveAssetDialogConfig.DialogTitleOverride = LOCTEXT("ExportMaterialTo", "Export Material To");
 	SaveAssetDialogConfig.DefaultPath = PathStr;
 	SaveAssetDialogConfig.ExistingAssetPolicy = ESaveAssetDialogExistingAssetPolicy::Disallow;
 	SaveAssetDialogConfig.DefaultAssetName = TEXT("T_MD_") + Material->GetName();
