Material Designer
- You can now press keys 1-0 to select layers.
- Pressing the key again will toggle between the stages.

diff --git a/Engine/Plugins/Experimental/DynamicMaterial/Source/DynamicMaterialEditor/Private/DynamicMaterialEditorCommands.cpp b/Engine/Plugins/Experimental/DynamicMaterial/Source/DynamicMaterialEditor/Private/DynamicMaterialEditorCommands.cpp
--- a/Engine/Plugins/Experimental/DynamicMaterial/Source/DynamicMaterialEditor/Private/DynamicMaterialEditorCommands.cpp
+++ b/Engine/Plugins/Experimental/DynamicMaterial/Source/DynamicMaterialEditor/Private/DynamicMaterialEditorCommands.cpp
@@ -75,6 +75,80 @@ void FDynamicMaterialEditorCommands::RegisterCommands()
 		, "Navigates back in the page history."
 		, EUserInterfaceActionType::Button
 		, FInputChord(EKeys::ThumbMouseButton));
+
+	TSharedPtr<FUICommandInfo> SelectLayer1;
+	TSharedPtr<FUICommandInfo> SelectLayer2;
+	TSharedPtr<FUICommandInfo> SelectLayer3;
+	TSharedPtr<FUICommandInfo> SelectLayer4;
+	TSharedPtr<FUICommandInfo> SelectLayer5;
+	TSharedPtr<FUICommandInfo> SelectLayer6;
+	TSharedPtr<FUICommandInfo> SelectLayer7;
+	TSharedPtr<FUICommandInfo> SelectLayer8;
+	TSharedPtr<FUICommandInfo> SelectLayer9;
+	TSharedPtr<FUICommandInfo> SelectLayer10;
+
+	UI_COMMAND(SelectLayer1
+		, "Select Layer 1"
+		, "Select the first layer (from the top of the stack)."
+		, EUserInterfaceActionType::Button
+		, FInputChord(EKeys::One));
+
+	UI_COMMAND(SelectLayer2
+		, "Select Layer 2"
+		, "Select the second layer (from the top of the stack)."
+		, EUserInterfaceActionType::Button
+		, FInputChord(EKeys::Two));
+
+	UI_COMMAND(SelectLayer3
+		, "Select Layer 3"
+		, "Select the third layer (from the top of the stack)."
+		, EUserInterfaceActionType::Button
+		, FInputChord(EKeys::Three));
+
+	UI_COMMAND(SelectLayer4
+		, "Select Layer 4"
+		, "Select the fourth layer (from the top of the stack)."
+		, EUserInterfaceActionType::Button
+		, FInputChord(EKeys::Four));
+
+	UI_COMMAND(SelectLayer5
+		, "Select Layer 5"
+		, "Select the fifth layer (from the top of the stack)."
+		, EUserInterfaceActionType::Button
+		, FInputChord(EKeys::Five));
+
+	UI_COMMAND(SelectLayer6
+		, "Select Layer 6"
+		, "Select the sixth layer (from the top of the stack)."
+		, EUserInterfaceActionType::Button
+		, FInputChord(EKeys::Six));
+
+	UI_COMMAND(SelectLayer7
+		, "Select Layer 7"
+		, "Select the seventh layer (from the top of the stack)."
+		, EUserInterfaceActionType::Button
+		, FInputChord(EKeys::Seven));
+
+	UI_COMMAND(SelectLayer8
+		, "Select Layer 8"
+		, "Select the eighth layer (from the top of the stack)."
+		, EUserInterfaceActionType::Button
+		, FInputChord(EKeys::Eight));
+
+	UI_COMMAND(SelectLayer9
+		, "Select Layer 9"
+		, "Select the ninth layer (from the top of the stack)."
+		, EUserInterfaceActionType::Button
+		, FInputChord(EKeys::Nine));
+
+	UI_COMMAND(SelectLayer10
+		, "Select Layer 10"
+		, "Select the tenth layer (from the top of the stack)."
+		, EUserInterfaceActionType::Button
+		, FInputChord(EKeys::Zero));
+
+	SelectLayers = {SelectLayer1, SelectLayer2, SelectLayer3, SelectLayer4, SelectLayer5,
+		SelectLayer6, SelectLayer7, SelectLayer8, SelectLayer9, SelectLayer10};
 }
 
 #undef LOCTEXT_NAMESPACE
diff --git a/Engine/Plugins/Experimental/DynamicMaterial/Source/DynamicMaterialEditor/Private/DynamicMaterialEditorCommands.h b/Engine/Plugins/Experimental/DynamicMaterial/Source/DynamicMaterialEditor/Private/DynamicMaterialEditorCommands.h
--- a/Engine/Plugins/Experimental/DynamicMaterial/Source/DynamicMaterialEditor/Private/DynamicMaterialEditorCommands.h
+++ b/Engine/Plugins/Experimental/DynamicMaterial/Source/DynamicMaterialEditor/Private/DynamicMaterialEditorCommands.h
@@ -25,4 +25,7 @@ public:
 
 	TSharedPtr<FUICommandInfo> NavigateForward;
 	TSharedPtr<FUICommandInfo> NavigateBack;
+
+	// Individual commands are initialised as SelectLayerX
+	TArray<TSharedPtr<FUICommandInfo>, TInlineAllocator<10>> SelectLayers;
 };
diff --git a/Engine/Plugins/Experimental/DynamicMaterial/Source/DynamicMaterialEditor/Private/UI/Widgets/Editor/SDMMaterialSlotEditor.cpp b/Engine/Plugins/Experimental/DynamicMaterial/Source/DynamicMaterialEditor/Private/UI/Widgets/Editor/SDMMaterialSlotEditor.cpp
--- a/Engine/Plugins/Experimental/DynamicMaterial/Source/DynamicMaterialEditor/Private/UI/Widgets/Editor/SDMMaterialSlotEditor.cpp
+++ b/Engine/Plugins/Experimental/DynamicMaterial/Source/DynamicMaterialEditor/Private/UI/Widgets/Editor/SDMMaterialSlotEditor.cpp
@@ -32,6 +32,7 @@
 #include "Model/DynamicMaterialModelDynamic.h"
 #include "Model/DynamicMaterialModelEditorOnlyData.h"
 #include "SAssetDropTarget.h"
+#include "SDMMaterialComponentEditor.h"
 #include "UI/DragDrop/DMLayerEffectsDragDropOperation.h"
 #include "UI/DragDrop/DMSlotLayerDragDropOperation.h"
 #include "UI/Menus/DMMaterialSlotLayerAddEffectMenus.h"
@@ -361,6 +362,78 @@ void SDMMaterialSlotEditor::DeleteSelectedLayer()
 	Slot->RemoveLayer(SelectedLayer);
 }
 
+bool SDMMaterialSlotEditor::SelectLayer_CanExecute(int32 InIndex) const
+{
+	return LayerViewSlot.HasWidget() && LayerViewSlot->GetItems().IsValidIndex(InIndex);
+}
+
+void SDMMaterialSlotEditor::SelectLayer_Execute(int32 InIndex)
+{
+	if (!SelectLayer_CanExecute(InIndex))
+	{
+		return;
+	}
+
+	TSharedPtr<SDMMaterialEditor> EditorWidget = GetEditorWidget();
+
+	if (!EditorWidget.IsValid())
+	{
+		return;
+	}
+
+	UDMMaterialSlot* Slot = GetSlot();
+
+	if (!Slot)
+	{
+		return;
+	}
+
+	TArray<UDMMaterialLayerObject*> Layers = Slot->GetLayers();
+
+	// Layer order is reversed.
+	const int32 LayerIndex = Layers.Num() - 1 - InIndex;
+
+	if (!Layers.IsValidIndex(InIndex))
+	{
+		return;
+	}
+
+	UDMMaterialLayerObject* Layer = Layers[LayerIndex];
+
+	// Switch between stages
+	if (LayerViewSlot->GetSelectedLayer() == Layer)
+	{
+		if (TSharedPtr<SDMMaterialComponentEditor> ComponentEditorWidget = EditorWidget->GetComponentEditorWidget())
+		{
+			UDMMaterialStage* BaseStage = Layer->GetFirstEnabledStage(EDMMaterialLayerStage::Base);
+			UDMMaterialStage* MaskStage = Layer->GetFirstEnabledStage(EDMMaterialLayerStage::Mask);
+
+			if (MaskStage && ComponentEditorWidget->GetObject() == BaseStage)
+			{
+				EditorWidget->EditComponent(MaskStage);
+			}
+			else if (BaseStage && ComponentEditorWidget->GetObject() == MaskStage)
+			{
+				EditorWidget->EditComponent(BaseStage);
+			}
+		}
+	}
+	// Select new layer
+	else
+	{
+		LayerViewSlot->SetSelectedLayer(Layer);
+
+		if (UDMMaterialStage* Stage = Layer->GetFirstEnabledStage(EDMMaterialLayerStage::All))
+		{
+			EditorWidget->EditComponent(Stage, /* Force Refersh */ false);
+		}
+		else
+		{
+			EditorWidget->EditComponent(nullptr);
+		}
+	}
+}
+
 TSharedRef<SDMMaterialSlotLayerView> SDMMaterialSlotEditor::GetLayerView() const
 {
 	return *LayerViewSlot;
diff --git a/Engine/Plugins/Experimental/DynamicMaterial/Source/DynamicMaterialEditor/Private/UI/Widgets/Editor/SDMMaterialSlotEditor.h b/Engine/Plugins/Experimental/DynamicMaterial/Source/DynamicMaterialEditor/Private/UI/Widgets/Editor/SDMMaterialSlotEditor.h
--- a/Engine/Plugins/Experimental/DynamicMaterial/Source/DynamicMaterialEditor/Private/UI/Widgets/Editor/SDMMaterialSlotEditor.h
+++ b/Engine/Plugins/Experimental/DynamicMaterial/Source/DynamicMaterialEditor/Private/UI/Widgets/Editor/SDMMaterialSlotEditor.h
@@ -61,6 +61,9 @@ public:
 	bool CanDeleteSelectedLayer() const;
 	void DeleteSelectedLayer();
 
+	bool SelectLayer_CanExecute(int32 InIndex) const;
+	void SelectLayer_Execute(int32 InIndex);
+
 	/** Slots */
 	TSharedRef<SDMMaterialSlotLayerView> GetLayerView() const;
 
diff --git a/Engine/Plugins/Experimental/DynamicMaterial/Source/DynamicMaterialEditor/Private/UI/Widgets/SDMMaterialEditor.cpp b/Engine/Plugins/Experimental/DynamicMaterial/Source/DynamicMaterialEditor/Private/UI/Widgets/SDMMaterialEditor.cpp
--- a/Engine/Plugins/Experimental/DynamicMaterial/Source/DynamicMaterialEditor/Private/UI/Widgets/SDMMaterialEditor.cpp
+++ b/Engine/Plugins/Experimental/DynamicMaterial/Source/DynamicMaterialEditor/Private/UI/Widgets/SDMMaterialEditor.cpp
@@ -628,13 +628,13 @@ void SDMMaterialEditor::BindCommands(SDMMaterialSlotEditor* InSlotEditor)
 	);
 
 	CommandList->MapAction(
-		FDynamicMaterialEditorCommands::Get().AddDefaultLayer,
+		DMEditorCommands.AddDefaultLayer,
 		FExecuteAction::CreateSP(InSlotEditor, &SDMMaterialSlotEditor::AddNewLayer),
 		FCanExecuteAction::CreateSP(InSlotEditor, &SDMMaterialSlotEditor::CanAddNewLayer)
 	);
 
 	CommandList->MapAction(
-		FDynamicMaterialEditorCommands::Get().InsertDefaultLayerAbove,
+		DMEditorCommands.InsertDefaultLayerAbove,
 		FExecuteAction::CreateSP(InSlotEditor, &SDMMaterialSlotEditor::InsertNewLayer),
 		FCanExecuteAction::CreateSP(InSlotEditor, &SDMMaterialSlotEditor::CanInsertNewLayer)
 	);
@@ -668,6 +668,15 @@ void SDMMaterialEditor::BindCommands(SDMMaterialSlotEditor* InSlotEditor)
 		FExecuteAction::CreateSP(InSlotEditor, &SDMMaterialSlotEditor::DeleteSelectedLayer),
 		FCanExecuteAction::CreateSP(InSlotEditor, &SDMMaterialSlotEditor::CanDeleteSelectedLayer)
 	);
+
+	for (int32 LayerIndex = 0; LayerIndex < DMEditorCommands.SelectLayers.Num(); ++LayerIndex)
+	{
+		CommandList->MapAction(
+			DMEditorCommands.SelectLayers[LayerIndex],
+			FExecuteAction::CreateSP(InSlotEditor, &SDMMaterialSlotEditor::SelectLayer_Execute, LayerIndex),
+			FCanExecuteAction::CreateSP(InSlotEditor, &SDMMaterialSlotEditor::SelectLayer_CanExecute, LayerIndex)
+		);
+	}
 }
 
 bool SDMMaterialEditor::IsPropertyValidForModel(EDMMaterialPropertyType InProperty) const
