Material Designer
- Fixed texture sets drops calling the wrong function.

diff --git a/Engine/Plugins/Experimental/DynamicMaterial/Source/DynamicMaterialEditor/Private/UI/Widgets/SDMMaterialDesigner.cpp b/Engine/Plugins/Experimental/DynamicMaterial/Source/DynamicMaterialEditor/Private/UI/Widgets/SDMMaterialDesigner.cpp
--- a/Engine/Plugins/Experimental/DynamicMaterial/Source/DynamicMaterialEditor/Private/UI/Widgets/SDMMaterialDesigner.cpp
+++ b/Engine/Plugins/Experimental/DynamicMaterial/Source/DynamicMaterialEditor/Private/UI/Widgets/SDMMaterialDesigner.cpp
@@ -442,7 +442,7 @@ void SDMMaterialDesigner::OnAssetsDropped(const FDragDropEvent& InDragDropEvent,
 		{
 			if (Content.IsValid() && Content->GetWidgetClass().GetWidgetType() == SDMMaterialEditor::StaticWidgetClass().GetWidgetType())
 			{
-				StaticCastSharedPtr<SDMMaterialEditor>(Content)->HandleDrop_CreateTextureSet({Asset});
+				StaticCastSharedPtr<SDMMaterialEditor>(Content)->HandleDrop_TextureSet(Cast<UDMTextureSet>(Asset.GetAsset()));
 				return;
 			}
 		}
