Material Designer
- The V+1-0 key combinations can now be used to alter the opacity of the selected layer.

diff --git a/Engine/Plugins/Experimental/DynamicMaterial/Source/DynamicMaterialEditor/Private/DynamicMaterialEditorCommands.cpp b/Engine/Plugins/Experimental/DynamicMaterial/Source/DynamicMaterialEditor/Private/DynamicMaterialEditorCommands.cpp
--- a/Engine/Plugins/Experimental/DynamicMaterial/Source/DynamicMaterialEditor/Private/DynamicMaterialEditorCommands.cpp
+++ b/Engine/Plugins/Experimental/DynamicMaterial/Source/DynamicMaterialEditor/Private/DynamicMaterialEditorCommands.cpp
@@ -149,6 +149,91 @@ void FDynamicMaterialEditorCommands::RegisterCommands()
 
 	SelectLayers = {SelectLayer1, SelectLayer2, SelectLayer3, SelectLayer4, SelectLayer5,
 		SelectLayer6, SelectLayer7, SelectLayer8, SelectLayer9, SelectLayer10};
+
+	TSharedPtr<FUICommandInfo> SetOpacity10;
+	TSharedPtr<FUICommandInfo> SetOpacity20;
+	TSharedPtr<FUICommandInfo> SetOpacity30;
+	TSharedPtr<FUICommandInfo> SetOpacity40;
+	TSharedPtr<FUICommandInfo> SetOpacity50;
+	TSharedPtr<FUICommandInfo> SetOpacity60;
+	TSharedPtr<FUICommandInfo> SetOpacity70;
+	TSharedPtr<FUICommandInfo> SetOpacity80;
+	TSharedPtr<FUICommandInfo> SetOpacity90;
+	TSharedPtr<FUICommandInfo> SetOpacity100;
+
+	UI_COMMAND(SetOpacity10
+		, "Set Opacity to 10%"
+		, "Sets the opacity of the selected layer to 10%."
+		, EUserInterfaceActionType::Button
+		, FInputChord());
+
+	UI_COMMAND(SetOpacity20
+		, "Set Opacity to 20%"
+		, "Sets the opacity of the selected layer to 20%."
+		, EUserInterfaceActionType::Button
+		, FInputChord());
+
+	UI_COMMAND(SetOpacity30
+		, "Set Opacity to 30%"
+		, "Sets the opacity of the selected layer to 30%."
+		, EUserInterfaceActionType::Button
+		, FInputChord());
+
+	UI_COMMAND(SetOpacity40
+		, "Set Opacity to 40%"
+		, "Sets the opacity of the selected layer to 40%."
+		, EUserInterfaceActionType::Button
+		, FInputChord());
+
+	UI_COMMAND(SetOpacity50
+		, "Set Opacity to 50%"
+		, "Sets the opacity of the selected layer to 50%."
+		, EUserInterfaceActionType::Button
+		, FInputChord());
+
+	UI_COMMAND(SetOpacity60
+		, "Set Opacity to 60%"
+		, "Sets the opacity of the selected layer to 60%."
+		, EUserInterfaceActionType::Button
+		, FInputChord());
+
+	UI_COMMAND(SetOpacity70
+		, "Set Opacity to 70%"
+		, "Sets the opacity of the selected layer to 70%."
+		, EUserInterfaceActionType::Button
+		, FInputChord());
+
+	UI_COMMAND(SetOpacity80
+		, "Set Opacity to 80%"
+		, "Sets the opacity of the selected layer to 80%."
+		, EUserInterfaceActionType::Button
+		, FInputChord());
+
+	UI_COMMAND(SetOpacity90
+		, "Set Opacity to 90%"
+		, "Sets the opacity of the selected layer to 90%."
+		, EUserInterfaceActionType::Button
+		, FInputChord());
+
+	UI_COMMAND(SetOpacity100
+		, "Set Opacity to 100%"
+		, "Sets the opacity of the selected layer to 100%."
+		, EUserInterfaceActionType::Button
+		, FInputChord());
+
+	/** Cannot find the key directly as then it will activate without the fake modifier key, 'V'. */
+	SetOpacities = {
+		{EKeys::One,   {0.1f, SetOpacity10.ToSharedRef()}},
+		{EKeys::Two,   {0.2f, SetOpacity20.ToSharedRef()}},
+		{EKeys::Three, {0.3f, SetOpacity30.ToSharedRef()}},
+		{EKeys::Four,  {0.4f, SetOpacity40.ToSharedRef()}},
+		{EKeys::Five,  {0.5f, SetOpacity50.ToSharedRef()}},
+		{EKeys::Six,   {0.6f, SetOpacity60.ToSharedRef()}},
+		{EKeys::Seven, {0.7f, SetOpacity70.ToSharedRef()}},
+		{EKeys::Eight, {0.8f, SetOpacity80.ToSharedRef()}},
+		{EKeys::Nine,  {0.9f, SetOpacity90.ToSharedRef()}},
+		{EKeys::Zero,  {1.0f, SetOpacity100.ToSharedRef()}}
+	};
 }
 
 #undef LOCTEXT_NAMESPACE
diff --git a/Engine/Plugins/Experimental/DynamicMaterial/Source/DynamicMaterialEditor/Private/DynamicMaterialEditorCommands.h b/Engine/Plugins/Experimental/DynamicMaterial/Source/DynamicMaterialEditor/Private/DynamicMaterialEditorCommands.h
--- a/Engine/Plugins/Experimental/DynamicMaterial/Source/DynamicMaterialEditor/Private/DynamicMaterialEditorCommands.h
+++ b/Engine/Plugins/Experimental/DynamicMaterial/Source/DynamicMaterialEditor/Private/DynamicMaterialEditorCommands.h
@@ -28,4 +28,13 @@ public:
 
 	// Individual commands are initialised as SelectLayerX
 	TArray<TSharedPtr<FUICommandInfo>, TInlineAllocator<10>> SelectLayers;
+
+	struct FOpacityCommand
+	{
+		const float Opacity;
+		TSharedRef<FUICommandInfo> Command;
+	};
+
+	// Individual commands are initialised as SetOpacityX
+	TMap<FKey, FOpacityCommand> SetOpacities;
 };
diff --git a/Engine/Plugins/Experimental/DynamicMaterial/Source/DynamicMaterialEditor/Private/UI/Widgets/Editor/SDMMaterialSlotEditor.cpp b/Engine/Plugins/Experimental/DynamicMaterial/Source/DynamicMaterialEditor/Private/UI/Widgets/Editor/SDMMaterialSlotEditor.cpp
--- a/Engine/Plugins/Experimental/DynamicMaterial/Source/DynamicMaterialEditor/Private/UI/Widgets/Editor/SDMMaterialSlotEditor.cpp
+++ b/Engine/Plugins/Experimental/DynamicMaterial/Source/DynamicMaterialEditor/Private/UI/Widgets/Editor/SDMMaterialSlotEditor.cpp
@@ -32,13 +32,13 @@
 #include "Model/DynamicMaterialModelDynamic.h"
 #include "Model/DynamicMaterialModelEditorOnlyData.h"
 #include "SAssetDropTarget.h"
-#include "SDMMaterialComponentEditor.h"
 #include "UI/DragDrop/DMLayerEffectsDragDropOperation.h"
 #include "UI/DragDrop/DMSlotLayerDragDropOperation.h"
 #include "UI/Menus/DMMaterialSlotLayerAddEffectMenus.h"
 #include "UI/Menus/DMMaterialSlotLayerMenus.h"
 #include "UI/Utils/DMDropTargetPrivateSetter.h"
 #include "UI/Utils/DMWidgetStatics.h"
+#include "UI/Widgets/Editor/SDMMaterialComponentEditor.h"
 #include "UI/Widgets/Editor/SlotEditor/SDMMaterialSlotLayerView.h"
 #include "UI/Widgets/SDMMaterialEditor.h"
 #include "Utils/DMMaterialSlotFunctionLibrary.h"
@@ -434,6 +434,19 @@ void SDMMaterialSlotEditor::SelectLayer_Execute(int32 InIndex)
 	}
 }
 
+bool SDMMaterialSlotEditor::SetOpacity_CanExecute()
+{
+	return LayerOpacityValueWeak.IsValid();
+}
+
+void SDMMaterialSlotEditor::SetOpacity_Execute(float InOpacity)
+{
+	if (UDMMaterialValueFloat1* OpacityValue = LayerOpacityValueWeak.Get())
+	{
+		OpacityValue->SetValue(InOpacity);
+	}
+}
+
 TSharedRef<SDMMaterialSlotLayerView> SDMMaterialSlotEditor::GetLayerView() const
 {
 	return *LayerViewSlot;
@@ -612,6 +625,7 @@ TSharedRef<SWidget> SDMMaterialSlotEditor::CreateSlot_SlotSettings()
 
 TSharedRef<SWidget> SDMMaterialSlotEditor::CreateSlot_LayerOpacity()
 {
+	LayerOpacityValueWeak.Reset();
 	LayerOpacityItem = nullptr;
 
 	if (LayerViewSlot.IsValid())
@@ -624,6 +638,8 @@ TSharedRef<SWidget> SDMMaterialSlotEditor::CreateSlot_LayerOpacity()
 				{
 					if (UDMMaterialValueFloat1* OpacityValue = Cast<UDMMaterialValueFloat1>(SelectedOpacityStageInputValue->GetValue()))
 					{
+						LayerOpacityValueWeak = OpacityValue;
+
 						UWorld* World = OpacityValue->GetWorld();
 						TSharedPtr<IDetailKeyframeHandler> KeyframeHandler = nullptr;
 
diff --git a/Engine/Plugins/Experimental/DynamicMaterial/Source/DynamicMaterialEditor/Private/UI/Widgets/Editor/SDMMaterialSlotEditor.h b/Engine/Plugins/Experimental/DynamicMaterial/Source/DynamicMaterialEditor/Private/UI/Widgets/Editor/SDMMaterialSlotEditor.h
--- a/Engine/Plugins/Experimental/DynamicMaterial/Source/DynamicMaterialEditor/Private/UI/Widgets/Editor/SDMMaterialSlotEditor.h
+++ b/Engine/Plugins/Experimental/DynamicMaterial/Source/DynamicMaterialEditor/Private/UI/Widgets/Editor/SDMMaterialSlotEditor.h
@@ -14,6 +14,7 @@ class SHorizontalBox;
 class SVerticalBox;
 class UDMMaterialLayerObject;
 class UDMMaterialSlot;
+class UDMMaterialValueFloat1;
 class UDMTextureSet;
 class UMaterialFunctionInterface;
 class UTexture;
@@ -64,6 +65,9 @@ public:
 	bool SelectLayer_CanExecute(int32 InIndex) const;
 	void SelectLayer_Execute(int32 InIndex);
 
+	bool SetOpacity_CanExecute();
+	void SetOpacity_Execute(float InOpacity);
+
 	/** Slots */
 	TSharedRef<SDMMaterialSlotLayerView> GetLayerView() const;
 
@@ -83,6 +87,7 @@ protected:
 	TDMWidgetSlot<SDMMaterialSlotLayerView> LayerViewSlot;
 	TDMWidgetSlot<SWidget> LayerSettingsSlot;
 
+	TWeakObjectPtr<UDMMaterialValueFloat1> LayerOpacityValueWeak;
 	TSharedPtr<ICustomDetailsViewItem> LayerOpacityItem;
 
 	TSharedRef<SWidget> CreateSlot_Container();
diff --git a/Engine/Plugins/Experimental/DynamicMaterial/Source/DynamicMaterialEditor/Private/UI/Widgets/SDMMaterialEditor.cpp b/Engine/Plugins/Experimental/DynamicMaterial/Source/DynamicMaterialEditor/Private/UI/Widgets/SDMMaterialEditor.cpp
--- a/Engine/Plugins/Experimental/DynamicMaterial/Source/DynamicMaterialEditor/Private/UI/Widgets/SDMMaterialEditor.cpp
+++ b/Engine/Plugins/Experimental/DynamicMaterial/Source/DynamicMaterialEditor/Private/UI/Widgets/SDMMaterialEditor.cpp
@@ -11,6 +11,7 @@
 #include "DynamicMaterialEditorSettings.h"
 #include "DynamicMaterialModule.h"
 #include "Engine/Texture.h"
+#include "Framework/Application/IInputProcessor.h"
 #include "Framework/Application/SlateApplication.h"
 #include "Framework/Commands/GenericCommands.h"
 #include "Materials/Material.h"
@@ -44,6 +45,67 @@
 
 #define LOCTEXT_NAMESPACE "SDMMaterialEditor"
 
+/**
+ * This is used to track a key, similar to how modifier keys are tracked by the engine...
+ * because non-modifier keys are not tracked.
+ */
+class FDMKeyTracker : public IInputProcessor
+{
+public:
+	FDMKeyTracker(const FKey& InTrackedKey)
+		: TrackedKey(InTrackedKey)
+		, bKeyDown(false)
+	{
+		
+	}
+
+	const FKey& GetTrackedKey() const
+	{
+		return TrackedKey;
+	}
+
+	bool IsKeyDown() const
+	{
+		return bKeyDown;
+	}
+
+	//~ Begin IInputProcessor
+	virtual void Tick(const float InDeltaTime, FSlateApplication& InSlateApp, TSharedRef<ICursor> InCursor) override
+	{		
+	}
+
+	virtual bool HandleKeyDownEvent(FSlateApplication& InSlateApp, const FKeyEvent& InKeyEvent) override
+	{
+		if (InKeyEvent.GetKey() == TrackedKey)
+		{
+			bKeyDown = true;
+		}
+
+		return false;
+	}
+
+	/** Key up input */
+	virtual bool HandleKeyUpEvent(FSlateApplication& InSlateApp, const FKeyEvent& InKeyEvent) override
+	{
+		if (InKeyEvent.GetKey() == TrackedKey)
+		{
+			bKeyDown = false;
+		}
+
+		return false;
+	}
+
+	virtual const TCHAR* GetDebugName() const override
+	{
+		return TEXT("FDMKeyTracker");
+	}
+	//~ End IInputProcessor
+
+private:
+	const FKey& TrackedKey;
+	bool bKeyDown;
+};
+
 FDMMaterialEditorPage FDMMaterialEditorPage::Preview = {EDMMaterialEditorMode::MaterialPreview, EDMMaterialPropertyType::None};
 FDMMaterialEditorPage FDMMaterialEditorPage::GlobalSettings = {EDMMaterialEditorMode::GlobalSettings, EDMMaterialPropertyType::None};
 FDMMaterialEditorPage FDMMaterialEditorPage::Properties = {EDMMaterialEditorMode::Properties, EDMMaterialPropertyType::None};
@@ -76,6 +138,11 @@ SDMMaterialEditor::~SDMMaterialEditor()
 	CloseMaterialPreviewTab();
 	DestroyMaterialPreviewToolTip();
 
+	if (KeyTracker_V.IsValid() && FSlateApplication::IsInitialized())
+	{
+		FSlateApplication::Get().UnregisterInputPreProcessor(KeyTracker_V.ToSharedRef());
+	}
+
 	if (!FDynamicMaterialModule::AreUObjectsSafe())
 	{
 		return;
@@ -123,6 +190,9 @@ void SDMMaterialEditor::Construct(const FArguments& InArgs, const TSharedRef<SDM
 	{
 		Settings->GetOnSettingsChanged().AddSP(this, &SDMMaterialEditor::OnSettingsChanged);
 	}
+
+	KeyTracker_V = MakeShared<FDMKeyTracker>(EKeys::V);
+	FSlateApplication::Get().RegisterInputPreProcessor(KeyTracker_V);
 }
 
 TSharedPtr<SDMMaterialDesigner> SDMMaterialEditor::GetDesignerWidget() const
@@ -565,6 +635,12 @@ bool SDMMaterialEditor::SupportsKeyboardFocus() const
 
 FReply SDMMaterialEditor::OnKeyDown(const FGeometry& InMyGeometry, const FKeyEvent& InKeyEvent)
 {
+	// Cannot make a key bind that has 2 buttons, so hard code that here.
+	if (CheckOpacityInput(InKeyEvent))
+	{
+		return FReply::Handled();
+	}
+
 	if (CommandList->ProcessCommandBindings(InKeyEvent))
 	{
 		return FReply::Handled();
@@ -639,6 +715,18 @@ void SDMMaterialEditor::BindCommands(SDMMaterialSlotEditor* InSlotEditor)
 		FCanExecuteAction::CreateSP(InSlotEditor, &SDMMaterialSlotEditor::CanInsertNewLayer)
 	);
 
+	for (const TPair<FKey, FDynamicMaterialEditorCommands::FOpacityCommand>& OpacityCommandPair : DMEditorCommands.SetOpacities)
+	{
+		const float Opacity = OpacityCommandPair.Value.Opacity;
+		const TSharedPtr<FUICommandInfo>& OpacityCommand = OpacityCommandPair.Value.Command;
+
+		CommandList->MapAction(
+			OpacityCommand,
+			FExecuteAction::CreateSP(InSlotEditor, &SDMMaterialSlotEditor::SetOpacity_Execute, Opacity),
+			FCanExecuteAction::CreateSP(InSlotEditor, &SDMMaterialSlotEditor::SetOpacity_CanExecute)
+		);
+	}
+
 	CommandList->MapAction(
 		GenericCommands.Copy,
 		FExecuteAction::CreateSP(InSlotEditor, &SDMMaterialSlotEditor::CopySelectedLayer),
@@ -1224,4 +1312,22 @@ bool SDMMaterialEditor::NavigateBack_CanExecute()
 	return PageHistoryActive > 0;
 }
 
+bool SDMMaterialEditor::CheckOpacityInput(const FKeyEvent& InKeyEvent)
+{
+	if (!KeyTracker_V.IsValid() || !KeyTracker_V->IsKeyDown() || InKeyEvent.GetKey() == KeyTracker_V->GetTrackedKey())
+	{
+		return false;
+	}
+
+	const FDynamicMaterialEditorCommands& DMEditorCommands = FDynamicMaterialEditorCommands::Get();
+
+	if (const FDynamicMaterialEditorCommands::FOpacityCommand* OpacityCommandPair = DMEditorCommands.SetOpacities.Find(InKeyEvent.GetKey()))
+	{
+		const TSharedRef<FUICommandInfo>& OpacityCommand = OpacityCommandPair->Command;
+		return CommandList->TryExecuteAction(OpacityCommand);
+	}
+
+	return false;
+}
+
 #undef LOCTEXT_NAMESPACE
diff --git a/Engine/Plugins/Experimental/DynamicMaterial/Source/DynamicMaterialEditor/Private/UI/Widgets/SDMMaterialEditor.h b/Engine/Plugins/Experimental/DynamicMaterial/Source/DynamicMaterialEditor/Private/UI/Widgets/SDMMaterialEditor.h
--- a/Engine/Plugins/Experimental/DynamicMaterial/Source/DynamicMaterialEditor/Private/UI/Widgets/SDMMaterialEditor.h
+++ b/Engine/Plugins/Experimental/DynamicMaterial/Source/DynamicMaterialEditor/Private/UI/Widgets/SDMMaterialEditor.h
@@ -11,6 +11,7 @@
 #include "Misc/Optional.h"
 #include "UI/Utils/DMWidgetSlot.h"
 
+class FDMKeyTracker;
 class FDMPreviewMaterialManager;
 class FSlotBase;
 class FUICommandList;
@@ -177,6 +178,7 @@ protected:
 	TOptional<FDMObjectMaterialProperty> ObjectMaterialPropertyOpt;
 
 	TSharedRef<FUICommandList> CommandList;
+	TSharedPtr<FDMKeyTracker> KeyTracker_V;
 	TSharedRef<FDMPreviewMaterialManager> PreviewMaterialManager;
 	TSharedPtr<SDockTab> MaterialPreviewTab;
 	TDMWidgetSlot<SWidget> MaterialPreviewTabSlot;
@@ -268,4 +270,6 @@ protected:
 
 	void NavigateBack_Execute();
 	bool NavigateBack_CanExecute();
+
+	bool CheckOpacityInput(const FKeyEvent& InKeyEvent);
 };
