Material Designer
- Switcher the order and default page of the wizard.
- Made the order a lot easier to change in the future.

diff --git a/Engine/Plugins/Experimental/DynamicMaterial/Source/DynamicMaterialEditor/Private/UI/Widgets/SDMMaterialWizard.cpp b/Engine/Plugins/Experimental/DynamicMaterial/Source/DynamicMaterialEditor/Private/UI/Widgets/SDMMaterialWizard.cpp
--- a/Engine/Plugins/Experimental/DynamicMaterial/Source/DynamicMaterialEditor/Private/UI/Widgets/SDMMaterialWizard.cpp
+++ b/Engine/Plugins/Experimental/DynamicMaterial/Source/DynamicMaterialEditor/Private/UI/Widgets/SDMMaterialWizard.cpp
@@ -158,11 +158,11 @@ TSharedRef<SWidget> SDMMaterialWizard::CreateLayout()
 				SAssignNew(Switcher, SWidgetSwitcher)
 				+ SWidgetSwitcher::Slot()
 				[
-					CreateTemplateListLayout()
+					CreateNewTemplateLayout()
 				]
 				+ SWidgetSwitcher::Slot()
 				[
-					CreateSelectPresetLayout()
+					CreateNewInstanceLayout()
 				]
 			]
 		];
@@ -170,50 +170,78 @@ TSharedRef<SWidget> SDMMaterialWizard::CreateLayout()
 
 TSharedRef<SWidget> SDMMaterialWizard::CreateModeSelector()
 {
-	return SNew(SHorizontalBox)
+	TSharedRef<SHorizontalBox> Container = SNew(SHorizontalBox);
 
-		+ SHorizontalBox::Slot()
-		.FillContentWidth(1.f)
-		.HAlign(EHorizontalAlignment::HAlign_Right)
-		.Padding(0.f, 0.f, 5.f, 0.f)
-		[
-			SNew(SCheckBox)
-			.Style(FAppStyle::Get(), "DetailsView.SectionButton")
-			.HAlign(EHorizontalAlignment::HAlign_Center)
-			.Padding(FMargin(10.f, 6.f))
-			.IsChecked(this, &SDMMaterialWizard::IsModeSelected, 0)
-			.OnCheckStateChanged(this, &SDMMaterialWizard::SetMode, 0)
-			.ToolTipText(LOCTEXT("TemplateModeToolTip", "Create a new Dynamic Material Instance based on a template."))
-			.Content()
-			[
-				SNew(STextBlock)
-				.TextStyle(FDynamicMaterialEditorStyle::Get(), "BoldFont")
-				.Text(LOCTEXT("TemplateMode", "New Instance"))
-			]
-		]
+	const TArray<EHorizontalAlignment> Alignments = {
+		EHorizontalAlignment::HAlign_Right,
+		EHorizontalAlignment::HAlign_Left
+	};
 
-		+ SHorizontalBox::Slot()
-		.FillContentWidth(1.f)
-		.HAlign(EHorizontalAlignment::HAlign_Left)
-		.Padding(5.f, 0.f, 0.f, 0.f)
-		[
-			SNew(SCheckBox)
-			.Style(FAppStyle::Get(), "DetailsView.SectionButton")
-			.HAlign(EHorizontalAlignment::HAlign_Center)
-			.Padding(FMargin(10.f, 6.f))
-			.IsChecked(this, &SDMMaterialWizard::IsModeSelected, 1)
-			.OnCheckStateChanged(this, &SDMMaterialWizard::SetMode, 1)
-			.ToolTipText(LOCTEXT("PresetModeToolTip", "Set up a new Material based on simple channel presets."))
-			.Content()
-			[
-				SNew(STextBlock)
-				.TextStyle(FDynamicMaterialEditorStyle::Get(), "BoldFont")
-				.Text(LOCTEXT("PresetMode", "New Template"))
-			]
-		];
+	const TArray<FMargin> Paddings = {
+		FMargin(0.f, 0.f, 5.f, 0.f),
+		FMargin(5.f, 0.f, 0.f, 0.f)
+	};
+
+	for (int32 SwitcherIndex = 0; SwitcherIndex < EDMMaterialWizardMode::Count; ++SwitcherIndex)
+	{
+		switch (SwitcherIndex)
+		{
+			case EDMMaterialWizardMode::Template:
+			{
+				Container->AddSlot()
+					.FillContentWidth(1.f)
+					.HAlign(Alignments[SwitcherIndex])
+					.Padding(Paddings[SwitcherIndex])
+					[
+						SNew(SCheckBox)
+						.Style(FAppStyle::Get(), "DetailsView.SectionButton")
+						.HAlign(EHorizontalAlignment::HAlign_Center)
+						.Padding(FMargin(10.f, 6.f))
+						.IsChecked(this, &SDMMaterialWizard::IsModeSelected, EDMMaterialWizardMode::Template)
+						.OnCheckStateChanged(this, &SDMMaterialWizard::SetMode, EDMMaterialWizardMode::Template)
+						.ToolTipText(LOCTEXT("PresetModeToolTip", "Set up a new Material based on simple channel presets."))
+						.Content()
+						[
+							SNew(STextBlock)
+							.TextStyle(FDynamicMaterialEditorStyle::Get(), "BoldFont")
+							.Text(LOCTEXT("PresetMode", "New Material"))
+						]
+					];
+
+				break;
+			}
+
+			case EDMMaterialWizardMode::Instance:
+			{
+				Container->AddSlot()
+					.FillContentWidth(1.f)
+					.HAlign(Alignments[SwitcherIndex])
+					.Padding(Paddings[SwitcherIndex])
+					[
+						SNew(SCheckBox)
+						.Style(FAppStyle::Get(), "DetailsView.SectionButton")
+						.HAlign(EHorizontalAlignment::HAlign_Center)
+						.Padding(FMargin(10.f, 6.f))
+						.IsChecked(this, &SDMMaterialWizard::IsModeSelected, EDMMaterialWizardMode::Instance)
+						.OnCheckStateChanged(this, &SDMMaterialWizard::SetMode, EDMMaterialWizardMode::Instance)
+						.ToolTipText(LOCTEXT("TemplateModeToolTip", "Create a new Dynamic Material Instance based on a template."))
+						.Content()
+						[
+							SNew(STextBlock)
+							.TextStyle(FDynamicMaterialEditorStyle::Get(), "BoldFont")
+							.Text(LOCTEXT("TemplateMode", "New Material Instance"))
+						]
+					];
+
+				break;
+			}
+		}
+	}
+
+	return Container;
 }
 
-TSharedRef<SWidget> SDMMaterialWizard::CreateSelectPresetLayout()
+TSharedRef<SWidget> SDMMaterialWizard::CreateNewTemplateLayout()
 {
 	using namespace UE::DynamicMaterialDesigner::Private;
 
@@ -233,7 +261,7 @@ TSharedRef<SWidget> SDMMaterialWizard::CreateSelectPresetLayout()
 		.HAlign(EHorizontalAlignment::HAlign_Fill)
 		.Padding(0.0f, 0.f, 0.0f, TitleContentDistance)
 		[
-			CreateSelectPreset_ChannelPresets()
+			CreateNewTemplate_ChannelPresets()
 		]
 
 		+ SVerticalBox::Slot()
@@ -253,7 +281,7 @@ TSharedRef<SWidget> SDMMaterialWizard::CreateSelectPresetLayout()
 			SAssignNew(PresetChannelContainer, SBox)
 			.HAlign(EHorizontalAlignment::HAlign_Fill)
 			[
-				CreateSelectPreset_ChannelList()
+				CreateNewTemplate_ChannelList()
 			]
 		]
 
@@ -262,11 +290,11 @@ TSharedRef<SWidget> SDMMaterialWizard::CreateSelectPresetLayout()
 		.HAlign(EHorizontalAlignment::HAlign_Fill)
 		.Padding(0.0f, SeparationDistance, 0.0f, 0.f)
 		[
-			CreateSelectPreset_AcceptButton()
+			CreateNewTemplate_AcceptButton()
 		];
 }
 
-TSharedRef<SWidget> SDMMaterialWizard::CreateSelectPreset_ChannelPresets()
+TSharedRef<SWidget> SDMMaterialWizard::CreateNewTemplate_ChannelPresets()
 {
 	using namespace UE::DynamicMaterialDesigner::Private;
 
@@ -296,7 +324,7 @@ TSharedRef<SWidget> SDMMaterialWizard::CreateSelectPreset_ChannelPresets()
 	return ChannelPresets;
 }
 
-TSharedRef<SWidget> SDMMaterialWizard::CreateSelectPreset_ChannelList()
+TSharedRef<SWidget> SDMMaterialWizard::CreateNewTemplate_ChannelList()
 {
 	using namespace UE::DynamicMaterialDesigner::Private;
 
@@ -337,7 +365,7 @@ TSharedRef<SWidget> SDMMaterialWizard::CreateSelectPreset_ChannelList()
 	return ChannelPresets;
 }
 
-TSharedRef<SWidget> SDMMaterialWizard::CreateTemplateListLayout()
+TSharedRef<SWidget> SDMMaterialWizard::CreateNewInstanceLayout()
 {
 	FAssetPickerConfig PickerConfig;
 	PickerConfig.SelectionMode = ESelectionMode::Single;
@@ -434,7 +462,7 @@ TSharedRef<SWidget> SDMMaterialWizard::CreateTemplateListLayout()
 		];
 }
 
-TSharedRef<SWidget> SDMMaterialWizard::CreateSelectPreset_AcceptButton()
+TSharedRef<SWidget> SDMMaterialWizard::CreateNewTemplate_AcceptButton()
 {
 	using namespace UE::DynamicMaterialDesigner::Private;
 
@@ -468,7 +496,7 @@ void SDMMaterialWizard::Preset_OnChange(ECheckBoxState InState, FName InPresetNa
 
 		if (PresetChannelContainer.IsValid())
 		{
-			PresetChannelContainer->SetContent(CreateSelectPreset_ChannelList());
+			PresetChannelContainer->SetContent(CreateNewTemplate_ChannelList());
 		}
 	}
 }
@@ -563,7 +591,7 @@ void SDMMaterialWizard::OpenMaterialInEditor()
 	}
 }
 
-ECheckBoxState SDMMaterialWizard::IsModeSelected(int32 InMode) const
+ECheckBoxState SDMMaterialWizard::IsModeSelected(EDMMaterialWizardMode InMode) const
 {
 	if (Switcher.IsValid() && Switcher->GetActiveWidgetIndex() == InMode)
 	{
@@ -573,7 +601,7 @@ ECheckBoxState SDMMaterialWizard::IsModeSelected(int32 InMode) const
 	return ECheckBoxState::Unchecked;
 }
 
-void SDMMaterialWizard::SetMode(ECheckBoxState InState, int32 InMode)
+void SDMMaterialWizard::SetMode(ECheckBoxState InState, EDMMaterialWizardMode InMode)
 {
 	if (InState == ECheckBoxState::Checked && Switcher.IsValid())
 	{
diff --git a/Engine/Plugins/Experimental/DynamicMaterial/Source/DynamicMaterialEditor/Private/UI/Widgets/SDMMaterialWizard.h b/Engine/Plugins/Experimental/DynamicMaterial/Source/DynamicMaterialEditor/Private/UI/Widgets/SDMMaterialWizard.h
--- a/Engine/Plugins/Experimental/DynamicMaterial/Source/DynamicMaterialEditor/Private/UI/Widgets/SDMMaterialWizard.h
+++ b/Engine/Plugins/Experimental/DynamicMaterial/Source/DynamicMaterialEditor/Private/UI/Widgets/SDMMaterialWizard.h
@@ -46,6 +46,14 @@ public:
 	UDynamicMaterialModel* GetMaterialModel() const;
 
 protected:
+	enum EDMMaterialWizardMode : uint8
+	{
+		Template = 0,
+		Instance = 1,
+
+		Count = 2
+	};
+
 	TWeakPtr<SDMMaterialDesigner> DesignerWidgetWeak;
 	FName CurrentPreset;
 	TSharedPtr<SBox> PresetChannelContainer;
@@ -59,11 +67,11 @@ protected:
 	/** Creation of widgets. */
 	TSharedRef<SWidget> CreateLayout();
 	TSharedRef<SWidget> CreateModeSelector();
-	TSharedRef<SWidget> CreateSelectPresetLayout();
-	TSharedRef<SWidget> CreateSelectPreset_ChannelPresets();
-	TSharedRef<SWidget> CreateSelectPreset_ChannelList();
-	TSharedRef<SWidget> CreateSelectPreset_AcceptButton();
-	TSharedRef<SWidget> CreateTemplateListLayout();
+	TSharedRef<SWidget> CreateNewTemplateLayout();
+	TSharedRef<SWidget> CreateNewTemplate_ChannelPresets();
+	TSharedRef<SWidget> CreateNewTemplate_ChannelList();
+	TSharedRef<SWidget> CreateNewTemplate_AcceptButton();
+	TSharedRef<SWidget> CreateNewInstanceLayout();
 
 	/** Attributes and Events */
 	ECheckBoxState Preset_GetState(FName InPresetName) const;
@@ -75,9 +83,9 @@ protected:
 
 	void OpenMaterialInEditor();
 
-	ECheckBoxState IsModeSelected(int32 InMode) const;
+	ECheckBoxState IsModeSelected(EDMMaterialWizardMode InMode) const;
 
-	void SetMode(ECheckBoxState InState, int32 InMode);
+	void SetMode(ECheckBoxState InState, EDMMaterialWizardMode InMode);
 
 	void OnSearchBoxChanged(const FText& InSearchText);
 	void OnSearchBoxCommitted(const FText& InSearchText, ETextCommit::Type InCommitInfo);
