Material Designer
- The material preview menu now correctly generated a dynamic menu based on whether it's a popout or not.

diff --git a/Engine/Plugins/Experimental/DynamicMaterial/Source/DynamicMaterialEditor/Private/UI/Widgets/Editor/SDMMaterialPreview.cpp b/Engine/Plugins/Experimental/DynamicMaterial/Source/DynamicMaterialEditor/Private/UI/Widgets/Editor/SDMMaterialPreview.cpp
--- a/Engine/Plugins/Experimental/DynamicMaterial/Source/DynamicMaterialEditor/Private/UI/Widgets/Editor/SDMMaterialPreview.cpp
+++ b/Engine/Plugins/Experimental/DynamicMaterial/Source/DynamicMaterialEditor/Private/UI/Widgets/Editor/SDMMaterialPreview.cpp
@@ -74,6 +74,7 @@ void SDMMaterialPreview::Construct(const FArguments& InArgs, const TSharedRef<SD
 {
 	EditorWidgetWeak = InEditorWidget;
 	bShowMenu = InArgs._ShowMenu;
+	bIsPopout = InArgs._IsPopout;
 	PreviewScene = MakeShareable(new FAdvancedPreviewScene(
 		FPreviewScene::ConstructionValues()
 		.SetCreatePhysicsScene(false)
@@ -477,29 +478,52 @@ TSharedRef<SWidget> SDMMaterialPreview::GenerateToolbarMenu()
 		FToolMenuSection& SettingsSection = Menu->AddSection(TEXT("Settings"), LOCTEXT("Settings", "Settings"));
 		SettingsSection.AddEntry(FToolMenuEntry::InitMenuEntry(MaterialEditorCommands.TogglePreviewBackground));
 
-		if (!bIsPopout)
-		{
-			FToolMenuSection& ActionsSection = Menu->AddSection(TEXT("Actions"), LOCTEXT("Actions", "Actions"));
-
-			FUIAction OpenPreviewTabAction;
-			OpenPreviewTabAction.ExecuteAction.BindSP(this, &SDMMaterialPreview::OpenMaterialPreviewTab);
-
-			ActionsSection.AddEntry(FToolMenuEntry::InitMenuEntry(
-				TEXT("PopoutMaterialPreviewTab"),
-				LOCTEXT("OpenPreview", "Open Preview"),
-				LOCTEXT("OpenPreviewToolTip", "Open a tab with a preview of the material."),
-				TAttribute<FSlateIcon>(),
-				OpenPreviewTabAction
-			));
-		}
+		Menu->AddDynamicSection(
+			TEXT("Actions"),
+			FNewSectionConstructChoice(FNewToolMenuDelegate::CreateStatic(&SDMMaterialPreview::AddActionMenu))		
+		);
 	}
 
+	UDMMaterialPreviewContext* PreviewContext = NewObject<UDMMaterialPreviewContext>();
+	PreviewContext->SetPreviewWidget(SharedThis(this));
+
 	FToolMenuContext Context;
 	Context.AppendCommandList(CommandList);
+	Context.AddObject(PreviewContext);
 
 	return ToolMenus->GenerateWidget(MenuName, Context);
 }
 
+void SDMMaterialPreview::AddActionMenu(UToolMenu* InMenu)
+{
+	const UDMMaterialPreviewContext* Context = InMenu->FindContext<UDMMaterialPreviewContext>();
+
+	if (!Context)
+	{
+		return;
+	}
+
+	TSharedPtr<SDMMaterialPreview> PreviewWidget = Context->GetPreviewWidget();
+
+	if (!PreviewWidget.IsValid() || PreviewWidget->bIsPopout)
+	{
+		return;
+	}
+
+	FToolMenuSection& ActionsSection = InMenu->AddSection(TEXT("Actions"), LOCTEXT("Actions", "Actions"));
+
+	FUIAction OpenPreviewTabAction;
+	OpenPreviewTabAction.ExecuteAction.BindSP(PreviewWidget.Get(), &SDMMaterialPreview::OpenMaterialPreviewTab);
+
+	ActionsSection.AddEntry(FToolMenuEntry::InitMenuEntry(
+		TEXT("PopoutMaterialPreviewTab"),
+		LOCTEXT("OpenPreview", "Open Preview"),
+		LOCTEXT("OpenPreviewToolTip", "Open a tab with a preview of the material."),
+		TAttribute<FSlateIcon>(),
+		OpenPreviewTabAction
+	));
+}
+
 void SDMMaterialPreview::OnEditorSettingsChanged(const FPropertyChangedEvent& InPropertyChangedEvent)
 {
 	UDynamicMaterialEditorSettings* Settings = UDynamicMaterialEditorSettings::Get();
diff --git a/Engine/Plugins/Experimental/DynamicMaterial/Source/DynamicMaterialEditor/Private/UI/Widgets/Editor/SDMMaterialPreview.h b/Engine/Plugins/Experimental/DynamicMaterial/Source/DynamicMaterialEditor/Private/UI/Widgets/Editor/SDMMaterialPreview.h
--- a/Engine/Plugins/Experimental/DynamicMaterial/Source/DynamicMaterialEditor/Private/UI/Widgets/Editor/SDMMaterialPreview.h
+++ b/Engine/Plugins/Experimental/DynamicMaterial/Source/DynamicMaterialEditor/Private/UI/Widgets/Editor/SDMMaterialPreview.h
@@ -9,17 +9,42 @@
 #include "UObject/ObjectPtr.h"
 #include "UObject/UObjectGlobals.h"
 
+#include "SDMMaterialPreview.generated.h"
+
 class APostProcessVolume;
 class FAdvancedPreviewScene;
 class FDMMaterialPreviewViewportClient;
 class FEditorViewportClient;
 class SDMMaterialEditor;
+class SDMMaterialPreview;
 class UDynamicMaterialModelBase;
 class UMaterialInterface;
 class UMeshComponent;
+class UToolMenu;
 enum class EDMMaterialPreviewMesh : uint8;
 struct FPropertyChangedEvent;
 
+UCLASS(MinimalAPI)
+class UDMMaterialPreviewContext : public UObject
+{
+	GENERATED_BODY()
+
+public:
+	void SetPreviewWidget(const TSharedRef<SDMMaterialPreview>& InPreviewWidget)
+	{
+		PreviewWidgetWeak = InPreviewWidget;
+	}
+
+	TSharedPtr<SDMMaterialPreview> GetPreviewWidget() const
+	{
+		return PreviewWidgetWeak.Pin();
+	}
+
+private:
+	TWeakPtr<SDMMaterialPreview> PreviewWidgetWeak;
+};
+
+
 /** Based on SMaterialEditor3DPreviewViewport (private) */
 class SDMMaterialPreview : public SEditorViewport, public FGCObject
 {
@@ -89,6 +114,8 @@ protected:
 
 	TSharedRef<SWidget> GenerateToolbarMenu();
 
+	static void AddActionMenu(UToolMenu* InMenu);
+
 	void OnEditorSettingsChanged(const FPropertyChangedEvent& InPropertyChangedEvent);
 
 	void OpenMaterialPreviewTab();
