Material Designer
- Cleared stored property handles correctly on editor shutdown.

diff --git a/Engine/Plugins/Experimental/DynamicMaterial/Source/DynamicMaterialEditor/Private/DynamicMaterialEditorModule.cpp b/Engine/Plugins/Experimental/DynamicMaterial/Source/DynamicMaterialEditor/Private/DynamicMaterialEditorModule.cpp
--- a/Engine/Plugins/Experimental/DynamicMaterial/Source/DynamicMaterialEditor/Private/DynamicMaterialEditorModule.cpp
+++ b/Engine/Plugins/Experimental/DynamicMaterial/Source/DynamicMaterialEditor/Private/DynamicMaterialEditorModule.cpp
@@ -45,6 +45,7 @@
 #include "UI/PropertyGenerators/DMTextureUVDynamicPropertyRowGenerator.h"
 #include "UI/PropertyGenerators/DMTextureUVPropertyRowGenerator.h"
 #include "UI/PropertyGenerators/DMThroughputPropertyRowGenerator.h"
+#include "UI/Utils/DMWidgetStatics.h"
 #include "UI/Utils/DynamicMaterialInstanceThumbnailRenderer.h"
 #include "UI/Widgets/SDMMaterialDesigner.h"
 
@@ -292,6 +293,8 @@ void FDynamicMaterialEditorModule::ShutdownModule()
 	UDMMaterialValueTexture::GetDefaultRGBTexture.Unbind();
 
 	FDMValueDetailsRowExtensions::Get().UnregisterRowExtensions();
+
+	FDMWidgetStatics::Get().ClearData();
 }
 
 void FDynamicMaterialEditorModule::OpenMaterialModel(UDynamicMaterialModelBase* InMaterialModel, UWorld* InWorld, bool bInInvokeTab) const
diff --git a/Engine/Plugins/Experimental/DynamicMaterial/Source/DynamicMaterialEditor/Private/UI/Utils/DMWidgetStatics.cpp b/Engine/Plugins/Experimental/DynamicMaterial/Source/DynamicMaterialEditor/Private/UI/Utils/DMWidgetStatics.cpp
--- a/Engine/Plugins/Experimental/DynamicMaterial/Source/DynamicMaterialEditor/Private/UI/Utils/DMWidgetStatics.cpp
+++ b/Engine/Plugins/Experimental/DynamicMaterial/Source/DynamicMaterialEditor/Private/UI/Utils/DMWidgetStatics.cpp
@@ -140,6 +140,12 @@ TSharedPtr<SWidget> FDMWidgetStatics::GetInnerPropertyValueWidget(const TSharedR
 	return nullptr;
 }
 
+void FDMWidgetStatics::ClearData()
+{
+	ExpansionStates.Empty();
+	PropertyHandleMap.Empty();
+}
+
 FDMPropertyHandle FDMWidgetStatics::CreatePropertyHandle(const void* InOwningWidget, UObject* InObject, FName InPropertyName)
 {
 	FDMPropertyHandle PropertyHandle;
diff --git a/Engine/Plugins/Experimental/DynamicMaterial/Source/DynamicMaterialEditor/Private/UI/Utils/DMWidgetStatics.h b/Engine/Plugins/Experimental/DynamicMaterial/Source/DynamicMaterialEditor/Private/UI/Utils/DMWidgetStatics.h
--- a/Engine/Plugins/Experimental/DynamicMaterial/Source/DynamicMaterialEditor/Private/UI/Utils/DMWidgetStatics.h
+++ b/Engine/Plugins/Experimental/DynamicMaterial/Source/DynamicMaterialEditor/Private/UI/Utils/DMWidgetStatics.h
@@ -31,6 +31,8 @@ public:
 
 	TSharedPtr<SWidget> GetInnerPropertyValueWidget(const TSharedRef<SWidget>& InWidget);
 
+	void ClearData();
+
 private:
 	struct FExpansionItem
 	{
