Material Designer
- Removed the text "Material Wizard" from the wizard page.

diff --git a/Engine/Plugins/Experimental/DynamicMaterial/Source/DynamicMaterialEditor/Private/UI/Widgets/SDMMaterialWizard.cpp b/Engine/Plugins/Experimental/DynamicMaterial/Source/DynamicMaterialEditor/Private/UI/Widgets/SDMMaterialWizard.cpp
--- a/Engine/Plugins/Experimental/DynamicMaterial/Source/DynamicMaterialEditor/Private/UI/Widgets/SDMMaterialWizard.cpp
+++ b/Engine/Plugins/Experimental/DynamicMaterial/Source/DynamicMaterialEditor/Private/UI/Widgets/SDMMaterialWizard.cpp
@@ -135,15 +135,7 @@ TSharedRef<SWidget> SDMMaterialWizard::CreateLayout()
 		.HAlign(EHorizontalAlignment::HAlign_Fill)
 		[
 			SNew(SVerticalBox)
-			+ SVerticalBox::Slot()
-			.AutoHeight()
-			.HAlign(HAlign_Center)
-			[
-				SNew(STextBlock)
-				.TextStyle(FDynamicMaterialEditorStyle::Get(), "ActorNameBig")
-				.Text(LOCTEXT("MaterialWizard", "Material Wizard"))
-			]
-
+			
 			+ SVerticalBox::Slot()
 			.AutoHeight()
 			.Padding(0.0f, SeparationDistance, 0.0f, 0.f)
