Material Designer
- bUsedWithNanite is now enabled if there is a displacement channel or bNaniteTessellationEnabled is enabled.

diff --git a/Engine/Plugins/Experimental/DynamicMaterial/Source/DynamicMaterialEditor/Private/Components/MaterialProperties/DMMPDisplacement.cpp b/Engine/Plugins/Experimental/DynamicMaterial/Source/DynamicMaterialEditor/Private/Components/MaterialProperties/DMMPDisplacement.cpp
--- a/Engine/Plugins/Experimental/DynamicMaterial/Source/DynamicMaterialEditor/Private/Components/MaterialProperties/DMMPDisplacement.cpp
+++ b/Engine/Plugins/Experimental/DynamicMaterial/Source/DynamicMaterialEditor/Private/Components/MaterialProperties/DMMPDisplacement.cpp
@@ -3,6 +3,7 @@
 #include "Components/MaterialProperties/DMMPDisplacement.h"
 
 #include "Components/MaterialValues/DMMaterialValueFloat1.h"
+#include "Materials/Material.h"
 #include "Materials/MaterialExpressionMaterialFunctionCall.h"
 #include "Model/DMMaterialBuildState.h"
 #include "Model/DynamicMaterialModelEditorOnlyData.h"
@@ -35,6 +36,20 @@ UMaterialExpression* UDMMaterialPropertyDisplacement::GetDefaultInput(
 	return CreateConstant(InBuildState, FVector::ZeroVector);
 }
 
+void UDMMaterialPropertyDisplacement::GenerateExpressions(const TSharedRef<FDMMaterialBuildState>& InBuildState) const
+{
+	Super::GenerateExpressions(InBuildState);
+
+	if (!InBuildState->GetPreviewObject())
+	{
+		if (UMaterial* GeneratedMaterial = InBuildState->GetDynamicMaterial())
+		{
+			// Set this to true if we have a displacement channel.
+			GeneratedMaterial->bUsedWithNanite = true;
+		}
+	}
+}
+
 void UDMMaterialPropertyDisplacement::AddAlphaMultiplier(const TSharedRef<FDMMaterialBuildState>& InBuildState) const
 {
 	UDMMaterialValueFloat1* AlphaValue = GetTypedComponent<UDMMaterialValueFloat1>(UDynamicMaterialModelEditorOnlyData::AlphaValueName);
diff --git a/Engine/Plugins/Experimental/DynamicMaterial/Source/DynamicMaterialEditor/Private/Components/MaterialProperties/DMMPDisplacement.h b/Engine/Plugins/Experimental/DynamicMaterial/Source/DynamicMaterialEditor/Private/Components/MaterialProperties/DMMPDisplacement.h
--- a/Engine/Plugins/Experimental/DynamicMaterial/Source/DynamicMaterialEditor/Private/Components/MaterialProperties/DMMPDisplacement.h
+++ b/Engine/Plugins/Experimental/DynamicMaterial/Source/DynamicMaterialEditor/Private/Components/MaterialProperties/DMMPDisplacement.h
@@ -17,6 +17,7 @@ public:
 
 	//~ Begin UDMMaterialProperty
 	virtual UMaterialExpression* GetDefaultInput(const TSharedRef<FDMMaterialBuildState>& InBuildState) const override;
+	virtual void GenerateExpressions(const TSharedRef<FDMMaterialBuildState>& InBuildState) const override;
 	virtual void AddAlphaMultiplier(const TSharedRef<FDMMaterialBuildState>& InBuildState) const override;
 	//~ End UDMMaterialProperty
 };
diff --git a/Engine/Plugins/Experimental/DynamicMaterial/Source/DynamicMaterialEditor/Private/Model/DynamicMaterialModelEditorOnlyData.cpp b/Engine/Plugins/Experimental/DynamicMaterial/Source/DynamicMaterialEditor/Private/Model/DynamicMaterialModelEditorOnlyData.cpp
--- a/Engine/Plugins/Experimental/DynamicMaterial/Source/DynamicMaterialEditor/Private/Model/DynamicMaterialModelEditorOnlyData.cpp
+++ b/Engine/Plugins/Experimental/DynamicMaterial/Source/DynamicMaterialEditor/Private/Model/DynamicMaterialModelEditorOnlyData.cpp
@@ -401,6 +401,9 @@ void UDynamicMaterialModelEditorOnlyData::BuildMaterial(bool bInDirtyAssets)
 	// is used with Niagara. It doesn't negatively affect the material in any meaningful way.
 	MaterialModel->DynamicMaterial->bUsedWithNiagaraMeshParticles = true;
 
+	// If bNaniteTessellationEnabled is false, bUsedWithNanite will be set back to true if there's a displacement property.
+	MaterialModel->DynamicMaterial->bUsedWithNanite = bNaniteTessellationEnabled;
+
 	switch (ShadingModel)
 	{
 		case EDMMaterialShadingModel::DefaultLit:
