Material Designer
- When checking if a throughput should keep input, a valid option is now a single scalar instead of needing to match the float count exactly (e.g. vector3 vs vector3.) This will mean that channel-filtered inputs will not be reset. This is also correct when it comes to materials as a single float is a valid input for anything that accepts scalars or vectors (where this is concerned).

diff --git a/Engine/Plugins/Experimental/DynamicMaterial/Source/DynamicMaterialEditor/Private/Components/DMMaterialStageThroughput.cpp b/Engine/Plugins/Experimental/DynamicMaterial/Source/DynamicMaterialEditor/Private/Components/DMMaterialStageThroughput.cpp
--- a/Engine/Plugins/Experimental/DynamicMaterial/Source/DynamicMaterialEditor/Private/Components/DMMaterialStageThroughput.cpp
+++ b/Engine/Plugins/Experimental/DynamicMaterial/Source/DynamicMaterialEditor/Private/Components/DMMaterialStageThroughput.cpp
@@ -152,9 +152,9 @@ bool UDMMaterialStageThroughput::ShouldKeepInput(int32 InThroughputInputIndex)
 		return false;
 	}
 
-	const int RequiredInputCount = UDMValueDefinitionLibrary::GetValueDefinition(InputConnectors[InThroughputInputIndex].Type).GetFloatCount();
+	const int32 RequiredInputCount = UDMValueDefinitionLibrary::GetValueDefinition(InputConnectors[InThroughputInputIndex].Type).GetFloatCount();
 
-	int ActualInputCount = 0;
+	int32 ActualInputCount = 0;
 
 	for (const FDMMaterialStageConnectorChannel& Channel : InputConnectionMap[InThroughputInputIndex].Channels)
 	{
@@ -256,7 +256,8 @@ bool UDMMaterialStageThroughput::ShouldKeepInput(int32 InThroughputInputIndex)
 		ActualInputCount += ThisInputCount;
 	}
 
-	return (ActualInputCount == RequiredInputCount);
+	// An input size of 1 float will always work.
+	return (ActualInputCount == 1) || (ActualInputCount == RequiredInputCount);
 }
 
 void UDMMaterialStageThroughput::ConnectOutputToInput(const TSharedRef<FDMMaterialBuildState>& InBuildState, int32 InThroughputInputIndex,
