Material Designer
- Added "(Inst)" to the toolbar when a material instance is open.
- Changed "Dynamic" to "Instance" on the preview.

diff --git a/Engine/Plugins/Experimental/DynamicMaterial/Source/DynamicMaterialEditor/Private/UI/Widgets/Editor/SDMToolBar.cpp b/Engine/Plugins/Experimental/DynamicMaterial/Source/DynamicMaterialEditor/Private/UI/Widgets/Editor/SDMToolBar.cpp
--- a/Engine/Plugins/Experimental/DynamicMaterial/Source/DynamicMaterialEditor/Private/UI/Widgets/Editor/SDMToolBar.cpp
+++ b/Engine/Plugins/Experimental/DynamicMaterial/Source/DynamicMaterialEditor/Private/UI/Widgets/Editor/SDMToolBar.cpp
@@ -193,6 +193,17 @@ TSharedRef<SWidget> SDMToolBar::CreateToolBarEntries()
 				]
 			]
 
+			+ SWrapBox::Slot()
+			.HAlign(HAlign_Left)
+			.VAlign(VAlign_Center)
+			.Padding(5.0f, 0.0f, 0.0f, 0.0f)
+			[
+				SAssignNew(InstanceWidget, STextBlock)
+				.TextStyle(FDynamicMaterialEditorStyle::Get(), "ActorName")
+				.Text(LOCTEXT("Instance", "(Inst)"))
+				.Visibility(EVisibility::Collapsed)
+			]
+
 			+ SWrapBox::Slot()
 			.HAlign(HAlign_Left)
 			.VAlign(VAlign_Center)
@@ -369,11 +380,13 @@ void SDMToolBar::SetButtonVisibilities()
 	{
 		OpenParentButton->SetVisibility(EVisibility::Visible);
 		ConvertToEditableButton->SetVisibility(EVisibility::Visible);
+		InstanceWidget->SetVisibility(EVisibility::Visible);
 	}
 	else
 	{
 		OpenParentButton->SetVisibility(EVisibility::Collapsed);
 		ConvertToEditableButton->SetVisibility(EVisibility::Collapsed);
+		InstanceWidget->SetVisibility(EVisibility::Collapsed);
 	}
 }
 
diff --git a/Engine/Plugins/Experimental/DynamicMaterial/Source/DynamicMaterialEditor/Private/UI/Widgets/Editor/SDMToolBar.h b/Engine/Plugins/Experimental/DynamicMaterial/Source/DynamicMaterialEditor/Private/UI/Widgets/Editor/SDMToolBar.h
--- a/Engine/Plugins/Experimental/DynamicMaterial/Source/DynamicMaterialEditor/Private/UI/Widgets/Editor/SDMToolBar.h
+++ b/Engine/Plugins/Experimental/DynamicMaterial/Source/DynamicMaterialEditor/Private/UI/Widgets/Editor/SDMToolBar.h
@@ -61,6 +61,7 @@ protected:
 	TSharedPtr<SWidget> AssetRowWidget;
 	TSharedPtr<STextBlock> ActorNameWidget;
 	TSharedPtr<STextBlock> AssetNameWidget;
+	TSharedPtr<STextBlock> InstanceWidget;
 	TSharedPtr<SWidget> OpenParentButton;
 	TSharedPtr<SWidget> ConvertToEditableButton;
 
diff --git a/Engine/Plugins/Experimental/DynamicMaterial/Source/DynamicMaterialEditor/Private/UI/Widgets/SDMMaterialEditor.cpp b/Engine/Plugins/Experimental/DynamicMaterial/Source/DynamicMaterialEditor/Private/UI/Widgets/SDMMaterialEditor.cpp
--- a/Engine/Plugins/Experimental/DynamicMaterial/Source/DynamicMaterialEditor/Private/UI/Widgets/SDMMaterialEditor.cpp
+++ b/Engine/Plugins/Experimental/DynamicMaterial/Source/DynamicMaterialEditor/Private/UI/Widgets/SDMMaterialEditor.cpp
@@ -1137,7 +1137,7 @@ TSharedRef<SWidget> SDMMaterialEditor::CreateSlot_Preview()
 			[
 				SNew(STextBlock)
 				.Font(FAppStyle::GetFontStyle("TinyText"))
-				.Text(LOCTEXT("DynamicMaterial", "Dynamic"))
+				.Text(LOCTEXT("MaterialInstance", "Instance"))
 				.ShadowColorAndOpacity(FLinearColor::Black)
 				.ShadowOffset(FVector2D(1.0))
 			];
