Material Designer
- Change printf parameters from %hc to %c.

diff --git a/Engine/Plugins/Experimental/DynamicMaterial/Source/DynamicMaterialEditor/Private/Components/DMMaterialEffect.cpp b/Engine/Plugins/Experimental/DynamicMaterial/Source/DynamicMaterialEditor/Private/Components/DMMaterialEffect.cpp
--- a/Engine/Plugins/Experimental/DynamicMaterial/Source/DynamicMaterialEditor/Private/Components/DMMaterialEffect.cpp
+++ b/Engine/Plugins/Experimental/DynamicMaterial/Source/DynamicMaterialEditor/Private/Components/DMMaterialEffect.cpp
@@ -94,7 +94,7 @@ UDMMaterialComponent* UDMMaterialEffect::GetParentComponent() const
 FString UDMMaterialEffect::GetComponentPathComponent() const
 {
 	return FString::Printf(
-		TEXT("%s%hc%i%hc"),
+		TEXT("%s%c%i%c"),
 		*UDMMaterialEffectStack::EffectsPathToken,
 		FDMComponentPath::ParameterOpen,
 		FindIndex(),
diff --git a/Engine/Plugins/Experimental/DynamicMaterial/Source/DynamicMaterialEditor/Private/Components/DMMaterialLayer.cpp b/Engine/Plugins/Experimental/DynamicMaterial/Source/DynamicMaterialEditor/Private/Components/DMMaterialLayer.cpp
--- a/Engine/Plugins/Experimental/DynamicMaterial/Source/DynamicMaterialEditor/Private/Components/DMMaterialLayer.cpp
+++ b/Engine/Plugins/Experimental/DynamicMaterial/Source/DynamicMaterialEditor/Private/Components/DMMaterialLayer.cpp
@@ -744,7 +744,7 @@ UDMMaterialComponent* UDMMaterialLayerObject::GetParentComponent() const
 FString UDMMaterialLayerObject::GetComponentPathComponent() const
 {
 	return FString::Printf(
-		TEXT("%s%hc%i%hc"),
+		TEXT("%s%c%i%c"),
 		*UDMMaterialSlot::LayersPathToken,
 		FDMComponentPath::ParameterOpen,
 		FindIndex(),
diff --git a/Engine/Plugins/Experimental/DynamicMaterial/Source/DynamicMaterialEditor/Private/Components/DMMaterialSlot.cpp b/Engine/Plugins/Experimental/DynamicMaterial/Source/DynamicMaterialEditor/Private/Components/DMMaterialSlot.cpp
--- a/Engine/Plugins/Experimental/DynamicMaterial/Source/DynamicMaterialEditor/Private/Components/DMMaterialSlot.cpp
+++ b/Engine/Plugins/Experimental/DynamicMaterial/Source/DynamicMaterialEditor/Private/Components/DMMaterialSlot.cpp
@@ -971,7 +971,7 @@ FString UDMMaterialSlot::GetComponentPathComponent() const
 			const FText Token = UE::DynamicMaterialEditor::Private::GetMaterialPropertyShortDisplayName(SlotProperties[0]);
 
 			return FString::Printf(
-				TEXT("%s%hc%s%hc"),
+				TEXT("%s%c%s%c"),
 				*UDynamicMaterialModelEditorOnlyData::SlotsPathToken,
 				FDMComponentPath::ParameterOpen,
 				*Token.ToString(),
@@ -981,7 +981,7 @@ FString UDMMaterialSlot::GetComponentPathComponent() const
 	}
 
 	return FString::Printf(
-		TEXT("%s%hc%i%hc"),
+		TEXT("%s%c%i%c"),
 		*UDynamicMaterialModelEditorOnlyData::SlotsPathToken,
 		FDMComponentPath::ParameterOpen,
 		Index,
diff --git a/Engine/Plugins/Experimental/DynamicMaterial/Source/DynamicMaterialEditor/Private/Components/DMMaterialStage.cpp b/Engine/Plugins/Experimental/DynamicMaterial/Source/DynamicMaterialEditor/Private/Components/DMMaterialStage.cpp
--- a/Engine/Plugins/Experimental/DynamicMaterial/Source/DynamicMaterialEditor/Private/Components/DMMaterialStage.cpp
+++ b/Engine/Plugins/Experimental/DynamicMaterial/Source/DynamicMaterialEditor/Private/Components/DMMaterialStage.cpp
@@ -144,7 +144,7 @@ FString UDMMaterialStage::GetComponentPathComponent() const
 {
 	if (UDMMaterialLayerObject* Layer = GetLayer())
 	{
-		FString TypeStr = "?";
+		FString TypeStr = TEXT("?");
 
 		switch (Layer->GetStageType(this))
 		{
@@ -162,7 +162,7 @@ FString UDMMaterialStage::GetComponentPathComponent() const
 		}
 
 		return FString::Printf(
-			TEXT("%s%hc%s%hc"),
+			TEXT("%s%c%s%c"),
 			*UDMMaterialLayerObject::StagesPathToken,
 			FDMComponentPath::ParameterOpen,
 			*TypeStr,
diff --git a/Engine/Plugins/Experimental/DynamicMaterial/Source/DynamicMaterialEditor/Private/Components/DMMaterialStageInput.cpp b/Engine/Plugins/Experimental/DynamicMaterial/Source/DynamicMaterialEditor/Private/Components/DMMaterialStageInput.cpp
--- a/Engine/Plugins/Experimental/DynamicMaterial/Source/DynamicMaterialEditor/Private/Components/DMMaterialStageInput.cpp
+++ b/Engine/Plugins/Experimental/DynamicMaterial/Source/DynamicMaterialEditor/Private/Components/DMMaterialStageInput.cpp
@@ -109,7 +109,7 @@ FString UDMMaterialStageInput::GetComponentPathComponent() const
 		}
 
 		return FString::Printf(
-			TEXT("%s%hc%i%hc"),
+			TEXT("%s%c%i%c"),
 			*UDMMaterialStage::InputsPathToken,
 			FDMComponentPath::ParameterOpen,
 			Index,
