Motion Design: Optimized Analytic Events and grouped common/similar events into one and differentiated via parameters

diff --git a/Engine/Plugins/Experimental/ActorModifierCore/Source/ActorModifierCore/Private/Subsystems/ActorModifierCoreSubsystem.cpp b/Engine/Plugins/Experimental/ActorModifierCore/Source/ActorModifierCore/Private/Subsystems/ActorModifierCoreSubsystem.cpp
--- a/Engine/Plugins/Experimental/ActorModifierCore/Source/ActorModifierCore/Private/Subsystems/ActorModifierCoreSubsystem.cpp
+++ b/Engine/Plugins/Experimental/ActorModifierCore/Source/ActorModifierCore/Private/Subsystems/ActorModifierCoreSubsystem.cpp
@@ -95,8 +95,10 @@ void UActorModifierCoreSubsystem::OnInsertModifier(const FActorModifierCoreStack
 	if (FEngineAnalytics::IsAvailable())
 	{
 		TArray<FAnalyticsEventAttribute> Attributes;
+		Attributes.Reserve(2);
+		Attributes.Emplace(TEXT("Action"), TEXT("Insert"));
 		Attributes.Emplace(TEXT("Name"), InInsertOp.NewModifierName.ToString());
-		FEngineAnalytics::GetProvider().RecordEvent(TEXT("Editor.Usage.ActorModifiers.InsertModifier"), Attributes);
+		FEngineAnalytics::GetProvider().RecordEvent(TEXT("Editor.Usage.ActorModifiers.Modifier"), Attributes);
 	}
 #endif
 }
diff --git a/Engine/Plugins/Experimental/Avalanche/Source/Avalanche/Private/AvaScene.cpp b/Engine/Plugins/Experimental/Avalanche/Source/Avalanche/Private/AvaScene.cpp
--- a/Engine/Plugins/Experimental/Avalanche/Source/Avalanche/Private/AvaScene.cpp
+++ b/Engine/Plugins/Experimental/Avalanche/Source/Avalanche/Private/AvaScene.cpp
@@ -27,15 +27,23 @@ DEFINE_LOG_CATEGORY_STATIC(LogAvaScene, Log, All);
 
 #define LOCTEXT_NAMESPACE "AvaScene"
 
-void AAvaScene::OnSceneCreated(FString&& InCreationType)
-{
 #if WITH_EDITOR
+void AAvaScene::NotifySceneEvent(ESceneAction InAction)
+{
 	if (FEngineAnalytics::IsAvailable())
 	{
-		FEngineAnalytics::GetProvider().RecordEvent(TEXT("Editor.Usage.MotionDesign.SceneCreated"), FAnalyticsEventAttribute(TEXT("CreationType"), MoveTemp(InCreationType)));
+		FString ActionName;
+		switch (InAction)
+		{
+		case ESceneAction::Created:     ActionName = TEXT("Created"); break;
+		case ESceneAction::Activated:   ActionName = TEXT("Activated"); break;
+		case ESceneAction::Deactivated: ActionName = TEXT("Deactivated"); break;
+		}
+		FEngineAnalytics::GetProvider().RecordEvent(TEXT("Editor.Usage.MotionDesign.Scene")
+			, FAnalyticsEventAttribute(TEXT("Action"), MoveTemp(ActionName)));
 	}
-#endif
 }
+#endif
 
 AAvaScene* AAvaScene::GetScene(ULevel* InLevel, bool bInCreateSceneIfNotFound)
 {
@@ -65,7 +73,9 @@ AAvaScene* AAvaScene::GetScene(ULevel* InLevel, bool bInCreateSceneIfNotFound)
 #endif
 
 	AAvaScene* const NewScene = World->SpawnActor<AAvaScene>(SpawnParameters);
-	OnSceneCreated(/*CreationType*/TEXT("Spawned"));
+#if WITH_EDITOR
+	NotifySceneEvent(ESceneAction::Created);
+#endif
 	return NewScene;
 }
 
@@ -182,7 +192,7 @@ bool AAvaScene::AddSequence(UAvaSequence* InSequence)
 #if WITH_EDITOR
 		if (FEngineAnalytics::IsAvailable())
 		{
-			FEngineAnalytics::GetProvider().RecordEvent(TEXT("Editor.Usage.MotionDesign.AddedSequence"));
+			FEngineAnalytics::GetProvider().RecordEvent(TEXT("Editor.Usage.MotionDesign.Sequence"), TEXT("Action"), TEXT("AddedSequence"));
 		}
 #endif
 
diff --git a/Engine/Plugins/Experimental/Avalanche/Source/Avalanche/Public/AvaScene.h b/Engine/Plugins/Experimental/Avalanche/Source/Avalanche/Public/AvaScene.h
--- a/Engine/Plugins/Experimental/Avalanche/Source/Avalanche/Public/AvaScene.h
+++ b/Engine/Plugins/Experimental/Avalanche/Source/Avalanche/Public/AvaScene.h
@@ -28,9 +28,17 @@ class AAvaScene : public AActor, public IAvaSequenceProvider, public IAvaSceneIn
 {
 	GENERATED_BODY()
 
-	static void OnSceneCreated(FString&& InCreationType);
-
 public:
+#if WITH_EDITOR
+	enum class ESceneAction : uint8
+	{
+		Created,
+		Activated,
+		Deactivated,
+	};
+	AVALANCHE_API static void NotifySceneEvent(ESceneAction InAction);
+#endif
+
 	AVALANCHE_API static AAvaScene* GetScene(ULevel* InLevel, bool bInCreateSceneIfNotFound);
 
 	AAvaScene();
diff --git a/Engine/Plugins/Experimental/Avalanche/Source/AvalancheEditor/Private/EditorProvider/AvaEditorProvider.cpp b/Engine/Plugins/Experimental/Avalanche/Source/AvalancheEditor/Private/EditorProvider/AvaEditorProvider.cpp
--- a/Engine/Plugins/Experimental/Avalanche/Source/AvalancheEditor/Private/EditorProvider/AvaEditorProvider.cpp
+++ b/Engine/Plugins/Experimental/Avalanche/Source/AvalancheEditor/Private/EditorProvider/AvaEditorProvider.cpp
@@ -30,3 +30,13 @@ void FAvaEditorProvider::GetActorsToEdit(TArray<AActor*>& InOutActorsToEdit) con
 {
 	FAvaEditorActorUtils::GetActorsToEdit(InOutActorsToEdit);
 }
+
+void FAvaEditorProvider::OnSceneActivated()
+{
+	AAvaScene::NotifySceneEvent(AAvaScene::ESceneAction::Activated);
+}
+
+void FAvaEditorProvider::OnSceneDeactivated()
+{
+	AAvaScene::NotifySceneEvent(AAvaScene::ESceneAction::Deactivated);
+}
diff --git a/Engine/Plugins/Experimental/Avalanche/Source/AvalancheEditor/Private/EditorProvider/AvaEditorProvider.h b/Engine/Plugins/Experimental/Avalanche/Source/AvalancheEditor/Private/EditorProvider/AvaEditorProvider.h
--- a/Engine/Plugins/Experimental/Avalanche/Source/AvalancheEditor/Private/EditorProvider/AvaEditorProvider.h
+++ b/Engine/Plugins/Experimental/Avalanche/Source/AvalancheEditor/Private/EditorProvider/AvaEditorProvider.h
@@ -11,5 +11,7 @@ public:
 	//~ Begin IAvaEditorProvider
 	virtual UObject* GetSceneObject(UWorld* InWorld, EAvaEditorObjectQueryType InQueryType) override;
 	virtual void GetActorsToEdit(TArray<AActor*>& InOutActorsToEdit) const override;
+	virtual void OnSceneActivated() override;
+	virtual void OnSceneDeactivated() override;
 	//~ End IAvaEditorProvider
 };
diff --git a/Engine/Plugins/Experimental/Avalanche/Source/AvalancheEditorCore/Private/AvaEditor.cpp b/Engine/Plugins/Experimental/Avalanche/Source/AvalancheEditorCore/Private/AvaEditor.cpp
--- a/Engine/Plugins/Experimental/Avalanche/Source/AvalancheEditorCore/Private/AvaEditor.cpp
+++ b/Engine/Plugins/Experimental/Avalanche/Source/AvalancheEditorCore/Private/AvaEditor.cpp
@@ -224,20 +224,15 @@ void FAvaEditor::ForEachExtension(TFunctionRef<void(const TSharedRef<IAvaEditorE
 	}
 }
 
-void FAvaEditor::RecordActivationChangedEvent()
+void FAvaEditor::NotifySceneActivationEvent()
 {
-	if (!FEngineAnalytics::IsAvailable())
-	{
-		return;
-	}
-
 	if (bIsActive)
 	{
-		FEngineAnalytics::GetProvider().RecordEvent(TEXT("Editor.Usage.MotionDesign.Activated"));
+		Provider->OnSceneActivated();
 	}
 	else
 	{
-		FEngineAnalytics::GetProvider().RecordEvent(TEXT("Editor.Usage.MotionDesign.Deactivated"));
+		Provider->OnSceneDeactivated();
 	}
 }
 
@@ -252,7 +247,7 @@ void FAvaEditor::Activate(TSharedPtr<IToolkitHost> InOverrideToolkitHost)
 
 	bIsActive = true;
 
-	RecordActivationChangedEvent();
+	NotifySceneActivationEvent();
 
 	BindDelegates();
 
@@ -306,7 +301,7 @@ void FAvaEditor::Deactivate()
 
 	bIsActive = false;
 
-	RecordActivationChangedEvent();
+	NotifySceneActivationEvent();
 
 	UnbindDelegates();
 
diff --git a/Engine/Plugins/Experimental/Avalanche/Source/AvalancheEditorCore/Private/AvaEditor.h b/Engine/Plugins/Experimental/Avalanche/Source/AvalancheEditorCore/Private/AvaEditor.h
--- a/Engine/Plugins/Experimental/Avalanche/Source/AvalancheEditorCore/Private/AvaEditor.h
+++ b/Engine/Plugins/Experimental/Avalanche/Source/AvalancheEditorCore/Private/AvaEditor.h
@@ -35,7 +35,7 @@ protected:
 
 	void ForEachExtension(TFunctionRef<void(const TSharedRef<IAvaEditorExtension>&)> InFunc) const;
 
-	void RecordActivationChangedEvent();
+	void NotifySceneActivationEvent();
 
 	//~ Begin IAvaEditor
 	virtual void Activate(TSharedPtr<IToolkitHost> InOverrideToolkitHost = nullptr) override;
diff --git a/Engine/Plugins/Experimental/Avalanche/Source/AvalancheEditorCore/Public/IAvaEditorProvider.h b/Engine/Plugins/Experimental/Avalanche/Source/AvalancheEditorCore/Public/IAvaEditorProvider.h
--- a/Engine/Plugins/Experimental/Avalanche/Source/AvalancheEditorCore/Public/IAvaEditorProvider.h
+++ b/Engine/Plugins/Experimental/Avalanche/Source/AvalancheEditorCore/Public/IAvaEditorProvider.h
@@ -27,4 +27,8 @@ public:
 
 	/** Optional: Add/Remove from the Actor array when performing operations such as Copy/Cut */
 	virtual void GetActorsToEdit(TArray<AActor*>& InOutActorsToEdit) const {}
+
+	virtual void OnSceneActivated() {}
+
+	virtual void OnSceneDeactivated() {}
 };
diff --git a/Engine/Plugins/Experimental/Avalanche/Source/AvalancheInteractiveTools/Private/Tools/AvaInteractiveToolsToolBase.cpp b/Engine/Plugins/Experimental/Avalanche/Source/AvalancheInteractiveTools/Private/Tools/AvaInteractiveToolsToolBase.cpp
--- a/Engine/Plugins/Experimental/Avalanche/Source/AvalancheInteractiveTools/Private/Tools/AvaInteractiveToolsToolBase.cpp
+++ b/Engine/Plugins/Experimental/Avalanche/Source/AvalancheInteractiveTools/Private/Tools/AvaInteractiveToolsToolBase.cpp
@@ -603,12 +603,13 @@ void UAvaInteractiveToolsToolBase::OnComplete()
 	if (FEngineAnalytics::IsAvailable())
 	{
 		TArray<FAnalyticsEventAttribute> Attributes;
+		Attributes.Reserve(2);
 		Attributes.Emplace(TEXT("ToolClass"), GetClass()->GetName());
 		if (SpawnedActor)
 		{
-			Attributes.Emplace(TEXT("ActorClass"), SpawnedActor->GetClass()->GetName());	
+			Attributes.Emplace(TEXT("ActorClass"), GetNameSafe(SpawnedActor->GetClass()));
 		}
-		FEngineAnalytics::GetProvider().RecordEvent(TEXT("Editor.Usage.MotionDesign.CompleteTool"), Attributes);
+		FEngineAnalytics::GetProvider().RecordEvent(TEXT("Editor.Usage.MotionDesign.PlaceActor"), Attributes);
 	}
 }
 
diff --git a/Engine/Plugins/Experimental/Avalanche/Source/AvalancheShapesEditor/Private/AvaShapeFactory.cpp b/Engine/Plugins/Experimental/Avalanche/Source/AvalancheShapesEditor/Private/AvaShapeFactory.cpp
--- a/Engine/Plugins/Experimental/Avalanche/Source/AvalancheShapesEditor/Private/AvaShapeFactory.cpp
+++ b/Engine/Plugins/Experimental/Avalanche/Source/AvalancheShapesEditor/Private/AvaShapeFactory.cpp
@@ -101,10 +101,10 @@ void UAvaShapeFactory::PostPlaceAsset(TArrayView<const FTypedElementHandle> InHa
 	if (!InPlacementOptions.bIsCreatingPreviewElements && FEngineAnalytics::IsAvailable())
 	{
 		TArray<FAnalyticsEventAttribute> Attributes;
-		if (MeshClass)
-		{
-			Attributes.Emplace(TEXT("MeshClass"), MeshClass->GetName());
-		}
-		FEngineAnalytics::GetProvider().RecordEvent(TEXT("Editor.Usage.MotionDesign.PlaceShape"), Attributes);
+		Attributes.Reserve(3);
+		Attributes.Emplace(TEXT("ToolClass"), GetNameSafe(GetClass()));
+		Attributes.Emplace(TEXT("ActorClass"), GetNameSafe(NewActorClass));
+		Attributes.Emplace(TEXT("SubobjectClass"), GetNameSafe(MeshClass));
+		FEngineAnalytics::GetProvider().RecordEvent(TEXT("Editor.Usage.MotionDesign.PlaceActor"), Attributes);
 	}
 }
diff --git a/Engine/Plugins/Experimental/Avalanche/Source/AvalancheTextEditor/Private/AvaTextActorFactory.cpp b/Engine/Plugins/Experimental/Avalanche/Source/AvalancheTextEditor/Private/AvaTextActorFactory.cpp
--- a/Engine/Plugins/Experimental/Avalanche/Source/AvalancheTextEditor/Private/AvaTextActorFactory.cpp
+++ b/Engine/Plugins/Experimental/Avalanche/Source/AvalancheTextEditor/Private/AvaTextActorFactory.cpp
@@ -15,6 +15,10 @@ void UAvaTextActorFactory::PostPlaceAsset(TArrayView<const FTypedElementHandle>
 	Super::PostPlaceAsset(InHandle, InPlacementInfo, InPlacementOptions);
 	if (!InPlacementOptions.bIsCreatingPreviewElements && FEngineAnalytics::IsAvailable())
 	{
-		FEngineAnalytics::GetProvider().RecordEvent(TEXT("Editor.Usage.MotionDesign.PlaceText"));
+		TArray<FAnalyticsEventAttribute> Attributes;
+		Attributes.Reserve(2);
+		Attributes.Emplace(TEXT("ToolClass"), GetNameSafe(GetClass()));
+		Attributes.Emplace(TEXT("ActorClass"), GetNameSafe(NewActorClass));
+		FEngineAnalytics::GetProvider().RecordEvent(TEXT("Editor.Usage.MotionDesign.PlaceActor"), Attributes);
 	}
 }
diff --git a/Engine/Plugins/Experimental/ClonerEffector/Source/ClonerEffectorEditor/Private/Cloner/CEClonerActorFactory.cpp b/Engine/Plugins/Experimental/ClonerEffector/Source/ClonerEffectorEditor/Private/Cloner/CEClonerActorFactory.cpp
--- a/Engine/Plugins/Experimental/ClonerEffector/Source/ClonerEffectorEditor/Private/Cloner/CEClonerActorFactory.cpp
+++ b/Engine/Plugins/Experimental/ClonerEffector/Source/ClonerEffectorEditor/Private/Cloner/CEClonerActorFactory.cpp
@@ -16,6 +16,10 @@ void UCEClonerActorFactory::PostPlaceAsset(TArrayView<const FTypedElementHandle>
 
 	if (!InPlacementOptions.bIsCreatingPreviewElements && FEngineAnalytics::IsAvailable())
 	{
-		FEngineAnalytics::GetProvider().RecordEvent(TEXT("Editor.Usage.ClonerEffector.PlaceCloner"));
+		TArray<FAnalyticsEventAttribute> Attributes;
+		Attributes.Reserve(2);
+		Attributes.Emplace(TEXT("ToolClass"), GetNameSafe(GetClass()));
+		Attributes.Emplace(TEXT("ActorClass"), GetNameSafe(NewActorClass));
+		FEngineAnalytics::GetProvider().RecordEvent(TEXT("Editor.Usage.MotionDesign.PlaceActor"), Attributes);
 	}
 }
diff --git a/Engine/Plugins/Experimental/ClonerEffector/Source/ClonerEffectorEditor/Private/Effector/CEEffectorActorFactory.cpp b/Engine/Plugins/Experimental/ClonerEffector/Source/ClonerEffectorEditor/Private/Effector/CEEffectorActorFactory.cpp
--- a/Engine/Plugins/Experimental/ClonerEffector/Source/ClonerEffectorEditor/Private/Effector/CEEffectorActorFactory.cpp
+++ b/Engine/Plugins/Experimental/ClonerEffector/Source/ClonerEffectorEditor/Private/Effector/CEEffectorActorFactory.cpp
@@ -16,6 +16,10 @@ void UCEEffectorActorFactory::PostPlaceAsset(TArrayView<const FTypedElementHandl
 
 	if (!InPlacementOptions.bIsCreatingPreviewElements && FEngineAnalytics::IsAvailable())
 	{
-		FEngineAnalytics::GetProvider().RecordEvent(TEXT("Editor.Usage.ClonerEffector.PlaceEffector"));
+		TArray<FAnalyticsEventAttribute> Attributes;
+		Attributes.Reserve(2);
+		Attributes.Emplace(TEXT("ToolClass"), GetNameSafe(GetClass()));
+		Attributes.Emplace(TEXT("ActorClass"), GetNameSafe(NewActorClass));
+		FEngineAnalytics::GetProvider().RecordEvent(TEXT("Editor.Usage.MotionDesign.PlaceActor"), Attributes);
 	}
 }
diff --git a/Engine/Plugins/Experimental/DynamicMaterial/Source/DynamicMaterialEditor/Private/Material/DynamicMaterialInstanceFactory.cpp b/Engine/Plugins/Experimental/DynamicMaterial/Source/DynamicMaterialEditor/Private/Material/DynamicMaterialInstanceFactory.cpp
--- a/Engine/Plugins/Experimental/DynamicMaterial/Source/DynamicMaterialEditor/Private/Material/DynamicMaterialInstanceFactory.cpp
+++ b/Engine/Plugins/Experimental/DynamicMaterial/Source/DynamicMaterialEditor/Private/Material/DynamicMaterialInstanceFactory.cpp
@@ -76,12 +76,11 @@ UObject* UDynamicMaterialInstanceFactory::FactoryCreateNew(UClass* Class, UObjec
 		static const FString AssetType = TEXT("Asset");
 		static const FString SubobjectType = TEXT("Subobject");
 
-		TArray<FAnalyticsEventAttribute> Attribs;
-		const bool bIsAsset = NewInstance->IsAsset();
-
-		Attribs.Add(FAnalyticsEventAttribute(TEXT("Type"), bIsAsset ? AssetType : SubobjectType));
-
-		FEngineAnalytics::GetProvider().RecordEvent(TEXT("Editor.Usage.MaterialDesigner.InstanceCreated"), Attribs);
+		TArray<FAnalyticsEventAttribute> Attributes;
+		Attributes.Reserve(2);
+		Attributes.Add(FAnalyticsEventAttribute(TEXT("Action"), TEXT("InstanceCreated")));
+		Attributes.Add(FAnalyticsEventAttribute(TEXT("ActionDetails"), NewInstance->IsAsset() ? AssetType : SubobjectType));
+		FEngineAnalytics::GetProvider().RecordEvent(TEXT("Editor.Usage.MaterialDesigner"), Attributes);
 	}
 
 	return NewInstance;
diff --git a/Engine/Plugins/Experimental/DynamicMaterial/Source/DynamicMaterialEditor/Private/UI/Menus/DMToolBarMenus.cpp b/Engine/Plugins/Experimental/DynamicMaterial/Source/DynamicMaterialEditor/Private/UI/Menus/DMToolBarMenus.cpp
--- a/Engine/Plugins/Experimental/DynamicMaterial/Source/DynamicMaterialEditor/Private/UI/Menus/DMToolBarMenus.cpp
+++ b/Engine/Plugins/Experimental/DynamicMaterial/Source/DynamicMaterialEditor/Private/UI/Menus/DMToolBarMenus.cpp
@@ -283,7 +283,7 @@ void FDMToolBarMenus::OpenMaterialEditorFromContext(UDMMenuContext* InMenuContex
 		{
 			if (FEngineAnalytics::IsAvailable())
 			{
-				FEngineAnalytics::GetProvider().RecordEvent(TEXT("Editor.Usage.MaterialDesigner.OpenedGeneratedMaterial"));
+				FEngineAnalytics::GetProvider().RecordEvent(TEXT("Editor.Usage.MaterialDesigner"), TEXT("Action"), TEXT("OpenedGeneratedMaterial"));
 			}
 
 			ModelEditorOnlyData->OpenMaterialEditor();
@@ -318,7 +318,7 @@ void FDMToolBarMenus::ExportMaterialInstanceFromInstance(TWeakObjectPtr<UDynamic
 
 			if (FEngineAnalytics::IsAvailable())
 			{
-				FEngineAnalytics::GetProvider().RecordEvent(TEXT("Editor.Usage.MaterialDesigner.ExportedMaterialInstance"));
+				FEngineAnalytics::GetProvider().RecordEvent(TEXT("Editor.Usage.MaterialDesigner"), TEXT("Action"), TEXT("ExportedMaterialInstance"));
 			}
 		}
 	}
@@ -363,7 +363,7 @@ void FDMToolBarMenus::ExportMaterialModelFromModel(TWeakObjectPtr<UDynamicMateri
 
 	if (FEngineAnalytics::IsAvailable())
 	{
-		FEngineAnalytics::GetProvider().RecordEvent(TEXT("Editor.Usage.MaterialDesigner.ExportedGeneratedMaterial"));
+		FEngineAnalytics::GetProvider().RecordEvent(TEXT("Editor.Usage.MaterialDesigner"), TEXT("Action"), TEXT("ExportedGeneratedMaterial"));
 	}
 }
 
@@ -418,7 +418,7 @@ void FDMToolBarMenus::SnapshotMaterial(TWeakObjectPtr<UDynamicMaterialModelBase>
 
 	if (FEngineAnalytics::IsAvailable())
 	{
-		FEngineAnalytics::GetProvider().RecordEvent(TEXT("Editor.Usage.MaterialDesigner.SnapshotMaterial"));
+		FEngineAnalytics::GetProvider().RecordEvent(TEXT("Editor.Usage.MaterialDesigner"), TEXT("Action"), TEXT("SnapshotMaterial"));
 	}
 }
 
diff --git a/Engine/Plugins/Experimental/DynamicMaterial/Source/DynamicMaterialEditor/Private/Utils/DMMaterialModelFunctionLibrary.cpp b/Engine/Plugins/Experimental/DynamicMaterial/Source/DynamicMaterialEditor/Private/Utils/DMMaterialModelFunctionLibrary.cpp
--- a/Engine/Plugins/Experimental/DynamicMaterial/Source/DynamicMaterialEditor/Private/Utils/DMMaterialModelFunctionLibrary.cpp
+++ b/Engine/Plugins/Experimental/DynamicMaterial/Source/DynamicMaterialEditor/Private/Utils/DMMaterialModelFunctionLibrary.cpp
@@ -230,7 +230,7 @@ UDynamicMaterialInstance* UDMMaterialModelFunctionLibrary::ExportMaterialInstanc
 
 	if (FEngineAnalytics::IsAvailable())
 	{
-		FEngineAnalytics::GetProvider().RecordEvent(TEXT("Editor.Usage.MaterialDesigner.ExportedMaterialInstance"));
+		FEngineAnalytics::GetProvider().RecordEvent(TEXT("Editor.Usage.MaterialDesigner"), TEXT("Action"), TEXT("ExportedMaterialInstance"));
 	}
 
 	return NewInstance;
@@ -292,7 +292,7 @@ UMaterial* UDMMaterialModelFunctionLibrary::ExportGeneratedMaterial(UDynamicMate
 
 	if (FEngineAnalytics::IsAvailable())
 	{
-		FEngineAnalytics::GetProvider().RecordEvent(TEXT("Editor.Usage.MaterialDesigner.ExportedGeneratedMaterial"));
+		FEngineAnalytics::GetProvider().RecordEvent(TEXT("Editor.Usage.MaterialDesigner"), TEXT("Action"), TEXT("ExportedGeneratedMaterial"));
 	}
 
 	return Cast<UMaterial>(NewAsset);
@@ -378,7 +378,7 @@ UDynamicMaterialModel* UDMMaterialModelFunctionLibrary::ExportToTemplateMaterial
 
 	if (FEngineAnalytics::IsAvailable())
 	{
-		FEngineAnalytics::GetProvider().RecordEvent(TEXT("Editor.Usage.MaterialDesigner.ExportToTemplateMaterialModel"));
+		FEngineAnalytics::GetProvider().RecordEvent(TEXT("Editor.Usage.MaterialDesigner"), TEXT("Action"), TEXT("ExportToTemplateMaterialModel"));
 	}
 
 	return NewModel;
@@ -486,7 +486,7 @@ UDynamicMaterialInstance* UDMMaterialModelFunctionLibrary::ExportToTemplateMater
 
 	if (FEngineAnalytics::IsAvailable())
 	{
-		FEngineAnalytics::GetProvider().RecordEvent(TEXT("Editor.Usage.MaterialDesigner.ExportToTemplateMaterialInstance"));
+		FEngineAnalytics::GetProvider().RecordEvent(TEXT("Editor.Usage.MaterialDesigner"), TEXT("Action"), TEXT("ExportToTemplateMaterialInstance"));
 	}
 
 	return NewInstance;
diff --git a/Engine/Plugins/Experimental/PropertyAnimatorCore/Source/PropertyAnimatorCore/Private/Subsystems/PropertyAnimatorCoreSubsystem.cpp b/Engine/Plugins/Experimental/PropertyAnimatorCore/Source/PropertyAnimatorCore/Private/Subsystems/PropertyAnimatorCoreSubsystem.cpp
--- a/Engine/Plugins/Experimental/PropertyAnimatorCore/Source/PropertyAnimatorCore/Private/Subsystems/PropertyAnimatorCoreSubsystem.cpp
+++ b/Engine/Plugins/Experimental/PropertyAnimatorCore/Source/PropertyAnimatorCore/Private/Subsystems/PropertyAnimatorCoreSubsystem.cpp
@@ -361,8 +361,10 @@ TSet<UPropertyAnimatorCoreBase*> UPropertyAnimatorCoreSubsystem::CreateAnimators
 	if (FEngineAnalytics::IsAvailable())
 	{
 		TArray<FAnalyticsEventAttribute> Attributes;
-		Attributes.Emplace(TEXT("Class"), InAnimatorClass->GetName());
-		FEngineAnalytics::GetProvider().RecordEvent(TEXT("Editor.Usage.PropertyAnimator.CreateAnimator"), Attributes);
+		Attributes.Reserve(2);
+		Attributes.Emplace(TEXT("Action"), TEXT("Created"));
+		Attributes.Emplace(TEXT("Class"), GetNameSafe(InAnimatorClass));
+		FEngineAnalytics::GetProvider().RecordEvent(TEXT("Editor.Usage.PropertyAnimator.Animator"), Attributes);
 	}
 #endif
 
diff --git a/Engine/Plugins/Experimental/SVGImporter/Source/SVGImporterEditor/Private/Factories/SVGActorFactory.cpp b/Engine/Plugins/Experimental/SVGImporter/Source/SVGImporterEditor/Private/Factories/SVGActorFactory.cpp
--- a/Engine/Plugins/Experimental/SVGImporter/Source/SVGImporterEditor/Private/Factories/SVGActorFactory.cpp
+++ b/Engine/Plugins/Experimental/SVGImporter/Source/SVGImporterEditor/Private/Factories/SVGActorFactory.cpp
@@ -52,6 +52,10 @@ void USVGActorFactory::PostPlaceAsset(TArrayView<const FTypedElementHandle> InHa
 
 	if (!InPlacementOptions.bIsCreatingPreviewElements && FEngineAnalytics::IsAvailable())
 	{
-		FEngineAnalytics::GetProvider().RecordEvent(TEXT("Editor.Usage.SVGImporter.PlaceSVG"));
+		TArray<FAnalyticsEventAttribute> Attributes;
+		Attributes.Reserve(2);
+		Attributes.Emplace(TEXT("ToolClass"), GetNameSafe(GetClass()));
+		Attributes.Emplace(TEXT("ActorClass"), GetNameSafe(NewActorClass));
+		FEngineAnalytics::GetProvider().RecordEvent(TEXT("Editor.Usage.MotionDesign.PlaceActor"), Attributes);
 	}
 }
