Material Designer
- "Material Instance" has been renamed "Material Designer Material"
- "Material Dynamic" has been renamed "Material Designer Instance"

diff --git a/Engine/Plugins/Experimental/Avalanche/Source/AvalancheShapes/Private/DynamicMeshes/AvaShapeDynMeshBase.cpp b/Engine/Plugins/Experimental/Avalanche/Source/AvalancheShapes/Private/DynamicMeshes/AvaShapeDynMeshBase.cpp
--- a/Engine/Plugins/Experimental/Avalanche/Source/AvalancheShapes/Private/DynamicMeshes/AvaShapeDynMeshBase.cpp
+++ b/Engine/Plugins/Experimental/Avalanche/Source/AvalancheShapes/Private/DynamicMeshes/AvaShapeDynMeshBase.cpp
@@ -1198,7 +1198,7 @@ void UAvaShapeDynamicMeshBase::OnMaterialTypeChanged(int32 MaterialIndex)
 
 #if WITH_EDITOR
 					const IDynamicMaterialEditorModule& MaterialDesignerModule = IDynamicMaterialEditorModule::Get();
-					MaterialDesignerModule.OpenMaterialInstance(NewInstance, GetWorld(), /* Invoke Tab */ true);
+					MaterialDesignerModule.OpenMaterial(NewInstance, GetWorld(), /* Invoke Tab */ true);
 #endif
 				}
 				break;
diff --git a/Engine/Plugins/Experimental/DynamicMaterial/Source/DynamicMaterial/Public/Components/DMMaterialComponentDynamic.h b/Engine/Plugins/Experimental/DynamicMaterial/Source/DynamicMaterial/Public/Components/DMMaterialComponentDynamic.h
--- a/Engine/Plugins/Experimental/DynamicMaterial/Source/DynamicMaterial/Public/Components/DMMaterialComponentDynamic.h
+++ b/Engine/Plugins/Experimental/DynamicMaterial/Source/DynamicMaterial/Public/Components/DMMaterialComponentDynamic.h
@@ -10,7 +10,7 @@ class UDynamicMaterialModelDynamic;
 /**
  * Base version of a dynamic material component. Links to the original in the parent material model.
  */
-UCLASS(MinimalAPI, BlueprintType, Abstract, ClassGroup = "Material Designer", meta = (DisplayName = "Material Designer Value Dynamic"))
+UCLASS(MinimalAPI, BlueprintType, Abstract, ClassGroup = "Material Designer", meta = (DisplayName = "Material Designer Value Instance"))
 class UDMMaterialComponentDynamic : public UDMMaterialComponent
 {
 	GENERATED_BODY()
diff --git a/Engine/Plugins/Experimental/DynamicMaterial/Source/DynamicMaterial/Public/Components/DMMaterialParameter.h b/Engine/Plugins/Experimental/DynamicMaterial/Source/DynamicMaterial/Public/Components/DMMaterialParameter.h
--- a/Engine/Plugins/Experimental/DynamicMaterial/Source/DynamicMaterial/Public/Components/DMMaterialParameter.h
+++ b/Engine/Plugins/Experimental/DynamicMaterial/Source/DynamicMaterial/Public/Components/DMMaterialParameter.h
@@ -8,7 +8,7 @@
 class UDynamicMaterialModel;
 
 /**
- * A parameter on a Material Designer Instance.
+ * A parameter on a Material Designer Material.
  */
 UCLASS(MinimalAPI, BlueprintType, meta = (DisplayName = "Material Designer Parameter"))
 class UDMMaterialParameter : public UDMMaterialLinkedComponent
diff --git a/Engine/Plugins/Experimental/DynamicMaterial/Source/DynamicMaterial/Public/Components/DMMaterialValueDynamic.h b/Engine/Plugins/Experimental/DynamicMaterial/Source/DynamicMaterial/Public/Components/DMMaterialValueDynamic.h
--- a/Engine/Plugins/Experimental/DynamicMaterial/Source/DynamicMaterial/Public/Components/DMMaterialValueDynamic.h
+++ b/Engine/Plugins/Experimental/DynamicMaterial/Source/DynamicMaterial/Public/Components/DMMaterialValueDynamic.h
@@ -16,7 +16,7 @@ class UMaterialInstanceDynamic;
 /**
  * A value used inside an instanced material instance. Links to the original value in the parent material.
  */
-UCLASS(MinimalAPI, BlueprintType, Abstract, ClassGroup = "Material Designer", meta = (DisplayName = "Material Designer Value Dynamic"))
+UCLASS(MinimalAPI, BlueprintType, Abstract, ClassGroup = "Material Designer", meta = (DisplayName = "Material Designer Value Instance"))
 class UDMMaterialValueDynamic : public UDMMaterialComponentDynamic, public IDMJsonSerializable
 {
 	GENERATED_BODY()
diff --git a/Engine/Plugins/Experimental/DynamicMaterial/Source/DynamicMaterial/Public/Components/DMTextureUVDynamic.h b/Engine/Plugins/Experimental/DynamicMaterial/Source/DynamicMaterial/Public/Components/DMTextureUVDynamic.h
--- a/Engine/Plugins/Experimental/DynamicMaterial/Source/DynamicMaterial/Public/Components/DMTextureUVDynamic.h
+++ b/Engine/Plugins/Experimental/DynamicMaterial/Source/DynamicMaterial/Public/Components/DMTextureUVDynamic.h
@@ -15,7 +15,7 @@ class UMaterialInstanceDynamic;
 /**
  * A texture uv used inside a instanced material instance. Links to the original texture uv in the parent material.
  */
-UCLASS(MinimalAPI, BlueprintType, ClassGroup = "Material Designer", meta = (DisplayName = "Material Designer Texture UV Dynamic"))
+UCLASS(MinimalAPI, BlueprintType, ClassGroup = "Material Designer", meta = (DisplayName = "Material Designer Texture UV Instance"))
 class UDMTextureUVDynamic : public UDMMaterialComponentDynamic
 {
 	GENERATED_BODY()
diff --git a/Engine/Plugins/Experimental/DynamicMaterial/Source/DynamicMaterial/Public/Material/DynamicMaterialInstance.h b/Engine/Plugins/Experimental/DynamicMaterial/Source/DynamicMaterial/Public/Material/DynamicMaterialInstance.h
--- a/Engine/Plugins/Experimental/DynamicMaterial/Source/DynamicMaterial/Public/Material/DynamicMaterialInstance.h
+++ b/Engine/Plugins/Experimental/DynamicMaterial/Source/DynamicMaterial/Public/Material/DynamicMaterialInstance.h
@@ -13,7 +13,7 @@ class UDMMaterialStageInputTextureUV;
 class UDMMaterialValue;
 #endif
 
-/** A Material Instance Dynamic with its own integrated Material Designer Model that generates the based Material. */
+/** A Material Designer Material with its own integrated Material Designer Model that generates the base Material. */
 UCLASS(MinimalAPI, ClassGroup = "Material Designer", DefaultToInstanced, BlueprintType, meta = (DisplayThumbnail = "true"))
 class UDynamicMaterialInstance : public UMaterialInstanceDynamic
 {
@@ -22,7 +22,7 @@ class UDynamicMaterialInstance : public UMaterialInstanceDynamic
 public:
 	UDynamicMaterialInstance();
 
-	/** Returns the Material Model associated with this Material Designer Instance. */
+	/** Returns the Material Model associated with this Material Designer Material. */
 	UFUNCTION(BlueprintPure, Category = "Material Designer")
 	DYNAMICMATERIAL_API UDynamicMaterialModelBase* GetMaterialModelBase();
 
diff --git a/Engine/Plugins/Experimental/DynamicMaterial/Source/DynamicMaterial/Public/Model/DynamicMaterialModel.h b/Engine/Plugins/Experimental/DynamicMaterial/Source/DynamicMaterial/Public/Model/DynamicMaterialModel.h
--- a/Engine/Plugins/Experimental/DynamicMaterial/Source/DynamicMaterial/Public/Model/DynamicMaterialModel.h
+++ b/Engine/Plugins/Experimental/DynamicMaterial/Source/DynamicMaterial/Public/Model/DynamicMaterialModel.h
@@ -187,10 +187,10 @@ public:
 	bool ConditionalFreeParameter(UDMMaterialParameter* InParameter);
 #endif
 
-	/** Called by the value when it updates. Updates Material Designer Instance and triggers the delegate. */
+	/** Called by the value when it updates. Updates Material Designer Material and triggers the delegate. */
 	void OnValueUpdated(UDMMaterialValue* InValue, EDMUpdateType InUpdateType);
 
-	/** Called by the texture uv when it updates. Updates Material Designer Instance and triggers the delegate. */
+	/** Called by the texture uv when it updates. Updates Material Designer Material and triggers the delegate. */
 	void OnTextureUVUpdated(UDMTextureUV* InTextureUV);
 
 	/** Returns the value update delegate to (un)subscribe. */
@@ -240,7 +240,7 @@ protected:
 	UPROPERTY(VisibleInstanceOnly, Instanced, BlueprintReadOnly, TextExportTransient, Category = "Material Designer")
 	TObjectPtr<UMaterial> DynamicMaterial = nullptr;
 
-	/** Material Designer Instance representing the MID for this Model. */
+	/** Material Designer Material representing the MID for this Model. */
 	UPROPERTY(VisibleInstanceOnly, BlueprintReadOnly, TextExportTransient, Category = "Material Designer")
 	TObjectPtr<UDynamicMaterialInstance> DynamicMaterialInstance;
 
diff --git a/Engine/Plugins/Experimental/DynamicMaterial/Source/DynamicMaterial/Public/Model/DynamicMaterialModelBase.h b/Engine/Plugins/Experimental/DynamicMaterial/Source/DynamicMaterial/Public/Model/DynamicMaterialModelBase.h
--- a/Engine/Plugins/Experimental/DynamicMaterial/Source/DynamicMaterial/Public/Model/DynamicMaterialModelBase.h
+++ b/Engine/Plugins/Experimental/DynamicMaterial/Source/DynamicMaterial/Public/Model/DynamicMaterialModelBase.h
@@ -23,16 +23,16 @@ public:
 	/**
 	 * Returns the Material Designer Model that is the base for this Model.
 	 * It will be this object for a Model.
-	 * It will be the parent Model for a Dynamic Model.
+	 * It will be the parent Model for a Model Instance.
 	 */
 	UFUNCTION(BlueprintPure, Category = "Material Designer")
 	virtual UDynamicMaterialModel* ResolveMaterialModel() PURE_VIRTUAL(UDynamicMaterialModelBase::ResolveMaterialModel, return nullptr;)
 
-	/** Returns the Material Designer Instance that contains this Model, if there is one. */
+	/** Returns the Material Designer Material that contains this Model, if there is one. */
 	UFUNCTION(BlueprintPure, Category = "Material Designer")
 	virtual UDynamicMaterialInstance* GetDynamicMaterialInstance() const PURE_VIRTUAL(UDynamicMaterialModelBase::GetDynamicMaterialInstance, return nullptr;)
 
-	/** Sets the Material Designer Instance for this Model. */
+	/** Sets the Material Designer Material for this Model. */
 	virtual void SetDynamicMaterialInstance(UDynamicMaterialInstance* InDynamicMaterialInstance) PURE_VIRTUAL(UDynamicMaterialModelBase::SetDynamicMaterialInstance)
 
 	/** Returns the UMaterial from the resolved Material Model. */
diff --git a/Engine/Plugins/Experimental/DynamicMaterial/Source/DynamicMaterial/Public/Model/DynamicMaterialModelDynamic.h b/Engine/Plugins/Experimental/DynamicMaterial/Source/DynamicMaterial/Public/Model/DynamicMaterialModelDynamic.h
--- a/Engine/Plugins/Experimental/DynamicMaterial/Source/DynamicMaterial/Public/Model/DynamicMaterialModelDynamic.h
+++ b/Engine/Plugins/Experimental/DynamicMaterial/Source/DynamicMaterial/Public/Model/DynamicMaterialModelDynamic.h
@@ -22,7 +22,7 @@ DECLARE_MULTICAST_DELEGATE_TwoParams(FDMOnTextureUVDynamicUpdated, UDynamicMater
 /**
  * Represents a MID-like version of a Material Designer Model. Uses dynamic values/texture uvs to link to the original model.
  */
-UCLASS(MinimalAPI, BlueprintType, ClassGroup = "Material Designer", meta = (DisplayName = "Material Designer Model Dynamic"))
+UCLASS(MinimalAPI, BlueprintType, ClassGroup = "Material Designer", meta = (DisplayName = "Material Designer Model Instance"))
 class UDynamicMaterialModelDynamic : public UDynamicMaterialModelBase
 {
 	GENERATED_BODY()
@@ -33,9 +33,9 @@ public:
 
 #if WITH_EDITOR
 	/**
-	 * Create a new Material Model Dynamic based on a parent model.
-	 * @param InOuter Could be the transient package, an asset package or a Material Designer Instance.
-	 * @return A new Dynamic Model with its components already initialized.
+	 * Create a new Material Designer Model Instance based on a parent Model.
+	 * @param InOuter Could be the transient package, an asset package or a Material Designer Material.
+	 * @return A new Material Designer Model Instance with its components already initialized.
 	 */
 	UFUNCTION(BlueprintCallable, Category = "Motion Design")
 	static DYNAMICMATERIAL_API UDynamicMaterialModelDynamic* Create(UObject* InOuter, UDynamicMaterialModel* InParentModel);
@@ -71,10 +71,10 @@ public:
 	DYNAMICMATERIAL_API void EnsureComponents();
 #endif
 
-	/** Called when a value changes so that the Material Designer Instance can be updated and the event broadcast */
+	/** Called when a value changes so that the Material Designer Material can be updated and the event broadcast */
 	void OnValueUpdated(UDMMaterialValueDynamic* InValueDynamic);
 
-	/** Called when a texture uv changes so that the Material Designer Instance can be updated and the event broadcast */
+	/** Called when a texture uv changes so that the Material Designer Material can be updated and the event broadcast */
 	void OnTextureUVUpdated(UDMTextureUVDynamic* InTextureUVDynamic);
 
 	/** Returns the on value update method so it can (un)subscribed to. */
@@ -137,7 +137,7 @@ protected:
 	UDynamicMaterialModel* EnsureParentModel();
 
 #if WITH_EDITOR
-	/** Scans the parent model and adds all the components from there as Dynamic Components. Should not be called twice. */
+	/** Scans the parent Material Model and adds all the components from there as Instance Components. Should not be called twice. */
 	void InitComponents();
 #endif
 };
diff --git a/Engine/Plugins/Experimental/DynamicMaterial/Source/DynamicMaterialEditor/Private/DMContentBrowserIntegration.cpp b/Engine/Plugins/Experimental/DynamicMaterial/Source/DynamicMaterialEditor/Private/DMContentBrowserIntegration.cpp
--- a/Engine/Plugins/Experimental/DynamicMaterial/Source/DynamicMaterialEditor/Private/DMContentBrowserIntegration.cpp
+++ b/Engine/Plugins/Experimental/DynamicMaterial/Source/DynamicMaterialEditor/Private/DMContentBrowserIntegration.cpp
@@ -69,28 +69,28 @@ void FDMContentBrowserIntegration::Disintegrate()
 void FDMContentBrowserIntegration::ExtendMenu(FMenuBuilder& InMenuBuilder, const TArray<FAssetData>& InSelectedAssets)
 {
 	InMenuBuilder.AddMenuEntry(
-		LOCTEXT("CreateMaterialDesignerInstanceFromTextureSet", "Create Material Designer Instance"),
-		LOCTEXT("CreateMaterialDesignerInstanceFromTextureSetTooltip", "Creates a Material Designer Instance in the content browser using a Texture Set."),
+		LOCTEXT("CreateMaterialDesignerMaterialFromTextureSet", "Create Material Designer Material"),
+		LOCTEXT("CreateMaterialDesignerInstanceFromTextureSetTooltip", "Creates a Material Designer Material in the content browser using a Texture Set."),
 		FSlateIconFinder::FindIconForClass(UMaterial::StaticClass()),
-		FUIAction(FExecuteAction::CreateStatic(&FDMContentBrowserIntegration::CreateMaterialDesignerInstanceFromTextureSet, InSelectedAssets))
+		FUIAction(FExecuteAction::CreateStatic(&FDMContentBrowserIntegration::CreateMaterialDesignerMaterialFromTextureSet, InSelectedAssets))
 	);
 
 	InMenuBuilder.AddMenuEntry(
-		LOCTEXT("UpdateMaterialDesignerInstanceFromTextureSetAdd", "Update Material Designer Instance (Add)"),
-		LOCTEXT("UpdateMaterialDesignerInstanceFromTextureAddSetTooltip", "Updates the opened Material Designer Instance using a Texture Set, adding new layers to the Model."),
+		LOCTEXT("UpdateMaterialDesignerInstanceFromTextureSetAdd", "Update Material Designer Material (Add)"),
+		LOCTEXT("UpdateMaterialDesignerInstanceFromTextureAddSetTooltip", "Updates the opened Material Designer Material using a Texture Set, adding new layers to the Material Model."),
 		FSlateIconFinder::FindIconForClass(UMaterial::StaticClass()),
-		FUIAction(FExecuteAction::CreateStatic(&FDMContentBrowserIntegration::UpdateMaterialDesignerInstanceFromTextureSet, InSelectedAssets, /* Replace */ false))
+		FUIAction(FExecuteAction::CreateStatic(&FDMContentBrowserIntegration::UpdateMaterialDesignerMaterialFromTextureSet, InSelectedAssets, /* Replace */ false))
 	);
 
 	InMenuBuilder.AddMenuEntry(
-		LOCTEXT("UpdateMaterialDesignerInstanceFromTextureSetReplace", "Update Material Designer Instance (Replace)"),
-		LOCTEXT("UpdateMaterialDesignerInstanceFromTextureSetReplaceTooltip", "Updates the opened Material Designer Instance using a Texture Set, replacing slots in the Model."),
+		LOCTEXT("UpdateMaterialDesignerInstanceFromTextureSetReplace", "Update Material Designer Material (Replace)"),
+		LOCTEXT("UpdateMaterialDesignerInstanceFromTextureSetReplaceTooltip", "Updates the opened Material Designer Material using a Texture Set, replacing layers in the Material Model."),
 		FSlateIconFinder::FindIconForClass(UMaterial::StaticClass()),
-		FUIAction(FExecuteAction::CreateStatic(&FDMContentBrowserIntegration::UpdateMaterialDesignerInstanceFromTextureSet, InSelectedAssets, /* Replace */ true))
+		FUIAction(FExecuteAction::CreateStatic(&FDMContentBrowserIntegration::UpdateMaterialDesignerMaterialFromTextureSet, InSelectedAssets, /* Replace */ true))
 	);
 }
 
-void FDMContentBrowserIntegration::CreateMaterialDesignerInstanceFromTextureSet(TArray<FAssetData> InSelectedAssets)
+void FDMContentBrowserIntegration::CreateMaterialDesignerMaterialFromTextureSet(TArray<FAssetData> InSelectedAssets)
 {
 	if (InSelectedAssets.IsEmpty())
 	{
@@ -100,13 +100,13 @@ void FDMContentBrowserIntegration::CreateMaterialDesignerInstanceFromTextureSet(
 	UDMTextureSetBlueprintFunctionLibrary::CreateTextureSetFromAssetsInteractive(
 		InSelectedAssets,
 		FDMTextureSetBuilderOnComplete::CreateStatic(
-			&FDMContentBrowserIntegration::OnCreateMaterialDesignerInstanceFromTextureSetComplete,
+			&FDMContentBrowserIntegration::OnCreateMaterialDesignerMaterialFromTextureSetComplete,
 			InSelectedAssets[0].PackagePath.ToString()
 		)
 	);
 }
 
-void FDMContentBrowserIntegration::OnCreateMaterialDesignerInstanceFromTextureSetComplete(UDMTextureSet* InTextureSet, bool bInAccepted,
+void FDMContentBrowserIntegration::OnCreateMaterialDesignerMaterialFromTextureSetComplete(UDMTextureSet* InTextureSet, bool bInAccepted,
 	FString InPath)
 {
 	if (!InTextureSet || !bInAccepted)
@@ -131,7 +131,7 @@ void FDMContentBrowserIntegration::OnCreateMaterialDesignerInstanceFromTextureSe
 	ON_SCOPE_EXIT
 	{
 		const IDynamicMaterialEditorModule& MaterialDesignerModule = IDynamicMaterialEditorModule::Get();
-		MaterialDesignerModule.OpenMaterialInstance(NewInstance, nullptr, /* Invoke Tab */ true);
+		MaterialDesignerModule.OpenMaterial(NewInstance, nullptr, /* Invoke Tab */ true);
 	};
 
 	UDynamicMaterialModelEditorOnlyData* EditorOnlyData = UDynamicMaterialModelEditorOnlyData::Get(NewInstance);
@@ -169,7 +169,7 @@ void FDMContentBrowserIntegration::OnCreateMaterialDesignerInstanceFromTextureSe
 	FAssetRegistryModule::AssetCreated(NewInstance);
 }
 
-void FDMContentBrowserIntegration::UpdateMaterialDesignerInstanceFromTextureSet(TArray<FAssetData> InSelectedAssets, bool bInReplace)
+void FDMContentBrowserIntegration::UpdateMaterialDesignerMaterialFromTextureSet(TArray<FAssetData> InSelectedAssets, bool bInReplace)
 {
 	if (InSelectedAssets.IsEmpty())
 	{
@@ -179,13 +179,13 @@ void FDMContentBrowserIntegration::UpdateMaterialDesignerInstanceFromTextureSet(
 	UDMTextureSetBlueprintFunctionLibrary::CreateTextureSetFromAssetsInteractive(
 		InSelectedAssets,
 		FDMTextureSetBuilderOnComplete::CreateStatic(
-			&FDMContentBrowserIntegration::OnUpdateMaterialDesignerInstanceFromTextureSetComplete,
+			&FDMContentBrowserIntegration::OnUpdateMaterialDesignerMaterialFromTextureSetComplete,
 			bInReplace
 		)
 	);
 }
 
-void FDMContentBrowserIntegration::OnUpdateMaterialDesignerInstanceFromTextureSetComplete(UDMTextureSet* InTextureSet, bool bInAccepted,
+void FDMContentBrowserIntegration::OnUpdateMaterialDesignerMaterialFromTextureSetComplete(UDMTextureSet* InTextureSet, bool bInAccepted,
 	bool bInReplace)
 {
 	if (!InTextureSet || !bInAccepted)
@@ -250,10 +250,10 @@ TSharedRef<FExtender> FDMContentBrowserIntegration::OnExtendContentBrowserAssetS
 			[InSelectedAssets](FMenuBuilder& InMenuBuilder)
 			{
 				InMenuBuilder.AddMenuEntry(
-					LOCTEXT("CreateDynamic", "Create Material Designer Dynamic"),
-					LOCTEXT("CreateDynamicooltip", "Create a dynamic instance from a Material Designer Instance or Model."),
+					LOCTEXT("CreateInstance", "Create Material Designer Instance"),
+					LOCTEXT("CreateInstanceTooltip", "Create a Material Designer Instance from a Material Designer Material."),
 					FSlateIconFinder::FindIconForClass(UMaterial::StaticClass()),
-					FUIAction(FExecuteAction::CreateStatic(&FDMContentBrowserIntegration::CreateDynamic, InSelectedAssets))
+					FUIAction(FExecuteAction::CreateStatic(&FDMContentBrowserIntegration::CreateInstance, InSelectedAssets))
 				);
 			}
 		)
@@ -262,7 +262,7 @@ TSharedRef<FExtender> FDMContentBrowserIntegration::OnExtendContentBrowserAssetS
 	return Extender;
 }
 
-void FDMContentBrowserIntegration::CreateDynamic(TArray<FAssetData> InSelectedAssets)
+void FDMContentBrowserIntegration::CreateInstance(TArray<FAssetData> InSelectedAssets)
 {
 	for (const FAssetData& SelectedAsset : InSelectedAssets)
 	{
@@ -270,20 +270,20 @@ void FDMContentBrowserIntegration::CreateDynamic(TArray<FAssetData> InSelectedAs
 		{
 			if (AssetClass->IsChildOf<UDynamicMaterialModel>())
 			{
-				CreateModelDynamic(Cast<UDynamicMaterialModel>(SelectedAsset.GetAsset()));
+				CreateModelInstance(Cast<UDynamicMaterialModel>(SelectedAsset.GetAsset()));
 				break;
 			}
 
 			if (AssetClass->IsChildOf<UDynamicMaterialInstance>())
 			{
-				CreateInstanceDynamic(Cast<UDynamicMaterialInstance>(SelectedAsset.GetAsset()));
+				CreateMaterialInstance(Cast<UDynamicMaterialInstance>(SelectedAsset.GetAsset()));
 				break;
 			}
 		}
 	}
 }
 
-void FDMContentBrowserIntegration::CreateModelDynamic(UDynamicMaterialModel* InModel)
+void FDMContentBrowserIntegration::CreateModelInstance(UDynamicMaterialModel* InModel)
 {
 	if (!InModel)
 	{
@@ -302,7 +302,7 @@ void FDMContentBrowserIntegration::CreateModelDynamic(UDynamicMaterialModel* InM
 		const EAppReturnType::Type Result = FMessageDialog::Open(
 			EAppMsgType::YesNo,
 			LOCTEXT("ExportMaterialFromModel", 
-				"Generating a Material Designer Dynamic requires that the generated material be exported from its package.\n\n"
+				"Generating a Material Designer Instance requires that the Generated Material be exported from its package.\n\n"
 				"The package containing the material will be saved. This may be a level.\n\n"
 				"Continue?")
 		);
@@ -366,7 +366,7 @@ void FDMContentBrowserIntegration::CreateModelDynamic(UDynamicMaterialModel* InM
 	FAssetRegistryModule::AssetCreated(ModelDynamic);
 }
 
-void FDMContentBrowserIntegration::CreateInstanceDynamic(UDynamicMaterialInstance* InInstance)
+void FDMContentBrowserIntegration::CreateMaterialInstance(UDynamicMaterialInstance* InInstance)
 {
 	if (!InInstance)
 	{
@@ -392,7 +392,7 @@ void FDMContentBrowserIntegration::CreateInstanceDynamic(UDynamicMaterialInstanc
 		const EAppReturnType::Type Result = FMessageDialog::Open(
 			EAppMsgType::YesNo,
 			LOCTEXT("ExportMaterialFromInstance",
-				"Generating a Material Designer Dynamic requires that the generated material and material model be exported from this package.\n\n"
+				"Generating a Material Designer Instance requires that the Generated Material and Material Model be exported from this package.\n\n"
 				"The package containing the material will be saved. This may be a level.\n\n"
 				"Continue?")
 		);
diff --git a/Engine/Plugins/Experimental/DynamicMaterial/Source/DynamicMaterialEditor/Private/DMContentBrowserIntegration.h b/Engine/Plugins/Experimental/DynamicMaterial/Source/DynamicMaterialEditor/Private/DMContentBrowserIntegration.h
--- a/Engine/Plugins/Experimental/DynamicMaterial/Source/DynamicMaterialEditor/Private/DMContentBrowserIntegration.h
+++ b/Engine/Plugins/Experimental/DynamicMaterial/Source/DynamicMaterialEditor/Private/DMContentBrowserIntegration.h
@@ -20,7 +20,7 @@ public:
 
 	static void Disintegrate();
 
-	static void UpdateMaterialDesignerInstanceFromTextureSet(TArray<FAssetData> InSelectedAssets, bool bInReplace);
+	static void UpdateMaterialDesignerMaterialFromTextureSet(TArray<FAssetData> InSelectedAssets, bool bInReplace);
 
 protected:
 	static FDelegateHandle TextureSetPopulateHandle;
@@ -28,17 +28,17 @@ protected:
 
 	static void ExtendMenu(FMenuBuilder& InMenuBuilder, const TArray<FAssetData>& InSelectedAssets);
 
-	static void CreateMaterialDesignerInstanceFromTextureSet(TArray<FAssetData> InSelectedAssets);
+	static void CreateMaterialDesignerMaterialFromTextureSet(TArray<FAssetData> InSelectedAssets);
 
-	static void OnCreateMaterialDesignerInstanceFromTextureSetComplete(UDMTextureSet* InTextureSet, bool bInAccepted, FString InPath);
+	static void OnCreateMaterialDesignerMaterialFromTextureSetComplete(UDMTextureSet* InTextureSet, bool bInAccepted, FString InPath);
 
-	static void OnUpdateMaterialDesignerInstanceFromTextureSetComplete(UDMTextureSet* InTextureSet, bool bInAccepted, bool bInReplace);
+	static void OnUpdateMaterialDesignerMaterialFromTextureSetComplete(UDMTextureSet* InTextureSet, bool bInAccepted, bool bInReplace);
 
 	static TSharedRef<FExtender> OnExtendContentBrowserAssetSelectionMenu(const TArray<FAssetData>& InSelectedAssets);
 
-	static void CreateDynamic(TArray<FAssetData> InSelectedAssets);
+	static void CreateInstance(TArray<FAssetData> InSelectedAssets);
 
-	static void CreateModelDynamic(UDynamicMaterialModel* InModel);
+	static void CreateModelInstance(UDynamicMaterialModel* InModel);
 
-	static void CreateInstanceDynamic(UDynamicMaterialInstance* InInstance);
+	static void CreateMaterialInstance(UDynamicMaterialInstance* InInstance);
 };
diff --git a/Engine/Plugins/Experimental/DynamicMaterial/Source/DynamicMaterialEditor/Private/DetailsPanel/Widgets/SDMDetailsPanelMaterialInterfaceWidget.cpp b/Engine/Plugins/Experimental/DynamicMaterial/Source/DynamicMaterialEditor/Private/DetailsPanel/Widgets/SDMDetailsPanelMaterialInterfaceWidget.cpp
--- a/Engine/Plugins/Experimental/DynamicMaterial/Source/DynamicMaterialEditor/Private/DetailsPanel/Widgets/SDMDetailsPanelMaterialInterfaceWidget.cpp
+++ b/Engine/Plugins/Experimental/DynamicMaterial/Source/DynamicMaterialEditor/Private/DetailsPanel/Widgets/SDMDetailsPanelMaterialInterfaceWidget.cpp
@@ -30,7 +30,7 @@ void SDMDetailsPanelMaterialInterfaceWidget::Construct(const FArguments& InArgs,
 	UObject* Value = nullptr;
 	InPropertyHandle->GetValue(Value);
 
-	UDynamicMaterialInstance* Instance = GetDynamicMaterialInstance();
+	UDynamicMaterialInstance* Material = GetMaterialDesignerMaterial();
 
 	FObjectPropertyBase* const ObjectProperty = CastField<FObjectPropertyBase>(PropertyHandle->GetProperty());
 
@@ -95,7 +95,7 @@ UObject* SDMDetailsPanelMaterialInterfaceWidget::GetAsset() const
 	return Value;
 }
 
-UDynamicMaterialInstance* SDMDetailsPanelMaterialInterfaceWidget::GetDynamicMaterialInstance() const
+UDynamicMaterialInstance* SDMDetailsPanelMaterialInterfaceWidget::GetMaterialDesignerMaterial() const
 {
 	return Cast<UDynamicMaterialInstance>(GetAsset());
 }
@@ -113,14 +113,14 @@ void SDMDetailsPanelMaterialInterfaceWidget::SetAsset(UObject* NewAsset)
 	PropertyHandle->SetValueFromFormattedString(NewAsset ? NewAsset->GetPathName() : "");
 }
 
-void SDMDetailsPanelMaterialInterfaceWidget::SetDynamicMaterialInstance(UDynamicMaterialInstance* NewInstance)
+void SDMDetailsPanelMaterialInterfaceWidget::SetMaterialDesignerMaterial(UDynamicMaterialInstance* InMaterial)
 {
-	SetAsset(NewInstance);
+	SetAsset(InMaterial);
 }
 
 FText SDMDetailsPanelMaterialInterfaceWidget::GetButtonText() const
 {
-	if (GetDynamicMaterialInstance())
+	if (GetMaterialDesignerMaterial())
 	{
 		return LOCTEXT("OpenMaterialDesignerModel", "Edit with Material Designer");
 	}
@@ -130,17 +130,17 @@ FText SDMDetailsPanelMaterialInterfaceWidget::GetButtonText() const
 
 FReply SDMDetailsPanelMaterialInterfaceWidget::OnButtonClicked()
 {
-	if (GetDynamicMaterialInstance())
+	if (GetMaterialDesignerMaterial())
 	{
-		return OpenDynamicMaterialInstanceTab();
+		return OpenMaterialDesignerTab();
 	}
 
-	return CreateDynamicMaterialInstance();
+	return CreateMaterialDesignerMaterial();
 }
 
-FReply SDMDetailsPanelMaterialInterfaceWidget::CreateDynamicMaterialInstance()
+FReply SDMDetailsPanelMaterialInterfaceWidget::CreateMaterialDesignerMaterial()
 {
-	UDynamicMaterialInstance* Instance = GetDynamicMaterialInstance();
+	UDynamicMaterialInstance* Instance = GetMaterialDesignerMaterial();
 
 	// We already have an instance, so we don't need to create one
 	if (Instance)
@@ -170,9 +170,9 @@ FReply SDMDetailsPanelMaterialInterfaceWidget::CreateDynamicMaterialInstance()
 	return FReply::Handled();
 }
 
-FReply SDMDetailsPanelMaterialInterfaceWidget::ClearDynamicMaterialInstance()
+FReply SDMDetailsPanelMaterialInterfaceWidget::ClearMaterialDesignerMaterial()
 {
-	UDynamicMaterialInstance* Instance = GetDynamicMaterialInstance();
+	UDynamicMaterialInstance* Instance = GetMaterialDesignerMaterial();
 
 	// We don't have an instance, so we don't need to clear it (and don't clear non-MDIs)
 	if (!Instance)
@@ -180,14 +180,14 @@ FReply SDMDetailsPanelMaterialInterfaceWidget::ClearDynamicMaterialInstance()
 		return FReply::Handled();
 	}
 
-	SetDynamicMaterialInstance(nullptr);
+	SetMaterialDesignerMaterial(nullptr);
 
 	return FReply::Handled();
 }
 
-FReply SDMDetailsPanelMaterialInterfaceWidget::OpenDynamicMaterialInstanceTab()
+FReply SDMDetailsPanelMaterialInterfaceWidget::OpenMaterialDesignerTab()
 {
-	UDynamicMaterialInstance* Instance = GetDynamicMaterialInstance();
+	UDynamicMaterialInstance* Instance = GetMaterialDesignerMaterial();
 
 	// We don't have a MDI, so don't try to open it.
 	if (!Instance)
diff --git a/Engine/Plugins/Experimental/DynamicMaterial/Source/DynamicMaterialEditor/Private/DetailsPanel/Widgets/SDMDetailsPanelMaterialInterfaceWidget.h b/Engine/Plugins/Experimental/DynamicMaterial/Source/DynamicMaterialEditor/Private/DetailsPanel/Widgets/SDMDetailsPanelMaterialInterfaceWidget.h
--- a/Engine/Plugins/Experimental/DynamicMaterial/Source/DynamicMaterialEditor/Private/DetailsPanel/Widgets/SDMDetailsPanelMaterialInterfaceWidget.h
+++ b/Engine/Plugins/Experimental/DynamicMaterial/Source/DynamicMaterialEditor/Private/DetailsPanel/Widgets/SDMDetailsPanelMaterialInterfaceWidget.h
@@ -27,13 +27,13 @@ protected:
 	UObject* GetAsset() const;
 	void SetAsset(UObject* NewAsset);
 
-	UDynamicMaterialInstance* GetDynamicMaterialInstance() const;
-	void SetDynamicMaterialInstance(UDynamicMaterialInstance* NewInstance);
+	UDynamicMaterialInstance* GetMaterialDesignerMaterial() const;
+	void SetMaterialDesignerMaterial(UDynamicMaterialInstance* InMaterial);
 
 	FReply OnButtonClicked();
-	FReply CreateDynamicMaterialInstance();
-	FReply ClearDynamicMaterialInstance();
-	FReply OpenDynamicMaterialInstanceTab();
+	FReply CreateMaterialDesignerMaterial();
+	FReply ClearMaterialDesignerMaterial();
+	FReply OpenMaterialDesignerTab();
 
 	FText GetButtonText() const;
 };
diff --git a/Engine/Plugins/Experimental/DynamicMaterial/Source/DynamicMaterialEditor/Private/DetailsPanel/Widgets/SDMMaterialListExtensionWidget.cpp b/Engine/Plugins/Experimental/DynamicMaterial/Source/DynamicMaterialEditor/Private/DetailsPanel/Widgets/SDMMaterialListExtensionWidget.cpp
--- a/Engine/Plugins/Experimental/DynamicMaterial/Source/DynamicMaterialEditor/Private/DetailsPanel/Widgets/SDMMaterialListExtensionWidget.cpp
+++ b/Engine/Plugins/Experimental/DynamicMaterial/Source/DynamicMaterialEditor/Private/DetailsPanel/Widgets/SDMMaterialListExtensionWidget.cpp
@@ -10,7 +10,6 @@
 #include "Engine/World.h"
 #include "IDynamicMaterialEditorModule.h"
 #include "Material/DynamicMaterialInstance.h"
-#include "Material/DynamicMaterialInstanceFactory.h"
 #include "MaterialList.h"
 #include "Widgets/Input/SButton.h"
 #include "Widgets/Layout/SBox.h"
@@ -71,7 +70,7 @@ UObject* SDMMaterialListExtensionWidget::GetAsset() const
 	return MaterialInterface;
 }
 
-UDynamicMaterialInstance* SDMMaterialListExtensionWidget::GetDynamicMaterialInstance() const
+UDynamicMaterialInstance* SDMMaterialListExtensionWidget::GetMaterialDesignerMaterial() const
 {
 	return Cast<UDynamicMaterialInstance>(GetAsset());
 }
@@ -99,7 +98,7 @@ void SDMMaterialListExtensionWidget::SetAsset(UObject* NewAsset)
 	MaterialItemView->ReplaceMaterial(MaterialInterface);
 }
 
-void SDMMaterialListExtensionWidget::SetDynamicMaterialInstance(UDynamicMaterialInstance* NewInstance)
+void SDMMaterialListExtensionWidget::SetMaterialDesignerMaterial(UDynamicMaterialInstance* InMaterial)
 {
 	if (MaterialItemViewWeak.IsValid())
 	{
@@ -116,7 +115,7 @@ void SDMMaterialListExtensionWidget::SetDynamicMaterialInstance(UDynamicMaterial
 							const FMaterialListItem& ListItem = MaterialItemViewWeak.Pin()->GetMaterialListItem();
 							const FDMObjectMaterialProperty MaterialProperty(CurrentComponent, ListItem.SlotIndex);
 
-							if (WorldSubsystem->ExecuteMaterialValueSetterDelegate(MaterialProperty, NewInstance))
+							if (WorldSubsystem->ExecuteMaterialValueSetterDelegate(MaterialProperty, InMaterial))
 							{
 								return;
 							}
@@ -127,12 +126,12 @@ void SDMMaterialListExtensionWidget::SetDynamicMaterialInstance(UDynamicMaterial
 		}
 	}
 
-	SetAsset(NewInstance);
+	SetAsset(InMaterial);
 }
 
 FText SDMMaterialListExtensionWidget::GetButtonText() const
 {
-	if (GetDynamicMaterialInstance())
+	if (GetMaterialDesignerMaterial())
 	{
 		return LOCTEXT("OpenMaterialDesignerModel", "Edit with Material Designer");
 	}
@@ -142,20 +141,20 @@ FText SDMMaterialListExtensionWidget::GetButtonText() const
 
 FReply SDMMaterialListExtensionWidget::OnButtonClicked()
 {
-	if (GetDynamicMaterialInstance())
+	if (GetMaterialDesignerMaterial())
 	{
-		return OpenDynamicMaterialInstanceTab();
+		return OpenMaterialDesignerTab();
 	}
 
-	return CreateDynamicMaterialInstance();
+	return CreateMaterialDesignerMaterial();
 }
 
-FReply SDMMaterialListExtensionWidget::CreateDynamicMaterialInstance()
+FReply SDMMaterialListExtensionWidget::CreateMaterialDesignerMaterial()
 {
-	UDynamicMaterialInstance* Instance = GetDynamicMaterialInstance();
+	UDynamicMaterialInstance* Material = GetMaterialDesignerMaterial();
 
 	// We already have an instance, so we don't need to create one
-	if (Instance)
+	if (Material)
 	{
 		return FReply::Handled();
 	}
@@ -184,22 +183,22 @@ FReply SDMMaterialListExtensionWidget::CreateDynamicMaterialInstance()
 	return FReply::Handled();
 }
 
-FReply SDMMaterialListExtensionWidget::ClearDynamicMaterialInstance()
+FReply SDMMaterialListExtensionWidget::ClearMaterialDesignerMaterial()
 {
-	UDynamicMaterialInstance* Instance = GetDynamicMaterialInstance();
+	UDynamicMaterialInstance* Material = GetMaterialDesignerMaterial();
 
 	// We don't have an instance, so we don't need to clear it (or any other asset in its place)
-	if (!Instance)
+	if (!Material)
 	{
 		return FReply::Handled();
 	}
 
-	SetDynamicMaterialInstance(nullptr);
+	SetMaterialDesignerMaterial(nullptr);
 
 	return FReply::Handled();
 }
 
-FReply SDMMaterialListExtensionWidget::OpenDynamicMaterialInstanceTab()
+FReply SDMMaterialListExtensionWidget::OpenMaterialDesignerTab()
 {
 	UPrimitiveComponent* CurrentComponent = CurrentComponentWeak.Get();
 	if (!ensure(CurrentComponent))
diff --git a/Engine/Plugins/Experimental/DynamicMaterial/Source/DynamicMaterialEditor/Private/DetailsPanel/Widgets/SDMMaterialListExtensionWidget.h b/Engine/Plugins/Experimental/DynamicMaterial/Source/DynamicMaterialEditor/Private/DetailsPanel/Widgets/SDMMaterialListExtensionWidget.h
--- a/Engine/Plugins/Experimental/DynamicMaterial/Source/DynamicMaterialEditor/Private/DetailsPanel/Widgets/SDMMaterialListExtensionWidget.h
+++ b/Engine/Plugins/Experimental/DynamicMaterial/Source/DynamicMaterialEditor/Private/DetailsPanel/Widgets/SDMMaterialListExtensionWidget.h
@@ -32,13 +32,13 @@ protected:
 	UObject* GetAsset() const;
 	void SetAsset(UObject* NewAsset);
 
-	UDynamicMaterialInstance* GetDynamicMaterialInstance() const;
-	void SetDynamicMaterialInstance(UDynamicMaterialInstance* NewInstance);
+	UDynamicMaterialInstance* GetMaterialDesignerMaterial() const;
+	void SetMaterialDesignerMaterial(UDynamicMaterialInstance* InMaterial);
 
 	FReply OnButtonClicked();
-	FReply CreateDynamicMaterialInstance();
-	FReply ClearDynamicMaterialInstance();
-	FReply OpenDynamicMaterialInstanceTab();
+	FReply CreateMaterialDesignerMaterial();
+	FReply ClearMaterialDesignerMaterial();
+	FReply OpenMaterialDesignerTab();
 
 	FText GetButtonText() const;
 };
diff --git a/Engine/Plugins/Experimental/DynamicMaterial/Source/DynamicMaterialEditor/Private/DynamicMaterialEditorModule.cpp b/Engine/Plugins/Experimental/DynamicMaterial/Source/DynamicMaterialEditor/Private/DynamicMaterialEditorModule.cpp
--- a/Engine/Plugins/Experimental/DynamicMaterial/Source/DynamicMaterialEditor/Private/DynamicMaterialEditorModule.cpp
+++ b/Engine/Plugins/Experimental/DynamicMaterial/Source/DynamicMaterialEditor/Private/DynamicMaterialEditorModule.cpp
@@ -352,12 +352,12 @@ void FDynamicMaterialEditorModule::OpenMaterialObjectProperty(const FDMObjectMat
 	}
 }
 
-void FDynamicMaterialEditorModule::OpenMaterialInstance(UDynamicMaterialInstance* InInstance, UWorld* InWorld, 
+void FDynamicMaterialEditorModule::OpenMaterial(UDynamicMaterialInstance* InMaterial, UWorld* InWorld, 
 	bool bInInvokeTab) const
 {
-	if (IsValid(InInstance))
+	if (IsValid(InMaterial))
 	{
-		if (UDynamicMaterialModel* InstanceModel = InInstance->GetMaterialModel())
+		if (UDynamicMaterialModel* InstanceModel = InMaterial->GetMaterialModel())
 		{
 			OpenMaterialModel(InstanceModel, InWorld, bInInvokeTab);
 		}
diff --git a/Engine/Plugins/Experimental/DynamicMaterial/Source/DynamicMaterialEditor/Private/DynamicMaterialEditorModule.h b/Engine/Plugins/Experimental/DynamicMaterial/Source/DynamicMaterialEditor/Private/DynamicMaterialEditorModule.h
--- a/Engine/Plugins/Experimental/DynamicMaterial/Source/DynamicMaterialEditor/Private/DynamicMaterialEditorModule.h
+++ b/Engine/Plugins/Experimental/DynamicMaterial/Source/DynamicMaterialEditor/Private/DynamicMaterialEditorModule.h
@@ -89,7 +89,7 @@ public:
 	virtual UDynamicMaterialModelBase* GetOpenedMaterialModel(UWorld* InWorld) const override;
 	virtual void OpenMaterialModel(UDynamicMaterialModelBase* InMaterialModel, UWorld* InWorld, bool bInInvokeTab) const override;
 	virtual void OpenMaterialObjectProperty(const FDMObjectMaterialProperty& InObjectProperty, UWorld* InWorld, bool bInInvokeTab) const override;
-	virtual void OpenMaterialInstance(UDynamicMaterialInstance* InInstance, UWorld* InWorld, bool bInInvokeTab) const override;
+	virtual void OpenMaterial(UDynamicMaterialInstance* InMaterial, UWorld* InWorld, bool bInInvokeTab) const override;
 	virtual void OnActorSelected(AActor* InActor, UWorld* InWorld, bool bInInvokeTab) const override;
 	virtual void ClearDynamicMaterialModel(UWorld* InWorld) const override;
 	//~ End IDynamicMaterialEditorModule
diff --git a/Engine/Plugins/Experimental/DynamicMaterial/Source/DynamicMaterialEditor/Private/Material/AssetDefinition_DynamicMaterialInstance.cpp b/Engine/Plugins/Experimental/DynamicMaterial/Source/DynamicMaterialEditor/Private/Material/AssetDefinition_DynamicMaterialInstance.cpp
--- a/Engine/Plugins/Experimental/DynamicMaterial/Source/DynamicMaterialEditor/Private/Material/AssetDefinition_DynamicMaterialInstance.cpp
+++ b/Engine/Plugins/Experimental/DynamicMaterial/Source/DynamicMaterialEditor/Private/Material/AssetDefinition_DynamicMaterialInstance.cpp
@@ -12,12 +12,12 @@
 
 FText UAssetDefinition_DynamicMaterialInstance::GetAssetDisplayName() const
 {
-	return LOCTEXT("DynamicMaterialInstance", "Dynamic Material Instance");
+	return LOCTEXT("MaterialDesigner", "Material Designer Material");
 }
 
 FText UAssetDefinition_DynamicMaterialInstance::GetAssetDisplayName(const FAssetData& InAssetData) const
 {
-	return LOCTEXT("MaterialDesignerInstance", "Material Designer Instance");
+	return GetAssetDisplayName();
 }
 
 TSoftClassPtr<> UAssetDefinition_DynamicMaterialInstance::GetAssetClass() const
diff --git a/Engine/Plugins/Experimental/DynamicMaterial/Source/DynamicMaterialEditor/Private/Material/DynamicMaterialInstanceFactory.cpp b/Engine/Plugins/Experimental/DynamicMaterial/Source/DynamicMaterialEditor/Private/Material/DynamicMaterialInstanceFactory.cpp
--- a/Engine/Plugins/Experimental/DynamicMaterial/Source/DynamicMaterialEditor/Private/Material/DynamicMaterialInstanceFactory.cpp
+++ b/Engine/Plugins/Experimental/DynamicMaterial/Source/DynamicMaterialEditor/Private/Material/DynamicMaterialInstanceFactory.cpp
@@ -88,14 +88,12 @@ UObject* UDynamicMaterialInstanceFactory::FactoryCreateNew(UClass* Class, UObjec
 
 FText UDynamicMaterialInstanceFactory::GetDisplayName() const
 {
-	static const FText DisplayName = LOCTEXT("MaterialDesignerInstance", "Material Designer Instance");
-	return DisplayName;
+	return LOCTEXT("MaterialDesignerInstance", "Material Designer Material");
 }
 
 FText UDynamicMaterialInstanceFactory::GetToolTip() const
 {
-	static const FText Tooltip = LOCTEXT("MaterialDesignerInstanceTooltip", "The Material Designer Instance is a combination of a Material Instance Dyanmic and a Material Designer Model.");
-	return Tooltip;
+	return LOCTEXT("MaterialDesignerInstanceTooltip", "The Material Designer Material is a combination of a Material Instance Dyanmic and a Material Designer Model.");
 }
 
 #undef LOCTEXT_NAMESPACE
diff --git a/Engine/Plugins/Experimental/DynamicMaterial/Source/DynamicMaterialEditor/Private/Model/AssetDefinition_DynamicMaterialModelDynamic.cpp b/Engine/Plugins/Experimental/DynamicMaterial/Source/DynamicMaterialEditor/Private/Model/AssetDefinition_DynamicMaterialModelDynamic.cpp
--- a/Engine/Plugins/Experimental/DynamicMaterial/Source/DynamicMaterialEditor/Private/Model/AssetDefinition_DynamicMaterialModelDynamic.cpp
+++ b/Engine/Plugins/Experimental/DynamicMaterial/Source/DynamicMaterialEditor/Private/Model/AssetDefinition_DynamicMaterialModelDynamic.cpp
@@ -10,7 +10,7 @@
 
 FText UAssetDefinition_DynamicMaterialModelDynamic::GetAssetDisplayName() const
 {
-	return LOCTEXT("DynamicMaterialModel", "Dynamic Material Model Dynamic");
+	return LOCTEXT("DynamicMaterialModel", "Dynamic Material Model Instance");
 }
 
 FText UAssetDefinition_DynamicMaterialModelDynamic::GetAssetDisplayName(const FAssetData& InAssetData) const
diff --git a/Engine/Plugins/Experimental/DynamicMaterial/Source/DynamicMaterialEditor/Private/UI/Menus/DMToolBarMenus.cpp b/Engine/Plugins/Experimental/DynamicMaterial/Source/DynamicMaterialEditor/Private/UI/Menus/DMToolBarMenus.cpp
--- a/Engine/Plugins/Experimental/DynamicMaterial/Source/DynamicMaterialEditor/Private/UI/Menus/DMToolBarMenus.cpp
+++ b/Engine/Plugins/Experimental/DynamicMaterial/Source/DynamicMaterialEditor/Private/UI/Menus/DMToolBarMenus.cpp
@@ -95,17 +95,17 @@ void FDMToolBarMenus::AddExportMenu(UToolMenu* InMenu)
 		return;
 	}
 
-	UDynamicMaterialInstance* Instance = MaterialModelBase->GetDynamicMaterialInstance();
+	UDynamicMaterialInstance* Material = MaterialModelBase->GetDynamicMaterialInstance();
 
-	if (!Instance)
+	if (!Material)
 	{
 		return;
 	}
 
-	const bool bAllowInstanceExport = IsValid(Instance->GetOuter()) && !Instance->GetOuter()->IsA<UPackage>();
-	const bool bAllowMaterialExport = IsValid(MaterialModelBase->GetGeneratedMaterial());
+	const bool bAllowMaterialExport = IsValid(Material->GetOuter()) && !Material->GetOuter()->IsA<UPackage>();
+	const bool bAllowGeneratedMaterialExport = IsValid(MaterialModelBase->GetGeneratedMaterial());
 
-	if (!bAllowInstanceExport && !bAllowMaterialExport)
+	if (!bAllowMaterialExport && !bAllowGeneratedMaterialExport)
 	{
 		return;
 	}
@@ -114,8 +114,8 @@ void FDMToolBarMenus::AddExportMenu(UToolMenu* InMenu)
 
 	NewSection.AddMenuEntry(
 		NAME_None,
-		LOCTEXT("OpenInUEMaterialEditor", "Open in Standard Material Editor"),
-		LOCTEXT("OpenInUEMaterialEditorTooltip", "Opens the currently editing generated Material Designer Instance material in the standard material editor."),
+		LOCTEXT("OpenInUEMaterialEditor", "Open Generated Material"),
+		LOCTEXT("OpenInUEMaterialEditorTooltip", "Opens the Generated Material in the standard Material Editor."),
 		FSlateIcon(),
 		FUIAction(FExecuteAction::CreateStatic(
 			&OpenMaterialEditorFromContext,
@@ -123,16 +123,16 @@ void FDMToolBarMenus::AddExportMenu(UToolMenu* InMenu)
 		))
 	);
 
-	if (bAllowInstanceExport)
+	if (bAllowGeneratedMaterialExport)
 	{
 		NewSection.AddMenuEntry(
 			NAME_None,
-			LOCTEXT("ExportMaterialInstance", "Export Material Designer Instance"),
-			LOCTEXT("ExportMaterialInstanceTooltip", "Export the material instance to an asset."),
+			LOCTEXT("ExportGeneratedMaterial", "Export Generated Material"),
+			LOCTEXT("ExportGeneratedMaterialTooltip", "Export the Generated Material as an asset."),
 			FSlateIcon(),
 			FUIAction(FExecuteAction::CreateStatic(
-				&FDMToolBarMenus::ExportMaterialInstanceFromInstance,
-				TWeakObjectPtr<UDynamicMaterialInstance>(Instance)
+				&FDMToolBarMenus::ExportMaterialModel,
+				TWeakObjectPtr<UDynamicMaterialModelBase>(MaterialModelBase)
 			))
 		);
 	}
@@ -141,20 +141,20 @@ void FDMToolBarMenus::AddExportMenu(UToolMenu* InMenu)
 	{
 		NewSection.AddMenuEntry(
 			NAME_None,
-			LOCTEXT("ExportGeneratedMaterial", "Export Generated Material"),
-			LOCTEXT("ExportGeneratedMaterialTooltip", "Export the generated material to an asset."),
+			LOCTEXT("ExportMaterial", "Export Material Designer Material"),
+			LOCTEXT("ExportMaterialInstanceTooltip", "Export the Material Designer Material as an asset."),
 			FSlateIcon(),
 			FUIAction(FExecuteAction::CreateStatic(
-				&FDMToolBarMenus::ExportMaterialModelFromModel,
-				TWeakObjectPtr<UDynamicMaterialModelBase>(MaterialModelBase)
+				&FDMToolBarMenus::ExportMaterial,
+				TWeakObjectPtr<UDynamicMaterialInstance>(Material)
 			))
 		);
 	}
 
 	NewSection.AddSubMenu(
 		NAME_None,
-		LOCTEXT("SnapshotMaterial", "Snapshop Material"),
-		LOCTEXT("SnapshotMaterialTooltip", "Take a snapshot of the material with the current settings and export it as a texture."),
+		LOCTEXT("SnapshotMaterial", "Snapshop Material Designer Material"),
+		LOCTEXT("SnapshotMaterialTooltip", "Take a snapshot of the Material Designer Material with the current values and export it as a Texture asset."),
 		FNewToolMenuChoice(FNewToolMenuDelegate::CreateStatic(
 			&FDMToolBarMenus::CreateSnapshotMaterialMenu
 		))
@@ -291,13 +291,13 @@ void FDMToolBarMenus::OpenMaterialEditorFromContext(UDMMenuContext* InMenuContex
 	}
 }
 
-void FDMToolBarMenus::ExportMaterialInstanceFromInstance(TWeakObjectPtr<UDynamicMaterialInstance> InMaterialInstanceWeak)
+void FDMToolBarMenus::ExportMaterial(TWeakObjectPtr<UDynamicMaterialInstance> InMaterialInstanceWeak)
 {
-	if (UDynamicMaterialInstance* MaterialInstance = InMaterialInstanceWeak.Get())
+	if (UDynamicMaterialInstance* Material = InMaterialInstanceWeak.Get())
 	{
 		IAssetTools& AssetTools = FModuleManager::LoadModuleChecked<FAssetToolsModule>("AssetTools").Get();
 		FString PackageName, AssetName;
-		AssetTools.CreateUniqueAssetName(MaterialInstance->GetName(), TEXT(""), PackageName, AssetName);
+		AssetTools.CreateUniqueAssetName(Material->GetName(), TEXT(""), PackageName, AssetName);
 
 		IContentBrowserSingleton& ContentBrowser = FModuleManager::LoadModuleChecked<FContentBrowserModule>("ContentBrowser").Get();
 		const FContentBrowserItemPath CurrentPath = ContentBrowser.GetCurrentPath();
@@ -314,17 +314,17 @@ void FDMToolBarMenus::ExportMaterialInstanceFromInstance(TWeakObjectPtr<UDynamic
 
 		if (!SaveObjectPath.IsEmpty())
 		{
-			UDMMaterialModelFunctionLibrary::ExportMaterialInstance(MaterialInstance->GetMaterialModelBase(), SaveObjectPath);
+			UDMMaterialModelFunctionLibrary::ExportMaterial(Material->GetMaterialModelBase(), SaveObjectPath);
 
 			if (FEngineAnalytics::IsAvailable())
 			{
-				FEngineAnalytics::GetProvider().RecordEvent(TEXT("Editor.Usage.MaterialDesigner"), TEXT("Action"), TEXT("ExportedMaterialInstance"));
+				FEngineAnalytics::GetProvider().RecordEvent(TEXT("Editor.Usage.MaterialDesigner"), TEXT("Action"), TEXT("ExportedMaterial"));
 			}
 		}
 	}
 }
 
-void FDMToolBarMenus::ExportMaterialModelFromModel(TWeakObjectPtr<UDynamicMaterialModelBase> InMaterialModelBaseWeak)
+void FDMToolBarMenus::ExportMaterialModel(TWeakObjectPtr<UDynamicMaterialModelBase> InMaterialModelBaseWeak)
 {
 	UDynamicMaterialModelBase* MaterialModelBase = InMaterialModelBaseWeak.Get();
 
@@ -382,7 +382,7 @@ void FDMToolBarMenus::SnapshotMaterial(TWeakObjectPtr<UDynamicMaterialModelBase>
 	{
 		if (!IsValid(MaterialInstance->Parent.Get()))
 		{
-			UE_LOG(LogDynamicMaterialEditor, Warning, TEXT("Unable to find world to find material instance parent."));
+			UE_LOG(LogDynamicMaterialEditor, Warning, TEXT("Unable to find world to find material parent."));
 			return;
 		}
 
@@ -502,53 +502,31 @@ void FDMToolBarMenus::CreateSnapshotMaterialMenu(UToolMenu* InMenu)
 	TWeakObjectPtr<UDynamicMaterialModelBase> MaterialModelWeak = MaterialModelBase;
 	FToolMenuSection& NewSection = InMenu->AddSection("SnapshotMaterial", LOCTEXT("SnapshotMaterial", "Snapshop Material"));
 
-	NewSection.AddMenuEntry(
-		NAME_None,
-		LOCTEXT("SnapshotMaterial512", "512x512"),
-		LOCTEXT("SnapshotMaterial512Tooltip", "Take a snapshot of the material with the current settings and export it as a texture with a resolution of 512x512 pixels."),
-		FSlateIcon(),
-		FUIAction(FExecuteAction::CreateStatic(
-			&FDMToolBarMenus::SnapshotMaterial,
-			MaterialModelWeak,
-			FIntPoint(512, 512)
-		))
-	);
+	const FText SnapshotNameFormat = LOCTEXT("SnapshotName", "{0}x{0}");
 
-	NewSection.AddMenuEntry(
-		NAME_None,
-		LOCTEXT("SnapshotMaterial1024", "1024x1024"),
-		LOCTEXT("SnapshotMaterial1024Tooltip", "Take a snapshot of the material with the current settings and export it as a texture with a resolution of 1024x1024 pixels."),
-		FSlateIcon(),
-		FUIAction(FExecuteAction::CreateStatic(
-			&FDMToolBarMenus::SnapshotMaterial,
-			MaterialModelWeak,
-			FIntPoint(1024, 1024)
-		))
+	const FText SnapshotTooltipFormat = LOCTEXT(
+		"SnapshotMaterialTooltip",
+		"Take a snapshot of the Material Designer Material with the current values and export it as a Texture asset with a resolution of {0} pixels."
 	);
 
-	NewSection.AddMenuEntry(
-		NAME_None,
-		LOCTEXT("SnapshotMaterial2048", "2048x2048"),
-		LOCTEXT("SnapshotMaterial2048Tooltip", "Take a snapshot of the material with the current settings and export it as a texture with a resolution of 2048x2048 pixels."),
-		FSlateIcon(),
-		FUIAction(FExecuteAction::CreateStatic(
-			&FDMToolBarMenus::SnapshotMaterial,
-			MaterialModelWeak,
-			FIntPoint(2048, 2048)
-		))
-	);
+	const TArray<int32> SnapshotTypeResolutions = {512, 1024, 2048, 4096};
 
-	NewSection.AddMenuEntry(
-		NAME_None,
-		LOCTEXT("SnapshotMaterial4096", "4096x4096"),
-		LOCTEXT("SnapshotMaterial4096Tooltip", "Take a snapshot of the material with the current settings and export it as a texture with a resolution of 4096x4096 pixels."),
-		FSlateIcon(),
-		FUIAction(FExecuteAction::CreateStatic(
-			&FDMToolBarMenus::SnapshotMaterial,
-			MaterialModelWeak,
-			FIntPoint(4096, 4096)
-		))
-	);
+	for (int32 SnapshotTypeResolution : SnapshotTypeResolutions)
+	{
+		const FText Name = FText::Format(SnapshotNameFormat, FText::AsNumber(SnapshotTypeResolution));
+
+		NewSection.AddMenuEntry(
+			NAME_None,
+			Name,
+			FText::Format(SnapshotTooltipFormat, Name),
+			FSlateIcon(),
+			FUIAction(FExecuteAction::CreateStatic(
+				&FDMToolBarMenus::SnapshotMaterial,
+				MaterialModelWeak,
+				FIntPoint(SnapshotTypeResolution, SnapshotTypeResolution)
+			))
+		);
+	}
 }
 
 #undef LOCTEXT_NAMESPACE
diff --git a/Engine/Plugins/Experimental/DynamicMaterial/Source/DynamicMaterialEditor/Private/UI/Menus/DMToolBarMenus.h b/Engine/Plugins/Experimental/DynamicMaterial/Source/DynamicMaterialEditor/Private/UI/Menus/DMToolBarMenus.h
--- a/Engine/Plugins/Experimental/DynamicMaterial/Source/DynamicMaterialEditor/Private/UI/Menus/DMToolBarMenus.h
+++ b/Engine/Plugins/Experimental/DynamicMaterial/Source/DynamicMaterialEditor/Private/UI/Menus/DMToolBarMenus.h
@@ -42,9 +42,9 @@ private:
 
 	static void OpenMaterialEditorFromContext(UDMMenuContext* InMenuContext);
 
-	static void ExportMaterialInstanceFromInstance(TWeakObjectPtr<UDynamicMaterialInstance> InMaterialInstanceWeak);
+	static void ExportMaterial(TWeakObjectPtr<UDynamicMaterialInstance> InMaterialInstanceWeak);
 
-	static void ExportMaterialModelFromModel(TWeakObjectPtr<UDynamicMaterialModelBase> InMaterialModelBaseWeak);
+	static void ExportMaterialModel(TWeakObjectPtr<UDynamicMaterialModelBase> InMaterialModelBaseWeak);
 
 	static void SnapshotMaterial(TWeakObjectPtr<UDynamicMaterialModelBase> InMaterialModelBaseWeak, FIntPoint InTextureSize);
 
diff --git a/Engine/Plugins/Experimental/DynamicMaterial/Source/DynamicMaterialEditor/Private/UI/Utils/DynamicMaterialInstanceThumbnailRenderer.h b/Engine/Plugins/Experimental/DynamicMaterial/Source/DynamicMaterialEditor/Private/UI/Utils/DynamicMaterialInstanceThumbnailRenderer.h
--- a/Engine/Plugins/Experimental/DynamicMaterial/Source/DynamicMaterialEditor/Private/UI/Utils/DynamicMaterialInstanceThumbnailRenderer.h
+++ b/Engine/Plugins/Experimental/DynamicMaterial/Source/DynamicMaterialEditor/Private/UI/Utils/DynamicMaterialInstanceThumbnailRenderer.h
@@ -11,7 +11,7 @@ class FDynamicMaterialInstanceThumbnailScene;
 class FRenderTarget;
 
 /**
- * This thumbnail renderer displays a given Material Designer Instance.
+ * This thumbnail renderer displays a given Material Designer Material.
  */
 UCLASS(config = Editor)
 class UDynamicMaterialInstanceThumbnailRenderer : public UDefaultSizedThumbnailRenderer
diff --git a/Engine/Plugins/Experimental/DynamicMaterial/Source/DynamicMaterialEditor/Private/UI/Widgets/Editor/SDMToolBar.cpp b/Engine/Plugins/Experimental/DynamicMaterial/Source/DynamicMaterialEditor/Private/UI/Widgets/Editor/SDMToolBar.cpp
--- a/Engine/Plugins/Experimental/DynamicMaterial/Source/DynamicMaterialEditor/Private/UI/Widgets/Editor/SDMToolBar.cpp
+++ b/Engine/Plugins/Experimental/DynamicMaterial/Source/DynamicMaterialEditor/Private/UI/Widgets/Editor/SDMToolBar.cpp
@@ -117,7 +117,7 @@ TSharedRef<SWidget> SDMToolBar::CreateToolBarEntries()
 					SNew(SButton)
 					.ContentPadding(LargeIconToolBarButtonContentPadding)
 					.ButtonStyle(FDynamicMaterialEditorStyle::Get(), "HoverHintOnly")
-					.ToolTipText(LOCTEXT("ExportMaterialInstance", "Save As"))
+					.ToolTipText(LOCTEXT("ExportMaterial", "Save As"))
 					.OnClicked(this, &SDMToolBar::OnExportMaterialInstanceButtonClicked)
 					[
 						SNew(SImage)
@@ -134,7 +134,7 @@ TSharedRef<SWidget> SDMToolBar::CreateToolBarEntries()
 					.Visibility(EVisibility::Collapsed)
 					.ContentPadding(LargeIconToolBarButtonContentPadding)
 					.ButtonStyle(FDynamicMaterialEditorStyle::Get(), "HoverHintOnly")
-					.ToolTipText(LOCTEXT("MaterialDesignerOpenParentTooltip", "Open the parent of this Material Designer Dynamic."))
+					.ToolTipText(LOCTEXT("MaterialDesignerOpenParentTooltip", "Open the parent of this Material Designer Instance."))
 					.OnClicked(this, &SDMToolBar::OnOpenParentClicked)
 					[
 						SNew(SImage)
@@ -151,7 +151,7 @@ TSharedRef<SWidget> SDMToolBar::CreateToolBarEntries()
 					.Visibility(EVisibility::Collapsed)
 					.ContentPadding(LargeIconToolBarButtonContentPadding)
 					.ButtonStyle(FDynamicMaterialEditorStyle::Get(), "HoverHintOnly")
-					.ToolTipText(LOCTEXT("MaterialDesignerConvertToEditableTooltip", "Convert this Material Designer Dyanmic to a fully editable material (and create a new shader)."))
+					.ToolTipText(LOCTEXT("MaterialDesignerConvertToEditableTooltip", "Convert this Material Designer Instance to a fully editable Material (and create a new shader)."))
 					.OnClicked(this, &SDMToolBar::OnConvertToEditableClicked)
 					[
 						SNew(SImage)
@@ -487,7 +487,7 @@ void SDMToolBar::OnMaterialSlotChanged(TSharedPtr<FDMObjectMaterialProperty> InS
 	}
 	else if (InSelectedSlot->GetOuter())
 	{
-		if (UDynamicMaterialModel* NewModel = UDMMaterialModelFunctionLibrary::CreateDynamicMaterialInObject(*InSelectedSlot.Get()))
+		if (UDynamicMaterialModel* NewModel = UDMMaterialModelFunctionLibrary::CreateMaterialInObject(*InSelectedSlot.Get()))
 		{
 			DesignerWidget->OpenObjectMaterialProperty(*InSelectedSlot);
 		}
@@ -565,7 +565,7 @@ FReply SDMToolBar::OnExportMaterialInstanceButtonClicked()
 		return FReply::Handled();
 	}
 
-	UDynamicMaterialInstance* NewInstance = UDMMaterialModelFunctionLibrary::ExportMaterialInstance(MaterialModelBase);
+	UDynamicMaterialInstance* NewInstance = UDMMaterialModelFunctionLibrary::ExportMaterial(MaterialModelBase);
 
 	if (!NewInstance)
 	{
@@ -894,7 +894,7 @@ FReply SDMToolBar::OnConvertToEditableClicked()
 
 	if (OldInstance)
 	{
-		NewInstance = UDMMaterialModelFunctionLibrary::ExportToTemplateMaterialInstance(CurrentModelDynamic);
+		NewInstance = UDMMaterialModelFunctionLibrary::ExportToTemplateMaterial(CurrentModelDynamic);
 
 		if (NewInstance)
 		{
diff --git a/Engine/Plugins/Experimental/DynamicMaterial/Source/DynamicMaterialEditor/Private/UI/Widgets/SDMActorMaterialSelector.cpp b/Engine/Plugins/Experimental/DynamicMaterial/Source/DynamicMaterialEditor/Private/UI/Widgets/SDMActorMaterialSelector.cpp
--- a/Engine/Plugins/Experimental/DynamicMaterial/Source/DynamicMaterialEditor/Private/UI/Widgets/SDMActorMaterialSelector.cpp
+++ b/Engine/Plugins/Experimental/DynamicMaterial/Source/DynamicMaterialEditor/Private/UI/Widgets/SDMActorMaterialSelector.cpp
@@ -208,7 +208,7 @@ FReply SDMActorMaterialSelector::OnCreateMaterialButtonClicked(int32 InActorProp
 		return FReply::Handled();
 	}
 
-	UDynamicMaterialModel* NewMaterialModel = UDMMaterialModelFunctionLibrary::CreateDynamicMaterialInObject(ActorProperties[InActorPropertyIndex]);
+	UDynamicMaterialModel* NewMaterialModel = UDMMaterialModelFunctionLibrary::CreateMaterialInObject(ActorProperties[InActorPropertyIndex]);
 
 	if (NewMaterialModel)
 	{
diff --git a/Engine/Plugins/Experimental/DynamicMaterial/Source/DynamicMaterialEditor/Private/UI/Widgets/SDMMaterialSelectPrompt.cpp b/Engine/Plugins/Experimental/DynamicMaterial/Source/DynamicMaterialEditor/Private/UI/Widgets/SDMMaterialSelectPrompt.cpp
--- a/Engine/Plugins/Experimental/DynamicMaterial/Source/DynamicMaterialEditor/Private/UI/Widgets/SDMMaterialSelectPrompt.cpp
+++ b/Engine/Plugins/Experimental/DynamicMaterial/Source/DynamicMaterialEditor/Private/UI/Widgets/SDMMaterialSelectPrompt.cpp
@@ -25,7 +25,7 @@ void SDMMaterialSelectPrompt::Construct(const FArguments& InArgs)
 		[
 			SNew(STextBlock)
 			.TextStyle(FDynamicMaterialEditorStyle::Get(), "RegularFont")
-			.Text(LOCTEXT("NoActiveMaterial", "No active Material Designer Instance."))
+			.Text(LOCTEXT("NoActiveMaterial", "No active Material Designer Material."))
 		]
 	];
 }
diff --git a/Engine/Plugins/Experimental/DynamicMaterial/Source/DynamicMaterialEditor/Private/UI/Widgets/SDMMaterialWizard.cpp b/Engine/Plugins/Experimental/DynamicMaterial/Source/DynamicMaterialEditor/Private/UI/Widgets/SDMMaterialWizard.cpp
--- a/Engine/Plugins/Experimental/DynamicMaterial/Source/DynamicMaterialEditor/Private/UI/Widgets/SDMMaterialWizard.cpp
+++ b/Engine/Plugins/Experimental/DynamicMaterial/Source/DynamicMaterialEditor/Private/UI/Widgets/SDMMaterialWizard.cpp
@@ -216,7 +216,7 @@ TSharedRef<SWidget> SDMMaterialWizard::CreateModeSelector()
 						.Padding(FMargin(10.f, 6.f))
 						.IsChecked(this, &SDMMaterialWizard::IsModeSelected, EDMMaterialWizardMode::Instance)
 						.OnCheckStateChanged(this, &SDMMaterialWizard::SetMode, EDMMaterialWizardMode::Instance)
-						.ToolTipText(LOCTEXT("TemplateModeToolTip", "Create a new Dynamic Material Instance based on a template."))
+						.ToolTipText(LOCTEXT("TemplateModeToolTip", "Create a new Material Designer Instance based on a template."))
 						.Content()
 						[
 							SNew(STextBlock)
@@ -779,7 +779,7 @@ void SDMMaterialWizard::CreateDynamicMaterialInInstance(UDynamicMaterialModel* I
 		return;
 	}
 
-	if (!UDMMaterialModelFunctionLibrary::CreateDynamicModelInInstance(InTemplateModel, InToInstance))
+	if (!UDMMaterialModelFunctionLibrary::CreateModelInstanceInMaterial(InTemplateModel, InToInstance))
 	{
 		UE::DynamicMaterialEditor::Private::LogError(TEXT("Failed to create new dynamic model in existing instance."));
 		return;
@@ -821,7 +821,7 @@ void SDMMaterialWizard::CreateNewDynamicInstanceInActor(UDynamicMaterialModel* I
 		return;
 	}
 
-	if (!UDMMaterialModelFunctionLibrary::CreateDynamicModelInInstance(InFromModel, NewInstance))
+	if (!UDMMaterialModelFunctionLibrary::CreateModelInstanceInMaterial(InFromModel, NewInstance))
 	{
 		UE::DynamicMaterialEditor::Private::LogError(TEXT("Failed to create new dynamic model in new instance."));
 		return;
diff --git a/Engine/Plugins/Experimental/DynamicMaterial/Source/DynamicMaterialEditor/Private/Utils/DMMaterialModelFunctionLibrary.cpp b/Engine/Plugins/Experimental/DynamicMaterial/Source/DynamicMaterialEditor/Private/Utils/DMMaterialModelFunctionLibrary.cpp
--- a/Engine/Plugins/Experimental/DynamicMaterial/Source/DynamicMaterialEditor/Private/Utils/DMMaterialModelFunctionLibrary.cpp
+++ b/Engine/Plugins/Experimental/DynamicMaterial/Source/DynamicMaterialEditor/Private/Utils/DMMaterialModelFunctionLibrary.cpp
@@ -61,7 +61,7 @@ TArray<FDMObjectMaterialProperty> UDMMaterialModelFunctionLibrary::GetActorMater
 	return ActorProperties;
 }
 
-UDynamicMaterialModel* UDMMaterialModelFunctionLibrary::CreateDynamicMaterialInObject(FDMObjectMaterialProperty& InMaterialProperty)
+UDynamicMaterialModel* UDMMaterialModelFunctionLibrary::CreateMaterialInObject(FDMObjectMaterialProperty& InMaterialProperty)
 {
 	if (!InMaterialProperty.IsValid())
 	{
@@ -108,7 +108,7 @@ UDynamicMaterialModel* UDMMaterialModelFunctionLibrary::CreateDynamicMaterialInO
 	return nullptr;
 }
 
-UDynamicMaterialInstance* UDMMaterialModelFunctionLibrary::ExportMaterialInstance(UDynamicMaterialModelBase* InMaterialModelBase)
+UDynamicMaterialInstance* UDMMaterialModelFunctionLibrary::ExportMaterial(UDynamicMaterialModelBase* InMaterialModelBase)
 {
 	if (!IsValid(InMaterialModelBase))
 	{
@@ -158,10 +158,10 @@ UDynamicMaterialInstance* UDMMaterialModelFunctionLibrary::ExportMaterialInstanc
 		return nullptr;
 	}
 
-	return ExportMaterialInstance(MaterialInstance->GetMaterialModelBase(), SaveObjectPath);
+	return ExportMaterial(MaterialInstance->GetMaterialModelBase(), SaveObjectPath);
 }
 
-UDynamicMaterialInstance* UDMMaterialModelFunctionLibrary::ExportMaterialInstance(UDynamicMaterialModelBase* InMaterialModel, const FString& InSavePath)
+UDynamicMaterialInstance* UDMMaterialModelFunctionLibrary::ExportMaterial(UDynamicMaterialModelBase* InMaterialModel, const FString& InSavePath)
 {
 	if (!IsValid(InMaterialModel))
 	{
@@ -178,7 +178,7 @@ UDynamicMaterialInstance* UDMMaterialModelFunctionLibrary::ExportMaterialInstanc
 	UDynamicMaterialInstance* Instance = InMaterialModel->GetDynamicMaterialInstance();
 	if (!Instance)
 	{
-		UE::DynamicMaterialEditor::Private::LogError(TEXT("Failed to find a Material Designer Instance to export."), true, InMaterialModel);
+		UE::DynamicMaterialEditor::Private::LogError(TEXT("Failed to find a Material Designer Material to export."), true, InMaterialModel);
 		return nullptr;
 	}
 
@@ -187,7 +187,7 @@ UDynamicMaterialInstance* UDMMaterialModelFunctionLibrary::ExportMaterialInstanc
 	UPackage* Package = CreatePackage(*PackagePath);
 	if (!Package)
 	{
-		UE::DynamicMaterialEditor::Private::LogError(FString::Printf(TEXT("Failed to create package for Material Designer Instance (%s)."), *PackagePath));
+		UE::DynamicMaterialEditor::Private::LogError(FString::Printf(TEXT("Failed to create package for Material Designer Material (%s)."), *PackagePath));
 		return nullptr;
 	}
 
@@ -199,7 +199,7 @@ UDynamicMaterialInstance* UDMMaterialModelFunctionLibrary::ExportMaterialInstanc
 	UObject* NewAsset = StaticDuplicateObjectEx(Params);
 	if (!NewAsset)
 	{
-		UE::DynamicMaterialEditor::Private::LogError(TEXT("Failed to create new Material Designer Instance asset."));
+		UE::DynamicMaterialEditor::Private::LogError(TEXT("Failed to create new Material Designer Material asset."));
 		return nullptr;
 	}
 
@@ -230,7 +230,7 @@ UDynamicMaterialInstance* UDMMaterialModelFunctionLibrary::ExportMaterialInstanc
 
 	if (FEngineAnalytics::IsAvailable())
 	{
-		FEngineAnalytics::GetProvider().RecordEvent(TEXT("Editor.Usage.MaterialDesigner"), TEXT("Action"), TEXT("ExportedMaterialInstance"));
+		FEngineAnalytics::GetProvider().RecordEvent(TEXT("Editor.Usage.MaterialDesigner"), TEXT("Action"), TEXT("ExportedMaterial"));
 	}
 
 	return NewInstance;
@@ -384,7 +384,7 @@ UDynamicMaterialModel* UDMMaterialModelFunctionLibrary::ExportToTemplateMaterial
 	return NewModel;
 }
 
-UDynamicMaterialInstance* UDMMaterialModelFunctionLibrary::ExportToTemplateMaterialInstance(UDynamicMaterialModelDynamic* InMaterialModelDynamic)
+UDynamicMaterialInstance* UDMMaterialModelFunctionLibrary::ExportToTemplateMaterial(UDynamicMaterialModelDynamic* InMaterialModelDynamic)
 {
 	UDynamicMaterialModel* ParentModel = InMaterialModelDynamic->GetParentModel();
 
@@ -431,10 +431,10 @@ UDynamicMaterialInstance* UDMMaterialModelFunctionLibrary::ExportToTemplateMater
 		return nullptr;
 	}
 
-	return ExportToTemplateMaterialInstance(InMaterialModelDynamic, SaveObjectPath);	
+	return ExportToTemplateMaterial(InMaterialModelDynamic, SaveObjectPath);	
 }
 
-UDynamicMaterialInstance* UDMMaterialModelFunctionLibrary::ExportToTemplateMaterialInstance(UDynamicMaterialModelDynamic* InMaterialModelDynamic, const FString& InSavePath)
+UDynamicMaterialInstance* UDMMaterialModelFunctionLibrary::ExportToTemplateMaterial(UDynamicMaterialModelDynamic* InMaterialModelDynamic, const FString& InSavePath)
 {
 	UDynamicMaterialModel* ParentModel = InMaterialModelDynamic->GetParentModel();
 
@@ -486,7 +486,7 @@ UDynamicMaterialInstance* UDMMaterialModelFunctionLibrary::ExportToTemplateMater
 
 	if (FEngineAnalytics::IsAvailable())
 	{
-		FEngineAnalytics::GetProvider().RecordEvent(TEXT("Editor.Usage.MaterialDesigner"), TEXT("Action"), TEXT("ExportToTemplateMaterialInstance"));
+		FEngineAnalytics::GetProvider().RecordEvent(TEXT("Editor.Usage.MaterialDesigner"), TEXT("Action"), TEXT("ExportToTemplateMaterial"));
 	}
 
 	return NewInstance;
@@ -534,7 +534,7 @@ bool UDMMaterialModelFunctionLibrary::IsModelValid(UDynamicMaterialModelBase* In
 	return true;
 }
 
-bool UDMMaterialModelFunctionLibrary::DuplicateModelBetweenInstances(UDynamicMaterialModel* InFromModel, UDynamicMaterialInstance* InToInstance)
+bool UDMMaterialModelFunctionLibrary::DuplicateModelBetweenMaterials(UDynamicMaterialModel* InFromModel, UDynamicMaterialInstance* InToInstance)
 {
 	if (!InFromModel || !InToInstance)
 	{
@@ -583,7 +583,7 @@ bool UDMMaterialModelFunctionLibrary::DuplicateModelBetweenInstances(UDynamicMat
 	return true;
 }
 
-bool UDMMaterialModelFunctionLibrary::CreateDynamicModelInInstance(UDynamicMaterialModel* InFromModel, UDynamicMaterialInstance* InToInstance)
+bool UDMMaterialModelFunctionLibrary::CreateModelInstanceInMaterial(UDynamicMaterialModel* InFromModel, UDynamicMaterialInstance* InToInstance)
 {
 	if (!InFromModel || !InToInstance)
 	{
@@ -607,7 +607,7 @@ bool UDMMaterialModelFunctionLibrary::CreateDynamicModelInInstance(UDynamicMater
 
 	if (!NewModelDynamic)
 	{
-		UE::DynamicMaterialEditor::Private::LogError(TEXT("Failed to make Dynamic Material."));
+		UE::DynamicMaterialEditor::Private::LogError(TEXT("Failed to make Material Designer Model Instance."));
 
 		// Put back the original model
 		CurrentModel->Rename(CurrentName.IsEmpty() ? nullptr : *CurrentName, InToInstance, UE::DynamicMaterial::RenameFlags);
diff --git a/Engine/Plugins/Experimental/DynamicMaterial/Source/DynamicMaterialEditor/Public/DMObjectMaterialProperty.h b/Engine/Plugins/Experimental/DynamicMaterial/Source/DynamicMaterialEditor/Public/DMObjectMaterialProperty.h
--- a/Engine/Plugins/Experimental/DynamicMaterial/Source/DynamicMaterialEditor/Public/DMObjectMaterialProperty.h
+++ b/Engine/Plugins/Experimental/DynamicMaterial/Source/DynamicMaterialEditor/Public/DMObjectMaterialProperty.h
@@ -16,7 +16,7 @@ class UObject;
 class UPrimitiveComponent;
 
 /**
- * Defines a material property slot that can be a Material Designer Instance.
+ * Defines a material property slot that can be a Material Designer Material.
  */
 USTRUCT()
 struct FDMObjectMaterialProperty
diff --git a/Engine/Plugins/Experimental/DynamicMaterial/Source/DynamicMaterialEditor/Public/IDynamicMaterialEditorModule.h b/Engine/Plugins/Experimental/DynamicMaterial/Source/DynamicMaterialEditor/Public/IDynamicMaterialEditorModule.h
--- a/Engine/Plugins/Experimental/DynamicMaterial/Source/DynamicMaterialEditor/Public/IDynamicMaterialEditorModule.h
+++ b/Engine/Plugins/Experimental/DynamicMaterial/Source/DynamicMaterialEditor/Public/IDynamicMaterialEditorModule.h
@@ -58,7 +58,7 @@ public:
 
 	virtual void OpenMaterialObjectProperty(const FDMObjectMaterialProperty& InObjectProperty, UWorld* InWorld, bool bInInvokeTab) const = 0;
 
-	virtual void OpenMaterialInstance(UDynamicMaterialInstance* InInstance, UWorld* InWorld, bool bInInvokeTab) const = 0;
+	virtual void OpenMaterial(UDynamicMaterialInstance* InInstance, UWorld* InWorld, bool bInInvokeTab) const = 0;
 
 	virtual void OnActorSelected(AActor* InActor, UWorld* InWorld, bool bInInvokeTab) const = 0;
 
diff --git a/Engine/Plugins/Experimental/DynamicMaterial/Source/DynamicMaterialEditor/Public/Utils/DMMaterialModelFunctionLibrary.h b/Engine/Plugins/Experimental/DynamicMaterial/Source/DynamicMaterialEditor/Public/Utils/DMMaterialModelFunctionLibrary.h
--- a/Engine/Plugins/Experimental/DynamicMaterial/Source/DynamicMaterialEditor/Public/Utils/DMMaterialModelFunctionLibrary.h
+++ b/Engine/Plugins/Experimental/DynamicMaterial/Source/DynamicMaterialEditor/Public/Utils/DMMaterialModelFunctionLibrary.h
@@ -16,7 +16,7 @@ class UMaterial;
 struct FDMObjectMaterialProperty;
 
 /**
- * Material Model / Instance Function Library
+ * Material / Model Function Library
  */
 UCLASS()
 class UDMMaterialModelFunctionLibrary : public UBlueprintFunctionLibrary
@@ -26,11 +26,11 @@ class UDMMaterialModelFunctionLibrary : public UBlueprintFunctionLibrary
 public:
 	DYNAMICMATERIALEDITOR_API static TArray<FDMObjectMaterialProperty> GetActorMaterialProperties(AActor* InActor);
 
-	DYNAMICMATERIALEDITOR_API static UDynamicMaterialModel* CreateDynamicMaterialInObject(FDMObjectMaterialProperty& InMaterialProperty);
+	DYNAMICMATERIALEDITOR_API static UDynamicMaterialModel* CreateMaterialInObject(FDMObjectMaterialProperty& InMaterialProperty);
 
-	DYNAMICMATERIALEDITOR_API static UDynamicMaterialInstance* ExportMaterialInstance(UDynamicMaterialModelBase* InMaterialModelBase);
+	DYNAMICMATERIALEDITOR_API static UDynamicMaterialInstance* ExportMaterial(UDynamicMaterialModelBase* InMaterialModelBase);
 
-	DYNAMICMATERIALEDITOR_API static UDynamicMaterialInstance* ExportMaterialInstance(UDynamicMaterialModelBase* InMaterialModel, const FString& InSavePath);
+	DYNAMICMATERIALEDITOR_API static UDynamicMaterialInstance* ExportMaterial(UDynamicMaterialModelBase* InMaterialModel, const FString& InSavePath);
 
 	DYNAMICMATERIALEDITOR_API static UMaterial* ExportGeneratedMaterial(UDynamicMaterialModelBase* InMaterialModelBase);
 
@@ -40,15 +40,15 @@ public:
 
 	DYNAMICMATERIALEDITOR_API static UDynamicMaterialModel* ExportToTemplateMaterialModel(UDynamicMaterialModelDynamic* InMaterialModelDynamic, const FString& InSavePath);
 
-	DYNAMICMATERIALEDITOR_API static UDynamicMaterialInstance* ExportToTemplateMaterialInstance(UDynamicMaterialModelDynamic* InMaterialModelDynamic);
+	DYNAMICMATERIALEDITOR_API static UDynamicMaterialInstance* ExportToTemplateMaterial(UDynamicMaterialModelDynamic* InMaterialModelDynamic);
 
-	DYNAMICMATERIALEDITOR_API static UDynamicMaterialInstance* ExportToTemplateMaterialInstance(UDynamicMaterialModelDynamic* InMaterialModelDynamic, const FString& InSavePath);
+	DYNAMICMATERIALEDITOR_API static UDynamicMaterialInstance* ExportToTemplateMaterial(UDynamicMaterialModelDynamic* InMaterialModelDynamic, const FString& InSavePath);
 
 	DYNAMICMATERIALEDITOR_API static bool IsModelValid(UDynamicMaterialModelBase* InMaterialModelBase);
 
-	DYNAMICMATERIALEDITOR_API static bool DuplicateModelBetweenInstances(UDynamicMaterialModel* InFromModel, UDynamicMaterialInstance* InToInstance);
+	DYNAMICMATERIALEDITOR_API static bool DuplicateModelBetweenMaterials(UDynamicMaterialModel* InFromModel, UDynamicMaterialInstance* InToInstance);
 
-	DYNAMICMATERIALEDITOR_API static bool CreateDynamicModelInInstance(UDynamicMaterialModel* InFromModel, UDynamicMaterialInstance* InToInstance);
+	DYNAMICMATERIALEDITOR_API static bool CreateModelInstanceInMaterial(UDynamicMaterialModel* InFromModel, UDynamicMaterialInstance* InToInstance);
 
 
 };
