Material Designer Texture Set
- Changed the background color of the texture name on the texture set creator to black and the color of the text to a slightly brighter white.

diff --git a/Engine/Plugins/Experimental/DynamicMaterial/Source/DynamicMaterialTextureSetEditor/Private/DMTextureSetStyle.cpp b/Engine/Plugins/Experimental/DynamicMaterial/Source/DynamicMaterialTextureSetEditor/Private/DMTextureSetStyle.cpp
--- a/Engine/Plugins/Experimental/DynamicMaterial/Source/DynamicMaterialTextureSetEditor/Private/DMTextureSetStyle.cpp
+++ b/Engine/Plugins/Experimental/DynamicMaterial/Source/DynamicMaterialTextureSetEditor/Private/DMTextureSetStyle.cpp
@@ -21,4 +21,8 @@ FDMTextureSetStyle::FDMTextureSetStyle()
 	Set("TextureSetConfig.Cell.Background", new FSlateRoundedBoxBrush(
 		FStyleColors::Recessed.GetSpecifiedColor(), 6.0f,
 		FStyleColors::Header.GetSpecifiedColor(), 2.0f));
+
+	Set("TextureSetConfig.TextureNameHighlight.Background", new FSlateColorBrush(FStyleColors::Black.GetSpecifiedColor()));
+
+	Set("TextureSetConfig.TextureNameHighlight.Color", FStyleColors::AccentWhite);
 }
diff --git a/Engine/Plugins/Experimental/DynamicMaterial/Source/DynamicMaterialTextureSetEditor/Private/Widgets/SDMTextureSetBuilderMaterialPropertyCell.cpp b/Engine/Plugins/Experimental/DynamicMaterial/Source/DynamicMaterialTextureSetEditor/Private/Widgets/SDMTextureSetBuilderMaterialPropertyCell.cpp
--- a/Engine/Plugins/Experimental/DynamicMaterial/Source/DynamicMaterialTextureSetEditor/Private/Widgets/SDMTextureSetBuilderMaterialPropertyCell.cpp
+++ b/Engine/Plugins/Experimental/DynamicMaterial/Source/DynamicMaterialTextureSetEditor/Private/Widgets/SDMTextureSetBuilderMaterialPropertyCell.cpp
@@ -127,8 +127,10 @@ void SDMTextureSetBuilderMaterialPropertyCell::Construct(const FArguments& InArg
 						.WrappingPolicy(ETextWrappingPolicy::AllowPerCharacterWrapping)
 						.WrapTextAt(110.f)
 						.Font(FAppStyle::GetFontStyle("TinyText"))
-						.HighlightText(this, &SDMTextureSetBuilderMaterialPropertyCell::GetTextureName)
 						.Visibility(this, &SDMTextureSetBuilderMaterialPropertyCell::GetTextureNameVisibility)
+						.HighlightText(this, &SDMTextureSetBuilderMaterialPropertyCell::GetTextureName)
+						.HighlightColor(FDMTextureSetStyle::Get().GetColor(TEXT("TextureSetConfig.TextureNameHighlight.Color")))
+						.HighlightShape(FDMTextureSetStyle::Get().GetBrush(TEXT("TextureSetConfig.TextureNameHighlight.Background")))
 					]
 				]
 			]
diff --git a/Engine/Plugins/Experimental/DynamicMaterial/Source/DynamicMaterialTextureSetEditor/Private/Widgets/SDMTextureSetBuilderUnassignedTextureCell.cpp b/Engine/Plugins/Experimental/DynamicMaterial/Source/DynamicMaterialTextureSetEditor/Private/Widgets/SDMTextureSetBuilderUnassignedTextureCell.cpp
--- a/Engine/Plugins/Experimental/DynamicMaterial/Source/DynamicMaterialTextureSetEditor/Private/Widgets/SDMTextureSetBuilderUnassignedTextureCell.cpp
+++ b/Engine/Plugins/Experimental/DynamicMaterial/Source/DynamicMaterialTextureSetEditor/Private/Widgets/SDMTextureSetBuilderUnassignedTextureCell.cpp
@@ -3,6 +3,7 @@
 #include "Widgets/SDMTextureSetBuilderUnassignedTextureCell.h"
 
 #include "AssetRegistry/AssetData.h"
+#include "DMTextureSetStyle.h"
 #include "Engine/Texture.h"
 #include "SDMTextureSetBuilder.h"
 #include "Widgets/Images/SImage.h"
@@ -69,7 +70,11 @@ void SDMTextureSetBuilderUnassignedTextureCell::Construct(const FArguments& InAr
 				SNew(STextBlock)
 				.Text(this, &SDMTextureSetBuilderUnassignedTextureCell::GetTextureName)
 				.WrappingPolicy(ETextWrappingPolicy::AllowPerCharacterWrapping)
-				.WrapTextAt(90.f)
+				.WrapTextAt(110.f)
+				.Font(FAppStyle::GetFontStyle("TinyText"))
+				.HighlightText(this, &SDMTextureSetBuilderUnassignedTextureCell::GetTextureName)
+				.HighlightColor(FDMTextureSetStyle::Get().GetColor(TEXT("TextureSetConfig.TextureNameHighlight.Color")))
+				.HighlightShape(FDMTextureSetStyle::Get().GetBrush(TEXT("TextureSetConfig.TextureNameHighlight.Background")))
 			]
 		]
 	];
