Material Designer
- Changed asset name in content browser to be "MD Material" or "MD Instance".
- Prefixes for exporter assets are now T_ M_ MD_ MDI_ MDM_ MDMI_
- Changed an analytics action name.

diff --git a/Engine/Plugins/Experimental/DynamicMaterial/Source/DynamicMaterial/Private/Material/DynamicMaterialInstance.cpp b/Engine/Plugins/Experimental/DynamicMaterial/Source/DynamicMaterial/Private/Material/DynamicMaterialInstance.cpp
--- a/Engine/Plugins/Experimental/DynamicMaterial/Source/DynamicMaterial/Private/Material/DynamicMaterialInstance.cpp
+++ b/Engine/Plugins/Experimental/DynamicMaterial/Source/DynamicMaterial/Private/Material/DynamicMaterialInstance.cpp
@@ -5,11 +5,43 @@
 #include "Components/MaterialValues/DMMaterialValueTexture.h"
 #include "Materials/Material.h"
 #include "Model/DynamicMaterialModel.h"
+#include "Model/DynamicMaterialModelDynamic.h"
+#include "UObject/AssetRegistryTagsContext.h"
 
 #if WITH_EDITOR
+#include "AssetRegistry/AssetData.h"
 #include "Model/IDynamicMaterialModelEditorOnlyDataInterface.h"
 #endif
 
+#define LOCTEXT_NAMESPACE "DynamicMaterialInstance"
+
+namespace UE::DynamicMaterial::Private
+{
+	const FLazyName ModelType = TEXT("ModelType");
+}
+
+const FString UDynamicMaterialInstance::ModelTypeTag_Material = TEXT("Material");
+const FString UDynamicMaterialInstance::ModelTypeTag_Instance = TEXT("Instance");
+
+#if WITH_EDITOR
+FString UDynamicMaterialInstance::GetMaterialTypeTag(const FAssetData& InAssetData)
+{
+	using namespace UE::DynamicMaterial::Private;
+
+	if (InAssetData.GetClass(EResolveClass::Yes) != UDynamicMaterialInstance::StaticClass())
+	{
+		return TEXT("");
+	}
+
+	if (!InAssetData.TagsAndValues.Contains(ModelType))
+	{
+		return TEXT("");		
+	}
+
+	return InAssetData.TagsAndValues.FindTag(ModelType).AsString();
+}
+#endif
+
 UDynamicMaterialInstance::UDynamicMaterialInstance()
 {
 	MaterialModelBase = nullptr;
@@ -32,7 +64,49 @@ UDynamicMaterialModel* UDynamicMaterialInstance::GetMaterialModel()
 	return nullptr;
 }
 
+void UDynamicMaterialInstance::GetAssetRegistryTags(FAssetRegistryTagsContext Context) const
+{
+	Super::GetAssetRegistryTags(Context);
+
+	using namespace UE::DynamicMaterial::Private;
+
+	if (MaterialModelBase)
+	{
+		if (MaterialModelBase->IsA<UDynamicMaterialModel>())
+		{
+			Context.AddTag(FAssetRegistryTag(
+				ModelType, 
+				ModelTypeTag_Material,
+				FAssetRegistryTag::TT_Alphabetical
+			));
+		}
+		else if (MaterialModelBase->IsA<UDynamicMaterialModelDynamic>())
+		{
+			Context.AddTag(FAssetRegistryTag(
+				ModelType,
+				ModelTypeTag_Instance,
+				FAssetRegistryTag::TT_Alphabetical
+			));
+		}
+	}
+}
+
 #if WITH_EDITOR
+void UDynamicMaterialInstance::GetAssetRegistryTagMetadata(TMap<FName, FAssetRegistryTagMetadata>& OutMetadata) const
+{
+	Super::GetAssetRegistryTagMetadata(OutMetadata);
+
+	using namespace UE::DynamicMaterial::Private;
+
+	OutMetadata.Add(
+		ModelType,
+		FAssetRegistryTagMetadata()
+		.SetDisplayName(LOCTEXT("ModelType", "Model Type"))
+		.SetTooltip(LOCTEXT("ModelTypeToolTip", "The type of Model used in this Material"))
+		.SetImportantValue(TEXT("0"))
+	);
+}
+
 void UDynamicMaterialInstance::SetMaterialModel(UDynamicMaterialModelBase* InMaterialModel)
 {
 	MaterialModelBase = InMaterialModel;
@@ -100,4 +174,6 @@ void UDynamicMaterialInstance::OnMaterialBuilt(UDynamicMaterialModelBase* InMate
 
 	InitializeMIDPublic();
 }
-#endif
\ No newline at end of file
+#endif
+
+#undef LOCTEXT_NAMESPACE
diff --git a/Engine/Plugins/Experimental/DynamicMaterial/Source/DynamicMaterial/Public/Material/DynamicMaterialInstance.h b/Engine/Plugins/Experimental/DynamicMaterial/Source/DynamicMaterial/Public/Material/DynamicMaterialInstance.h
--- a/Engine/Plugins/Experimental/DynamicMaterial/Source/DynamicMaterial/Public/Material/DynamicMaterialInstance.h
+++ b/Engine/Plugins/Experimental/DynamicMaterial/Source/DynamicMaterial/Public/Material/DynamicMaterialInstance.h
@@ -20,6 +20,13 @@ class UDynamicMaterialInstance : public UMaterialInstanceDynamic
 	GENERATED_BODY()
 
 public:
+	DYNAMICMATERIAL_API static const FString ModelTypeTag_Material;
+	DYNAMICMATERIAL_API static const FString ModelTypeTag_Instance;
+
+#if WITH_EDITOR
+	DYNAMICMATERIAL_API static FString GetMaterialTypeTag(const FAssetData& InAssetData);
+#endif
+
 	UDynamicMaterialInstance();
 
 	/** Returns the Material Model associated with this Material Designer Material. */
@@ -30,6 +37,10 @@ public:
 	UFUNCTION(BlueprintPure, Category = "Material Designer")
 	DYNAMICMATERIAL_API UDynamicMaterialModel* GetMaterialModel();
 
+	//~ Begin UObject
+	DYNAMICMATERIAL_API virtual void GetAssetRegistryTags(FAssetRegistryTagsContext Context) const override;
+	//~ End UObject
+
 #if WITH_EDITOR
 	/** Sets the Material Model used for this Instance. */
 	DYNAMICMATERIAL_API void SetMaterialModel(UDynamicMaterialModelBase* InMaterialModel);
@@ -43,6 +54,7 @@ public:
 	//~ Begin UObject
 	DYNAMICMATERIAL_API virtual void PostDuplicate(bool bInDuplicateForPIE) override;
 	DYNAMICMATERIAL_API virtual void PostEditImport() override;
+	DYNAMICMATERIAL_API virtual void GetAssetRegistryTagMetadata(TMap<FName, FAssetRegistryTagMetadata>& OutMetadata) const override;
 	//~ End UObject
 #endif
 
diff --git a/Engine/Plugins/Experimental/DynamicMaterial/Source/DynamicMaterialEditor/Private/DMContentBrowserIntegration.cpp b/Engine/Plugins/Experimental/DynamicMaterial/Source/DynamicMaterialEditor/Private/DMContentBrowserIntegration.cpp
--- a/Engine/Plugins/Experimental/DynamicMaterial/Source/DynamicMaterialEditor/Private/DMContentBrowserIntegration.cpp
+++ b/Engine/Plugins/Experimental/DynamicMaterial/Source/DynamicMaterialEditor/Private/DMContentBrowserIntegration.cpp
@@ -23,6 +23,7 @@
 #include "PackageTools.h"
 #include "ScopedTransaction.h"
 #include "Styling/SlateIconFinder.h"
+#include "Utils/DMMaterialModelFunctionLibrary.h"
 
 #define LOCTEXT_NAMESPACE "FDMContentBrowserIntegration"
 
@@ -153,7 +154,7 @@ void FDMContentBrowserIntegration::OnCreateMaterialDesignerMaterialFromTextureSe
 	FString UniquePackageName;
 	FString UniqueAssetName;
 
-	const FString BasePackageName = InPath / TEXT("MDI_NewMaterial");
+	const FString BasePackageName = InPath / TEXT("MD_NewMaterial");
 	AssetToolsModule.Get().CreateUniqueAssetName(BasePackageName, TEXT(""), UniquePackageName, UniqueAssetName);
 
 	UPackage* Package = CreatePackage(*UniquePackageName);
@@ -327,9 +328,11 @@ void FDMContentBrowserIntegration::CreateModelInstance(UDynamicMaterialModel* In
 		return;
 	}
 
+	const FString CurrentName = TEXT("MDM_") + UDMMaterialModelFunctionLibrary::RemoveAssetPrefix(InModel->GetName());
+
 	IAssetTools& AssetTools = FModuleManager::LoadModuleChecked<FAssetToolsModule>("AssetTools").Get();
 	FString PackageName, AssetName;
-	AssetTools.CreateUniqueAssetName(InModel->GetName(), TEXT(""), PackageName, AssetName);
+	AssetTools.CreateUniqueAssetName(CurrentName, TEXT(""), PackageName, AssetName);
 
 	IContentBrowserSingleton& ContentBrowser = FModuleManager::LoadModuleChecked<FContentBrowserModule>("ContentBrowser").Get();
 	const FContentBrowserItemPath CurrentPath = ContentBrowser.GetCurrentPath();
@@ -430,10 +433,7 @@ void FDMContentBrowserIntegration::CreateMaterialInstance(UDynamicMaterialInstan
 	ModelDynamic->SetDynamicMaterialInstance(Instance);
 	Instance->InitializeMIDPublic();
 
-	FString CurrentName = InInstance->GetName();
-	CurrentName = CurrentName.StartsWith(TEXT("MDI_"))
-		? (TEXT("MDD_") + CurrentName.RightChop(4))
-		: (TEXT("MDD_") + CurrentName);
+	const FString CurrentName = TEXT("MDI_") + UDMMaterialModelFunctionLibrary::RemoveAssetPrefix(InInstance->GetName());
 
 	IAssetTools& AssetTools = FModuleManager::LoadModuleChecked<FAssetToolsModule>("AssetTools").Get();
 	FString PackageName, AssetName;
diff --git a/Engine/Plugins/Experimental/DynamicMaterial/Source/DynamicMaterialEditor/Private/Material/AssetDefinition_DynamicMaterialInstance.cpp b/Engine/Plugins/Experimental/DynamicMaterial/Source/DynamicMaterialEditor/Private/Material/AssetDefinition_DynamicMaterialInstance.cpp
--- a/Engine/Plugins/Experimental/DynamicMaterial/Source/DynamicMaterialEditor/Private/Material/AssetDefinition_DynamicMaterialInstance.cpp
+++ b/Engine/Plugins/Experimental/DynamicMaterial/Source/DynamicMaterialEditor/Private/Material/AssetDefinition_DynamicMaterialInstance.cpp
@@ -6,17 +6,30 @@
 #include "IAssetTools.h"
 #include "Material/DynamicMaterialInstance.h"
 #include "Model/DynamicMaterialModel.h"
+#include "Model/DynamicMaterialModelDynamic.h"
 #include "ThumbnailRendering/SceneThumbnailInfoWithPrimitive.h"
 
 #define LOCTEXT_NAMESPACE "AssetDefinition_DynamicMaterialInstance"
 
 FText UAssetDefinition_DynamicMaterialInstance::GetAssetDisplayName() const
 {
-	return LOCTEXT("MaterialDesigner", "Material Designer Material");
+	return LOCTEXT("MaterialDesigner", "Material Designer");
 }
 
 FText UAssetDefinition_DynamicMaterialInstance::GetAssetDisplayName(const FAssetData& InAssetData) const
 {
+	const FString ModelTypeTag = UDynamicMaterialInstance::GetMaterialTypeTag(InAssetData);
+
+	if (ModelTypeTag == UDynamicMaterialInstance::ModelTypeTag_Material)
+	{
+		return LOCTEXT("MaterialDesignerMaterial", "MD Material");
+	}
+
+	if (ModelTypeTag == UDynamicMaterialInstance::ModelTypeTag_Instance)
+	{
+		return LOCTEXT("MaterialDesignerInstance", "MD Instance");
+	}
+
 	return GetAssetDisplayName();
 }
 
@@ -32,7 +45,7 @@ FLinearColor UAssetDefinition_DynamicMaterialInstance::GetAssetColor() const
 
 TConstArrayView<FAssetCategoryPath> UAssetDefinition_DynamicMaterialInstance::GetAssetCategories() const
 {
-	static const TArray<FAssetCategoryPath> Categories = {EAssetCategoryPaths::Material};
+	static TArray<FAssetCategoryPath> Categories = {EAssetCategoryPaths::Material};
 	return Categories;
 }
 
diff --git a/Engine/Plugins/Experimental/DynamicMaterial/Source/DynamicMaterialEditor/Private/Material/DynamicMaterialInstanceFactory.cpp b/Engine/Plugins/Experimental/DynamicMaterial/Source/DynamicMaterialEditor/Private/Material/DynamicMaterialInstanceFactory.cpp
--- a/Engine/Plugins/Experimental/DynamicMaterial/Source/DynamicMaterialEditor/Private/Material/DynamicMaterialInstanceFactory.cpp
+++ b/Engine/Plugins/Experimental/DynamicMaterial/Source/DynamicMaterialEditor/Private/Material/DynamicMaterialInstanceFactory.cpp
@@ -78,7 +78,7 @@ UObject* UDynamicMaterialInstanceFactory::FactoryCreateNew(UClass* Class, UObjec
 
 		TArray<FAnalyticsEventAttribute> Attributes;
 		Attributes.Reserve(2);
-		Attributes.Add(FAnalyticsEventAttribute(TEXT("Action"), TEXT("InstanceCreated")));
+		Attributes.Add(FAnalyticsEventAttribute(TEXT("Action"), TEXT("MaterialCreated")));
 		Attributes.Add(FAnalyticsEventAttribute(TEXT("ActionDetails"), NewInstance->IsAsset() ? AssetType : SubobjectType));
 		FEngineAnalytics::GetProvider().RecordEvent(TEXT("Editor.Usage.MaterialDesigner"), Attributes);
 	}
diff --git a/Engine/Plugins/Experimental/DynamicMaterial/Source/DynamicMaterialEditor/Private/UI/Menus/DMToolBarMenus.cpp b/Engine/Plugins/Experimental/DynamicMaterial/Source/DynamicMaterialEditor/Private/UI/Menus/DMToolBarMenus.cpp
--- a/Engine/Plugins/Experimental/DynamicMaterial/Source/DynamicMaterialEditor/Private/UI/Menus/DMToolBarMenus.cpp
+++ b/Engine/Plugins/Experimental/DynamicMaterial/Source/DynamicMaterialEditor/Private/UI/Menus/DMToolBarMenus.cpp
@@ -295,9 +295,11 @@ void FDMToolBarMenus::ExportMaterial(TWeakObjectPtr<UDynamicMaterialInstance> In
 {
 	if (UDynamicMaterialInstance* Material = InMaterialInstanceWeak.Get())
 	{
+		const FString CurrentName = TEXT("MD_") + UDMMaterialModelFunctionLibrary::RemoveAssetPrefix(Material->GetName());
+
 		IAssetTools& AssetTools = FModuleManager::LoadModuleChecked<FAssetToolsModule>("AssetTools").Get();
 		FString PackageName, AssetName;
-		AssetTools.CreateUniqueAssetName(Material->GetName(), TEXT(""), PackageName, AssetName);
+		AssetTools.CreateUniqueAssetName(CurrentName, TEXT(""), PackageName, AssetName);
 
 		IContentBrowserSingleton& ContentBrowser = FModuleManager::LoadModuleChecked<FContentBrowserModule>("ContentBrowser").Get();
 		const FContentBrowserItemPath CurrentPath = ContentBrowser.GetCurrentPath();
@@ -341,6 +343,10 @@ void FDMToolBarMenus::ExportMaterialModel(TWeakObjectPtr<UDynamicMaterialModelBa
 		return;
 	}
 
+	UDynamicMaterialInstance* Material = MaterialModelBase->GetDynamicMaterialInstance();
+
+	const FString CurrentName = TEXT("M_") + UDMMaterialModelFunctionLibrary::RemoveAssetPrefix(Material ? Material->GetName() : MaterialModelBase->GetName());
+
 	IContentBrowserSingleton& ContentBrowser = FModuleManager::LoadModuleChecked<FContentBrowserModule>("ContentBrowser").Get();
 	const FContentBrowserItemPath CurrentPath = ContentBrowser.GetCurrentPath();
 	const FString PathStr = CurrentPath.HasInternalPath() ? CurrentPath.GetInternalPathString() : TEXT("/Game");
@@ -348,7 +354,7 @@ void FDMToolBarMenus::ExportMaterialModel(TWeakObjectPtr<UDynamicMaterialModelBa
 	FSaveAssetDialogConfig SaveAssetDialogConfig;
 	SaveAssetDialogConfig.DialogTitleOverride = LOCTEXT("SaveAssetDialogTitle", "Save Asset As");
 	SaveAssetDialogConfig.DefaultPath = PathStr;
-	SaveAssetDialogConfig.DefaultAssetName = GeneratedMaterial->GetName();
+	SaveAssetDialogConfig.DefaultAssetName = CurrentName;
 	SaveAssetDialogConfig.ExistingAssetPolicy = ESaveAssetDialogExistingAssetPolicy::Disallow;
 
 	FContentBrowserModule& ContentBrowserModule = FModuleManager::LoadModuleChecked<FContentBrowserModule>("ContentBrowser");
@@ -377,8 +383,9 @@ void FDMToolBarMenus::SnapshotMaterial(TWeakObjectPtr<UDynamicMaterialModelBase>
 	}
 
 	UMaterialInterface* Material = MaterialModelBase->GetGeneratedMaterial();
+	UDynamicMaterialInstance* MaterialInstance = MaterialModelBase->GetDynamicMaterialInstance();
 
-	if (UDynamicMaterialInstance* MaterialInstance = MaterialModelBase->GetDynamicMaterialInstance())
+	if (MaterialInstance)
 	{
 		if (!IsValid(MaterialInstance->Parent.Get()))
 		{
@@ -395,6 +402,13 @@ void FDMToolBarMenus::SnapshotMaterial(TWeakObjectPtr<UDynamicMaterialModelBase>
 		return;
 	}
 
+	const FString CurrentName = TEXT("T_")
+		+ UDMMaterialModelFunctionLibrary::RemoveAssetPrefix(MaterialInstance ? MaterialInstance->GetName() : MaterialModelBase->GetName())
+		+ TEXT("_")
+		+ FString::FromInt(InTextureSize.X)
+		+ TEXT("x")
+		+ FString::FromInt(InTextureSize.Y);
+
 	// Choose asset location
 	IContentBrowserSingleton& ContentBrowser = FModuleManager::LoadModuleChecked<FContentBrowserModule>("ContentBrowser").Get();
 	const FContentBrowserItemPath CurrentPath = ContentBrowser.GetCurrentPath();
@@ -404,7 +418,7 @@ void FDMToolBarMenus::SnapshotMaterial(TWeakObjectPtr<UDynamicMaterialModelBase>
 	SaveAssetDialogConfig.DialogTitleOverride = LOCTEXT("ExportMaterialTo", "Export Material To");
 	SaveAssetDialogConfig.DefaultPath = PathStr;
 	SaveAssetDialogConfig.ExistingAssetPolicy = ESaveAssetDialogExistingAssetPolicy::Disallow;
-	SaveAssetDialogConfig.DefaultAssetName = TEXT("T_MD_") + Material->GetName();
+	SaveAssetDialogConfig.DefaultAssetName = CurrentName;
 
 	FContentBrowserModule& ContentBrowserModule = FModuleManager::LoadModuleChecked<FContentBrowserModule>("ContentBrowser");
 	FString SaveObjectPath = ContentBrowserModule.Get().CreateModalSaveAssetDialog(SaveAssetDialogConfig);
diff --git a/Engine/Plugins/Experimental/DynamicMaterial/Source/DynamicMaterialEditor/Private/Utils/DMMaterialModelFunctionLibrary.cpp b/Engine/Plugins/Experimental/DynamicMaterial/Source/DynamicMaterialEditor/Private/Utils/DMMaterialModelFunctionLibrary.cpp
--- a/Engine/Plugins/Experimental/DynamicMaterial/Source/DynamicMaterialEditor/Private/Utils/DMMaterialModelFunctionLibrary.cpp
+++ b/Engine/Plugins/Experimental/DynamicMaterial/Source/DynamicMaterialEditor/Private/Utils/DMMaterialModelFunctionLibrary.cpp
@@ -122,20 +122,7 @@ UDynamicMaterialInstance* UDMMaterialModelFunctionLibrary::ExportMaterial(UDynam
 		return nullptr;
 	}
 
-	FString CurrentName = MaterialInstance->GetName();
-
-	if (InMaterialModelBase->IsA<UDynamicMaterialModel>())
-	{
-		CurrentName = CurrentName.StartsWith(TEXT("MDI_"))
-			? CurrentName
-			: (TEXT("MDI_") + CurrentName);
-	}
-	else
-	{
-		CurrentName = CurrentName.StartsWith(TEXT("MDD_"))
-			? CurrentName
-			: (TEXT("MDD_") + CurrentName);
-	}
+	const FString CurrentName = (InMaterialModelBase->IsA<UDynamicMaterialModel>() ? TEXT("MD_") : TEXT("MDI_")) + RemoveAssetPrefix(MaterialInstance->GetName());
 
 	IAssetTools& AssetTools = FModuleManager::LoadModuleChecked<FAssetToolsModule>("AssetTools").Get();
 	FString PackageName, AssetName;
@@ -309,11 +296,7 @@ UDynamicMaterialModel* UDMMaterialModelFunctionLibrary::ExportToTemplateMaterial
 	}
 
 	// Where should we save it? (Always export to CB)
-	FString CurrentName = InMaterialModelDynamic->GetName();
-
-	CurrentName = CurrentName.StartsWith(TEXT("MDD_"))
-		? (TEXT("MDM_") + CurrentName.RightChop(4))
-		: (TEXT("MDM_") + CurrentName);
+	const FString CurrentName = TEXT("MDM_") + RemoveAssetPrefix(InMaterialModelDynamic->GetName());
 
 	IAssetTools& AssetTools = FModuleManager::LoadModuleChecked<FAssetToolsModule>("AssetTools").Get();
 	FString PackageName, AssetName;
@@ -403,11 +386,7 @@ UDynamicMaterialInstance* UDMMaterialModelFunctionLibrary::ExportToTemplateMater
 	}
 
 	// Where should we save it? (Always export to CB)
-	FString CurrentName = OldInstance ? OldInstance->GetName() : InMaterialModelDynamic->GetName();
-
-	CurrentName = CurrentName.StartsWith(TEXT("MDD_"))
-		? (TEXT("MDI_") + CurrentName.RightChop(4))
-		: (TEXT("MDI_") + CurrentName);
+	const FString CurrentName = TEXT("MD_") + RemoveAssetPrefix(OldInstance ? OldInstance->GetName() : InMaterialModelDynamic->GetName());
 
 	IAssetTools& AssetTools = FModuleManager::LoadModuleChecked<FAssetToolsModule>("AssetTools").Get();
 	FString PackageName, AssetName;
@@ -624,4 +603,24 @@ bool UDMMaterialModelFunctionLibrary::CreateModelInstanceInMaterial(UDynamicMate
 	return true;
 }
 
+FString UDMMaterialModelFunctionLibrary::RemoveAssetPrefix(const FString& InAssetName)
+{
+	if (InAssetName.StartsWith(TEXT("MD_"))) // Material Designer asset
+	{
+		return InAssetName.RightChop(3);
+	}
+	else if (InAssetName.StartsWith(TEXT("MDI_")) // Material Designer Instance
+		|| InAssetName.StartsWith(TEXT("MDD_")) // Material Designer Dynamic (defunct)
+		|| InAssetName.StartsWith(TEXT("MDM_"))) // Material Designer Model
+	{
+		return InAssetName.RightChop(4);
+	}
+	else if (InAssetName.StartsWith(TEXT("MDMI_"))) // Material Designer Model Instance
+	{
+		return InAssetName.RightChop(5);
+	}
+
+	return InAssetName;
+}
+
 #undef LOCTEXT_NAMESPACE
diff --git a/Engine/Plugins/Experimental/DynamicMaterial/Source/DynamicMaterialEditor/Public/Utils/DMMaterialModelFunctionLibrary.h b/Engine/Plugins/Experimental/DynamicMaterial/Source/DynamicMaterialEditor/Public/Utils/DMMaterialModelFunctionLibrary.h
--- a/Engine/Plugins/Experimental/DynamicMaterial/Source/DynamicMaterialEditor/Public/Utils/DMMaterialModelFunctionLibrary.h
+++ b/Engine/Plugins/Experimental/DynamicMaterial/Source/DynamicMaterialEditor/Public/Utils/DMMaterialModelFunctionLibrary.h
@@ -50,5 +50,5 @@ public:
 
 	DYNAMICMATERIALEDITOR_API static bool CreateModelInstanceInMaterial(UDynamicMaterialModel* InFromModel, UDynamicMaterialInstance* InToInstance);
 
-
+	DYNAMICMATERIALEDITOR_API static FString RemoveAssetPrefix(const FString& InAssetName);
 };
diff --git a/Engine/Plugins/Experimental/DynamicMaterial/Source/DynamicMaterialTextureSetEditor/Private/DMTextureSetContentBrowserIntegrationPrivate.cpp b/Engine/Plugins/Experimental/DynamicMaterial/Source/DynamicMaterialTextureSetEditor/Private/DMTextureSetContentBrowserIntegrationPrivate.cpp
--- a/Engine/Plugins/Experimental/DynamicMaterial/Source/DynamicMaterialTextureSetEditor/Private/DMTextureSetContentBrowserIntegrationPrivate.cpp
+++ b/Engine/Plugins/Experimental/DynamicMaterial/Source/DynamicMaterialTextureSetEditor/Private/DMTextureSetContentBrowserIntegrationPrivate.cpp
@@ -132,7 +132,7 @@ void FDMTextureSetContentBrowserIntegrationPrivate::OnCreateTextureSetComplete(U
 	FString UniquePackageName;
 	FString UniqueAssetName;
 
-	const FString BasePackageName = InPath / TEXT("NewTextureSet");
+	const FString BasePackageName = InPath / TEXT("MDTS_NewTextureSet");
 	AssetToolsModule.Get().CreateUniqueAssetName(BasePackageName, TEXT(""), UniquePackageName, UniqueAssetName);
 
 	UPackage* Package = CreatePackage(*UniquePackageName);
