Material Designer
- Localization fix 590509

diff --git a/Engine/Plugins/Experimental/DynamicMaterial/Source/DynamicMaterialEditor/Private/UI/Menus/DMToolBarMenus.cpp b/Engine/Plugins/Experimental/DynamicMaterial/Source/DynamicMaterialEditor/Private/UI/Menus/DMToolBarMenus.cpp
--- a/Engine/Plugins/Experimental/DynamicMaterial/Source/DynamicMaterialEditor/Private/UI/Menus/DMToolBarMenus.cpp
+++ b/Engine/Plugins/Experimental/DynamicMaterial/Source/DynamicMaterialEditor/Private/UI/Menus/DMToolBarMenus.cpp
@@ -514,12 +514,12 @@ void FDMToolBarMenus::CreateSnapshotMaterialMenu(UToolMenu* InMenu)
 	}
 
 	TWeakObjectPtr<UDynamicMaterialModelBase> MaterialModelWeak = MaterialModelBase;
-	FToolMenuSection& NewSection = InMenu->AddSection("SnapshotMaterial", LOCTEXT("SnapshotMaterial", "Snapshop Material"));
+	FToolMenuSection& NewSection = InMenu->AddSection("SnapshotMaterial", LOCTEXT("SnapshotMaterialSection", "Snapshop Material"));
 
 	const FText SnapshotNameFormat = LOCTEXT("SnapshotName", "{0}x{0}");
 
 	const FText SnapshotTooltipFormat = LOCTEXT(
-		"SnapshotMaterialTooltip",
+		"SnapshotMaterialMenuEntryTooltip",
 		"Take a snapshot of the Material Designer Material with the current values and export it as a Texture asset with a resolution of {0} pixels."
 	);
 
