Material Designer
- Moved the build system into its own editor subsystem.
- Made it copy the build queue before processing it, so that changing the queue during processing does not cause issues.

diff --git a/Engine/Plugins/Experimental/DynamicMaterial/Source/DynamicMaterialEditor/Private/DynamicMaterialEditorModule.cpp b/Engine/Plugins/Experimental/DynamicMaterial/Source/DynamicMaterialEditor/Private/DynamicMaterialEditorModule.cpp
--- a/Engine/Plugins/Experimental/DynamicMaterial/Source/DynamicMaterialEditor/Private/DynamicMaterialEditorModule.cpp
+++ b/Engine/Plugins/Experimental/DynamicMaterial/Source/DynamicMaterialEditor/Private/DynamicMaterialEditorModule.cpp
@@ -246,8 +246,6 @@ void FDynamicMaterialEditorModule::StartupModule()
 	RegisterComponentPropertyRowGeneratorDelegate<UDMMaterialEffectFunction,       FDMMaterialEffectFunctionPropertyRowGenerator>();
 	RegisterComponentPropertyRowGeneratorDelegate<UDMMaterialStageInputFunction,   FDMMaterialStageFunctionPropertyRowGenerator>();
 
-	BuildRequestList.Empty();
-
 	UDMMaterialValueTexture::GetDefaultRGBTexture.BindLambda([]()
 		{
 			const FDMDefaultMaterialPropertySlotValue& DefaultValue = UDynamicMaterialEditorSettings::Get()->GetDefaultSlotValue(EDMMaterialPropertyType::BaseColor);
@@ -288,8 +286,6 @@ void FDynamicMaterialEditorModule::ShutdownModule()
 
 	FDMMaterialModelDefaults::UnregsiterDefaultsDelegates();
 
-	BuildRequestList.Empty();
-
 	UDMMaterialValueTexture::GetDefaultRGBTexture.Unbind();
 
 	FDMValueDetailsRowExtensions::Get().UnregisterRowExtensions();
@@ -431,81 +427,6 @@ TSharedRef<SWidget> FDynamicMaterialEditorModule::CreateEditor(UDynamicMaterialM
 	return NewDesigner;
 }
 
-void FDynamicMaterialEditorModule::Tick(float DeltaTime)
-{
-	if (UDMMaterialComponent::CanClean() == false)
-	{
-		return;
-	}
-
-	for (const FDMBuildRequestEntry& ToBuild : BuildRequestList)
-	{
-		if (UObject* Object = FindObject<UObject>(nullptr, *ToBuild.AssetPath, false))
-		{
-			ProcessBuildRequest(Object, ToBuild.bDirtyAssets);
-		}
-	}
-
-	BuildRequestList.Empty();
-}
-
-TStatId FDynamicMaterialEditorModule::GetStatId() const
-{
-	RETURN_QUICK_DECLARE_CYCLE_STAT(FDynamicMaterialEditorModule, STATGROUP_Tickables);
-}
-
-void FDynamicMaterialEditorModule::AddBuildRequest(UObject* InToBuild, bool bInDirtyAssets)
-{
-	if (!IsValid(InToBuild))
-	{
-		return;
-	}
-
-	FDynamicMaterialEditorModule::Get().BuildRequestList.Add({InToBuild->GetPathName(), bInDirtyAssets});
-	
-	static const double VeryShortTime = 0.0001;
-
-	// Make sure we don't spam updates on a single tick.
-	UDMMaterialComponent::PreventClean(VeryShortTime);
-}
-
-void FDynamicMaterialEditorModule::RemoveBuildRequest(UObject* InToNotBuild)
-{
-	if (!InToNotBuild)
-	{
-		return;
-	}
-
-	FDynamicMaterialEditorModule::Get().BuildRequestList.Remove({InToNotBuild->GetPathName(), false});
-}
-
-void FDynamicMaterialEditorModule::RemoveBuildRequestForOuter(UObject* InOuter)
-{
-	if (!InOuter)
-	{
-		return;
-	}
-
-	const FString ObjectPath = InOuter->GetPathName();
-	const int32 ObjectPathLength = ObjectPath.Len();
-
-	for (TSet<FDMBuildRequestEntry>::TIterator Iter(FDynamicMaterialEditorModule::Get().BuildRequestList); Iter; ++Iter)
-	{
-		if (Iter->AssetPath.Len() > ObjectPathLength && Iter->AssetPath.StartsWith(ObjectPath))
-		{
-			// Make sure it's a path separator after the parent path.
-			switch (Iter->AssetPath[ObjectPathLength])
-			{
-				case '.':
-				case '/':
-				case ':':
-					Iter.RemoveCurrent();
-					break;
-			}
-		}
-	}
-}
-
 void FDynamicMaterialEditorModule::OpenEditor(UWorld* InWorld) const
 {
 	if (!IsValid(InWorld))
@@ -536,21 +457,6 @@ UDynamicMaterialModelBase* FDynamicMaterialEditorModule::GetOpenedMaterialModel(
 	return nullptr;
 }
 
-void FDynamicMaterialEditorModule::ProcessBuildRequest(UObject* InToBuild, bool bInDirtyAssets)
-{
-	if (!IsValid(InToBuild))
-	{
-		return;
-	}
-
-	if (InToBuild->GetClass()->ImplementsInterface(UDMBuildable::StaticClass()) == false)
-	{
-		return;
-	}
-
-	IDMBuildable::Execute_DoBuild(InToBuild, bInDirtyAssets);
-}
-
 void FDynamicMaterialEditorModule::MapCommands()
 {
 	const FDynamicMaterialEditorCommands& DMEditorCommands = FDynamicMaterialEditorCommands::Get();
diff --git a/Engine/Plugins/Experimental/DynamicMaterial/Source/DynamicMaterialEditor/Private/DynamicMaterialEditorModule.h b/Engine/Plugins/Experimental/DynamicMaterial/Source/DynamicMaterialEditor/Private/DynamicMaterialEditorModule.h
--- a/Engine/Plugins/Experimental/DynamicMaterial/Source/DynamicMaterialEditor/Private/DynamicMaterialEditorModule.h
+++ b/Engine/Plugins/Experimental/DynamicMaterial/Source/DynamicMaterialEditor/Private/DynamicMaterialEditorModule.h
@@ -3,7 +3,6 @@
 #pragma once
 
 #include "IDynamicMaterialEditorModule.h"
-#include "TickableEditorObject.h"
 
 #include "DMEDefs.h"
 #include "Templates/SharedPointer.h"
@@ -21,7 +20,6 @@ class UDMMaterialComponent;
 class UDMMaterialStageSource;
 class UDynamicMaterialInstance;
 class UDynamicMaterialModelBase;
-struct FDMBuildRequestList;
 
 DECLARE_LOG_CATEGORY_EXTERN(LogDynamicMaterialEditor, Log, All);
 
@@ -41,26 +39,10 @@ DECLARE_DELEGATE_RetVal_TwoParams(TSharedPtr<SWidget>, FDMCreateValueEditWidgetD
 DECLARE_DELEGATE_FourParams(FDMComponentPropertyRowGeneratorDelegate, const TSharedRef<SDMMaterialComponentEditor>&, UDMMaterialComponent*,
 	TArray<FDMPropertyHandle>&, TSet<UDMMaterialComponent*>&)
 
-struct FDMBuildRequestEntry
-{
-	FString AssetPath;
-	bool bDirtyAssets;
-
-	bool operator==(const FDMBuildRequestEntry& Other) const
-	{
-		return AssetPath == Other.AssetPath;
-	}
-
-	friend uint32 GetTypeHash(const FDMBuildRequestEntry& InEntry)
-	{
-		return GetTypeHash(InEntry.AssetPath);
-	}
-};
-
 /**
  * Material Designer - Build your own materials in a slimline editor!
  */
-class FDynamicMaterialEditorModule : public IDynamicMaterialEditorModule, public FTickableEditorObject
+class FDynamicMaterialEditorModule : public IDynamicMaterialEditorModule
 {
 public:
 	static const FName TabId;
@@ -81,6 +63,10 @@ public:
 
 	FDynamicMaterialEditorModule();
 
+	const TSharedRef<FUICommandList>& GetCommandList() const { return CommandList; }
+
+	void OnWizardComplete(UDynamicMaterialModel* InModel);
+
 	//~ Begin IDynamicMaterialEditorModule
 	virtual void RegisterCustomMaterialPropertyGenerator(UClass* InClass, FDMGetObjectMaterialPropertiesDelegate InGenerator) override;
 	virtual void RegisterMaterialModelCreatedCallback(const TSharedRef<IDMOnWizardCompleteCallback> InCallback)  override;
@@ -99,33 +85,14 @@ public:
 	virtual void ShutdownModule() override;
 	//~ End IModuleInterface
 
-	//~ Begin FTickableEditorObject
-	virtual void Tick(float DeltaTime) override;
-	virtual ETickableTickType GetTickableTickType() const override { return ETickableTickType::Always; }
-	virtual TStatId GetStatId() const override;
-	//~ End FTickableEditorObject
-
-	void AddBuildRequest(UObject* InToBuild, bool bInDirtyAssets);
-
-	void RemoveBuildRequest(UObject* InToNotBuild);
-
-	void RemoveBuildRequestForOuter(UObject* InOuter);
-
-	const TSharedRef<FUICommandList>& GetCommandList() const { return CommandList; }
-
-	void OnWizardComplete(UDynamicMaterialModel* InModel);
-
 protected:
 	static TMap<UClass*, FDMComponentPropertyRowGeneratorDelegate> ComponentPropertyRowGenerators;
 	static TMap<UClass*, FDMGetObjectMaterialPropertiesDelegate> CustomMaterialPropertyGenerators;
 	static FDMOnUIValueUpdate OnUIValueUpdate;
 	static TArray<TSharedRef<IDMOnWizardCompleteCallback>> OnWizardCompleteCallbacks;
 
-	TSet<FDMBuildRequestEntry> BuildRequestList;
 	TSharedRef<FUICommandList> CommandList;
 
-	void ProcessBuildRequest(UObject* InToBuild, bool bInDirtyAssets);
-
 	void MapCommands();
 	void UnmapCommands();
 };
diff --git a/Engine/Plugins/Experimental/DynamicMaterial/Source/DynamicMaterialEditor/Private/LevelEditor/DMLevelEditorIntegration.cpp b/Engine/Plugins/Experimental/DynamicMaterial/Source/DynamicMaterialEditor/Private/LevelEditor/DMLevelEditorIntegration.cpp
--- a/Engine/Plugins/Experimental/DynamicMaterial/Source/DynamicMaterialEditor/Private/LevelEditor/DMLevelEditorIntegration.cpp
+++ b/Engine/Plugins/Experimental/DynamicMaterial/Source/DynamicMaterialEditor/Private/LevelEditor/DMLevelEditorIntegration.cpp
@@ -2,7 +2,6 @@
 
 #include "LevelEditor/DMLevelEditorIntegration.h"
 
-#include "DynamicMaterialEditorModule.h"
 #include "LevelEditor.h"
 #include "LevelEditor/DMLevelEditorIntegrationInstance.h"
 #include "Model/DynamicMaterialModelBase.h"
@@ -10,6 +9,7 @@
 #include "Modules/ModuleManager.h"
 #include "Templates/SharedPointer.h"
 #include "UI/Widgets/SDMMaterialDesigner.h"
+#include "Utils/DMBuildRequestSubsystem.h"
 
 namespace UE::DynamicMaterialEditor::Private
 {
@@ -138,7 +138,10 @@ TSharedPtr<SDockTab> FDMLevelEditorIntegration::InvokeTabForWorld(UWorld* InWorl
 
 void FDMLevelEditorIntegration::OnMapTearDown(UWorld* InWorld)
 {
-	FDynamicMaterialEditorModule::Get().RemoveBuildRequestForOuter(InWorld);
+	if (UDMBuildRequestSubsystem* BuildRequestSubsystem = UDMBuildRequestSubsystem::Get())
+	{
+		BuildRequestSubsystem->RemoveBuildRequestForOuter(InWorld);
+	}
 
 	FDMLevelEditorIntegrationInstance* Instance = FDMLevelEditorIntegrationInstance::GetMutableIntegrationForWorld(InWorld);
 
@@ -213,7 +216,10 @@ void FDMLevelEditorIntegration::OnMapLoad(UWorld* InWorld)
 
 			if (UDynamicMaterialModelEditorOnlyData* EditorOnlyData = UDynamicMaterialModelEditorOnlyData::Get(MaterialModel))
 			{
-				FDynamicMaterialEditorModule::Get().AddBuildRequest(EditorOnlyData, /* Dirty Assets */ false);
+				if (UDMBuildRequestSubsystem* BuildRequestSubsystem = UDMBuildRequestSubsystem::Get())
+				{
+					BuildRequestSubsystem->AddBuildRequest(EditorOnlyData, /* Dirty Assets */ false);
+				}				
 			}
 		}
 	}
diff --git a/Engine/Plugins/Experimental/DynamicMaterial/Source/DynamicMaterialEditor/Private/Model/DynamicMaterialModelEditorOnlyData.cpp b/Engine/Plugins/Experimental/DynamicMaterial/Source/DynamicMaterialEditor/Private/Model/DynamicMaterialModelEditorOnlyData.cpp
--- a/Engine/Plugins/Experimental/DynamicMaterial/Source/DynamicMaterialEditor/Private/Model/DynamicMaterialModelEditorOnlyData.cpp
+++ b/Engine/Plugins/Experimental/DynamicMaterial/Source/DynamicMaterialEditor/Private/Model/DynamicMaterialModelEditorOnlyData.cpp
@@ -49,8 +49,8 @@
 #include "Misc/Guid.h"
 #include "Model/DMMaterialBuildState.h"
 #include "Model/DynamicMaterialModel.h"
-#include "Model/DynamicMaterialModelEditorOnlyDataVersion.h"
 #include "UObject/Package.h"
+#include "Utils/DMBuildRequestSubsystem.h"
 #include "Utils/DMUtils.h"
 
 #define LOCTEXT_NAMESPACE "MaterialDesignerModel"
@@ -582,9 +582,9 @@ void UDynamicMaterialModelEditorOnlyData::RequestMaterialBuild(bool bInImmediate
 		{
 			BuildMaterial(/* Dirty Assets */ false);
 		}
-		else
+		else if (UDMBuildRequestSubsystem* BuildRequestSubsystem = UDMBuildRequestSubsystem::Get())
 		{
-			FDynamicMaterialEditorModule::Get().AddBuildRequest(this, /* Dirty Packages */ !UE::GetIsEditorLoadingPackage());
+			BuildRequestSubsystem->AddBuildRequest(this, /* Dirty Packages */ !UE::GetIsEditorLoadingPackage());
 		}
 	}
 }
diff --git a/Engine/Plugins/Experimental/DynamicMaterial/Source/DynamicMaterialEditor/Private/Utils/DMBuildRequestSubsystem.cpp b/Engine/Plugins/Experimental/DynamicMaterial/Source/DynamicMaterialEditor/Private/Utils/DMBuildRequestSubsystem.cpp
new file mode 100644
--- /dev/null
+++ b/Engine/Plugins/Experimental/DynamicMaterial/Source/DynamicMaterialEditor/Private/Utils/DMBuildRequestSubsystem.cpp
@@ -0,0 +1,119 @@
+// Copyright Epic Games, Inc. All Rights Reserved.
+
+#include "Utils/DMBuildRequestSubsystem.h"
+
+#include "Components/DMMaterialComponent.h"
+#include "DMEDefs.h"
+#include "Editor.h"
+
+UDMBuildRequestSubsystem* UDMBuildRequestSubsystem::Get()
+{
+	if (!GEditor)
+	{
+		return nullptr;
+	}
+
+	return GEditor->GetEditorSubsystem<UDMBuildRequestSubsystem>();
+}
+
+void UDMBuildRequestSubsystem::Deinitialize()
+{
+	Super::Deinitialize();
+
+	BuildRequestList.Empty();
+}
+
+void UDMBuildRequestSubsystem::Tick(float InDeltaTime)
+{
+	ProcessBuildRequestList();
+}
+
+TStatId UDMBuildRequestSubsystem::GetStatId() const
+{
+	RETURN_QUICK_DECLARE_CYCLE_STAT(UDMBuildRequestSubsystem, STATGROUP_Tickables);
+}
+
+void UDMBuildRequestSubsystem::AddBuildRequest(UObject* InToBuild, bool bInDirtyAssets)
+{
+	if (!IsValid(InToBuild))
+	{
+		return;
+	}
+
+	BuildRequestList.Add({InToBuild->GetPathName(), bInDirtyAssets});
+
+	// Make sure we don't spam updates on a single tick.
+	UDMMaterialComponent::PreventClean(UE_KINDA_SMALL_NUMBER);
+}
+
+void UDMBuildRequestSubsystem::RemoveBuildRequest(UObject* InToNotBuild)
+{
+	if (!InToNotBuild)
+	{
+		return;
+	}
+
+	BuildRequestList.Remove({InToNotBuild->GetPathName(), /* Dirty Assets */ false});
+}
+
+void UDMBuildRequestSubsystem::RemoveBuildRequestForOuter(UObject* InOuter)
+{
+	if (!InOuter)
+	{
+		return;
+	}
+
+	const FString ObjectPath = InOuter->GetPathName();
+	const int32 ObjectPathLength = ObjectPath.Len();
+
+	for (TSet<FDMBuildRequestEntry>::TIterator Iter(BuildRequestList); Iter; ++Iter)
+	{
+		if (Iter->AssetPath.Len() > ObjectPathLength && Iter->AssetPath.StartsWith(ObjectPath))
+		{
+			// Make sure it's a path separator after the parent path.
+			switch (Iter->AssetPath[ObjectPathLength])
+			{
+				case '.':
+				case '/':
+				case ':':
+					Iter.RemoveCurrent();
+					break;
+			}
+		}
+	}
+}
+
+void UDMBuildRequestSubsystem::ProcessBuildRequestList()
+{
+	if (UDMMaterialComponent::CanClean() == false)
+	{
+		return;
+	}
+
+	TArray<FDMBuildRequestEntry> BuildRequestListLocal;
+	BuildRequestListLocal.Append(BuildRequestList.Array());
+	BuildRequestList.Empty();
+
+	for (const FDMBuildRequestEntry& ToBuild : BuildRequestListLocal)
+	{
+		if (UObject* Object = FindObject<UObject>(nullptr, *ToBuild.AssetPath, false))
+		{
+			ProcessBuildRequest(Object, ToBuild.bDirtyAssets);
+		}
+	}
+}
+
+void UDMBuildRequestSubsystem::ProcessBuildRequest(UObject* InToBuild, bool bInDirtyAssets)
+{
+	if (!IsValid(InToBuild))
+	{
+		return;
+	}
+
+	if (InToBuild->GetClass()->ImplementsInterface(UDMBuildable::StaticClass()) == false)
+	{
+		return;
+	}
+
+	IDMBuildable::Execute_DoBuild(InToBuild, bInDirtyAssets);
+}
diff --git a/Engine/Plugins/Experimental/DynamicMaterial/Source/DynamicMaterialEditor/Private/Utils/DMBuildRequestSubsystem.h b/Engine/Plugins/Experimental/DynamicMaterial/Source/DynamicMaterialEditor/Private/Utils/DMBuildRequestSubsystem.h
new file mode 100644
--- /dev/null
+++ b/Engine/Plugins/Experimental/DynamicMaterial/Source/DynamicMaterialEditor/Private/Utils/DMBuildRequestSubsystem.h
@@ -0,0 +1,58 @@
+// Copyright Epic Games, Inc. All Rights Reserved.
+
+#pragma once
+
+#include "EditorSubsystem.h"
+#include "TickableEditorObject.h"
+
+#include "Containers/Set.h"
+
+#include "DMBuildRequestSubsystem.generated.h"
+
+struct FDMBuildRequestEntry
+{
+	FString AssetPath;
+	bool bDirtyAssets;
+
+	bool operator==(const FDMBuildRequestEntry& InOther) const
+	{
+		return AssetPath == InOther.AssetPath;
+	}
+
+	friend uint32 GetTypeHash(const FDMBuildRequestEntry& InEntry)
+	{
+		return GetTypeHash(InEntry.AssetPath);
+	}
+};
+
+UCLASS()
+class UDMBuildRequestSubsystem : public UEditorSubsystem, public FTickableEditorObject
+{
+	GENERATED_BODY()
+
+public:
+	static UDMBuildRequestSubsystem* Get();
+
+	void AddBuildRequest(UObject* InToBuild, bool bInDirtyAssets);
+
+	void RemoveBuildRequest(UObject* InToNotBuild);
+
+	void RemoveBuildRequestForOuter(UObject* InOuter);
+
+	//~ Begin USubsystem
+	virtual void Deinitialize() override;
+	//~ End USubsystem
+
+	//~ Begin FTickableEditorObject
+	virtual void Tick(float InDeltaTime) override;
+	virtual ETickableTickType GetTickableTickType() const override { return ETickableTickType::Always; }
+	virtual TStatId GetStatId() const override;
+	//~ End FTickableEditorObject
+
+protected:
+	TSet<FDMBuildRequestEntry> BuildRequestList;
+
+	void ProcessBuildRequestList();
+
+	void ProcessBuildRequest(UObject* InToBuild, bool bInDirtyAssets);
+};
