Material Designer
- When selecting an actor with no valid material slot, it will tell you this.

diff --git a/Engine/Plugins/Experimental/DynamicMaterial/Source/DynamicMaterialEditor/Private/UI/Widgets/SDMActorMaterialSelector.cpp b/Engine/Plugins/Experimental/DynamicMaterial/Source/DynamicMaterialEditor/Private/UI/Widgets/SDMActorMaterialSelector.cpp
--- a/Engine/Plugins/Experimental/DynamicMaterial/Source/DynamicMaterialEditor/Private/UI/Widgets/SDMActorMaterialSelector.cpp
+++ b/Engine/Plugins/Experimental/DynamicMaterial/Source/DynamicMaterialEditor/Private/UI/Widgets/SDMActorMaterialSelector.cpp
@@ -120,6 +120,11 @@ TSharedRef<SWidget> SDMActorMaterialSelector::CreateSelectorLayout()
 			];
 	}
 
+	if (ListOuter->NumSlots() == 1)
+	{
+		return CreateNoPropertiesLayout();
+	}
+
 	return SNew(SScrollBox)
 		.Orientation(EOrientation::Orient_Vertical)
 		+ SScrollBox::Slot()
