Material Designer
- Redesign of the wizard new material page.
- Fix for the channels list so it appears correctly.
- Fix for the template selector to enable engine, plugin and developer content.

diff --git a/Engine/Plugins/Experimental/DynamicMaterial/Source/DynamicMaterialEditor/Private/DynamicMaterialEditorStyle.cpp b/Engine/Plugins/Experimental/DynamicMaterial/Source/DynamicMaterialEditor/Private/DynamicMaterialEditorStyle.cpp
--- a/Engine/Plugins/Experimental/DynamicMaterial/Source/DynamicMaterialEditor/Private/DynamicMaterialEditorStyle.cpp
+++ b/Engine/Plugins/Experimental/DynamicMaterial/Source/DynamicMaterialEditor/Private/DynamicMaterialEditorStyle.cpp
@@ -7,6 +7,7 @@
 #include "Interfaces/IPluginManager.h"
 #include "Styling/SlateStyleMacros.h"
 #include "Styling/SlateStyleRegistry.h"
+#include "Styling/StarshipCoreStyle.h"
 #include "Styling/StyleColors.h"
 
 namespace UE::DynamicMaterialEditor::Private
@@ -123,6 +124,11 @@ void FDynamicMaterialEditorStyle::SetupGeneralStyles()
 		.SetPressed(FSlateRoundedBoxBrush(FLinearColor(1, 1, 1, 0.25f), 4.0f, FLinearColor(1, 1, 1, 0.5f), 1.0f))
 		.SetNormalPadding(FMargin(0, 0, 0, 1))
 		.SetPressedPadding(FMargin(0, 1, 0, 0)));
+
+	Set("DulledSectionButton", FCheckBoxStyle(FStarshipCoreStyle::GetCoreStyle().GetWidgetStyle<FCheckBoxStyle>("FilterBar.BasicFilterButton"))
+		.SetCheckedImage(FSlateRoundedBoxBrush(FStyleColors::AccentBrown, 4.0f, FStyleColors::Input, 1.0f))
+		.SetCheckedHoveredImage(FSlateRoundedBoxBrush(FStyleColors::AccentFolder, 4.0f, FStyleColors::Input, 1.0f))
+		.SetCheckedPressedImage(FSlateRoundedBoxBrush(FStyleColors::AccentFolder, 4.0f, FStyleColors::Input, 1.0f)));
 }
 
 void FDynamicMaterialEditorStyle::SetupStageStyles()
diff --git a/Engine/Plugins/Experimental/DynamicMaterial/Source/DynamicMaterialEditor/Private/UI/Widgets/SDMMaterialWizard.cpp b/Engine/Plugins/Experimental/DynamicMaterial/Source/DynamicMaterialEditor/Private/UI/Widgets/SDMMaterialWizard.cpp
--- a/Engine/Plugins/Experimental/DynamicMaterial/Source/DynamicMaterialEditor/Private/UI/Widgets/SDMMaterialWizard.cpp
+++ b/Engine/Plugins/Experimental/DynamicMaterial/Source/DynamicMaterialEditor/Private/UI/Widgets/SDMMaterialWizard.cpp
@@ -164,16 +164,6 @@ TSharedRef<SWidget> SDMMaterialWizard::CreateModeSelector()
 {
 	TSharedRef<SHorizontalBox> Container = SNew(SHorizontalBox);
 
-	const TArray<EHorizontalAlignment> Alignments = {
-		EHorizontalAlignment::HAlign_Right,
-		EHorizontalAlignment::HAlign_Left
-	};
-
-	const TArray<FMargin> Paddings = {
-		FMargin(0.f, 0.f, 5.f, 0.f),
-		FMargin(5.f, 0.f, 0.f, 0.f)
-	};
-
 	for (int32 SwitcherIndex = 0; SwitcherIndex < EDMMaterialWizardMode::Count; ++SwitcherIndex)
 	{
 		switch (SwitcherIndex)
@@ -181,9 +171,9 @@ TSharedRef<SWidget> SDMMaterialWizard::CreateModeSelector()
 			case EDMMaterialWizardMode::Template:
 			{
 				Container->AddSlot()
-					.FillContentWidth(1.f)
-					.HAlign(Alignments[SwitcherIndex])
-					.Padding(Paddings[SwitcherIndex])
+					.AutoWidth()
+					.HAlign(EHorizontalAlignment::HAlign_Left)
+					.Padding(0.f, 0.f, 5.f, 0.f)
 					[
 						SNew(SCheckBox)
 						.Style(FAppStyle::Get(), "DetailsView.SectionButton")
@@ -206,9 +196,9 @@ TSharedRef<SWidget> SDMMaterialWizard::CreateModeSelector()
 			case EDMMaterialWizardMode::Instance:
 			{
 				Container->AddSlot()
-					.FillContentWidth(1.f)
-					.HAlign(Alignments[SwitcherIndex])
-					.Padding(Paddings[SwitcherIndex])
+					.AutoWidth()
+					.HAlign(EHorizontalAlignment::HAlign_Left)
+					.Padding(0.f, 0.f, 5.f, 0.f)
 					[
 						SNew(SCheckBox)
 						.Style(FAppStyle::Get(), "DetailsView.SectionButton")
@@ -240,7 +230,7 @@ TSharedRef<SWidget> SDMMaterialWizard::CreateNewTemplateLayout()
 	return SNew(SVerticalBox)
 		+ SVerticalBox::Slot()
 		.AutoHeight()
-		.HAlign(HAlign_Left)
+		.HAlign(EHorizontalAlignment::HAlign_Left)
 		.Padding(0.0f, SeparationDistance, 0.0f, TitleContentDistance)
 		[
 			SNew(STextBlock)
@@ -258,7 +248,7 @@ TSharedRef<SWidget> SDMMaterialWizard::CreateNewTemplateLayout()
 
 		+ SVerticalBox::Slot()
 		.AutoHeight()
-		.HAlign(HAlign_Left)
+		.HAlign(EHorizontalAlignment::HAlign_Left)
 		.Padding(0.0f, SeparationDistance, 0.0f, TitleContentDistance)
 		[
 			SNew(STextBlock)
@@ -279,7 +269,7 @@ TSharedRef<SWidget> SDMMaterialWizard::CreateNewTemplateLayout()
 
 		+ SVerticalBox::Slot()
 		.AutoHeight()
-		.HAlign(EHorizontalAlignment::HAlign_Fill)
+		.HAlign(EHorizontalAlignment::HAlign_Left)
 		.Padding(0.0f, SeparationDistance, 0.0f, 0.f)
 		[
 			CreateNewTemplate_AcceptButton()
@@ -300,7 +290,7 @@ TSharedRef<SWidget> SDMMaterialWizard::CreateNewTemplate_ChannelPresets()
 		ChannelPresets->AddSlot()
 			[
 				SNew(SCheckBox)
-				.Style(FAppStyle::Get(), "DetailsView.SectionButton")
+				.Style(FDynamicMaterialEditorStyle::Get(), "DulledSectionButton")
 				.HAlign(EHorizontalAlignment::HAlign_Center)
 				.Padding(ButtonPadding)
 				.IsChecked(this, &SDMMaterialWizard::Preset_GetState, Preset.Name)
@@ -320,13 +310,6 @@ TSharedRef<SWidget> SDMMaterialWizard::CreateNewTemplate_ChannelList()
 {
 	using namespace UE::DynamicMaterialDesigner::Private;
 
-	UDynamicMaterialModelEditorOnlyData* ModelEditorOnlyData = UDynamicMaterialModelEditorOnlyData::Get(GetMaterialModel());
-
-	if (!ModelEditorOnlyData)
-	{
-		return SNullWidget::NullWidget;
-	}
-
 	TSharedRef<SWrapBox> ChannelPresets = SNew(SWrapBox)
 		.UseAllottedSize(true)
 		.InnerSlotPadding(WrapBoxSlotPadding)
@@ -334,23 +317,38 @@ TSharedRef<SWidget> SDMMaterialWizard::CreateNewTemplate_ChannelList()
 
 	if (const FDMMaterialChannelListPreset* Preset = GetDefault<UDynamicMaterialEditorSettings>()->GetPresetByName(CurrentPreset))
 	{
-		for (const TPair<EDMMaterialPropertyType, UDMMaterialProperty*>& Property : ModelEditorOnlyData->GetMaterialProperties())
+		TArray<EDMMaterialPropertyType> ChannelList;
+
+		for (const TPair<EDMMaterialPropertyType, UDMMaterialProperty*>& Property : GetDefault<UDynamicMaterialModelEditorOnlyData>()->GetMaterialProperties())
 		{
 			if (Property.Key == EDMMaterialPropertyType::OpacityMask)
 			{
 				continue;
 			}
 
-			if (Preset->IsPropertyEnabled(Property.Key))
+			if (!Preset->IsPropertyEnabled(Property.Key))
 			{
-				ChannelPresets->AddSlot()
-					.Padding(TextPadding)
-					[
-						SNew(STextBlock)
-						.TextStyle(FDynamicMaterialEditorStyle::Get(), "RegularFont")
-						.Text(UE::DynamicMaterialEditor::Private::GetMaterialPropertyShortDisplayName(Property.Key))
-					];
+				continue;
 			}
+
+			ChannelList.Add(Property.Key);
+		}
+
+		for (int32 Index = 0; Index < ChannelList.Num(); ++Index)
+		{
+			const FText PropertyName = UE::DynamicMaterialEditor::Private::GetMaterialPropertyLongDisplayName(ChannelList[Index]);
+
+			const FText Name = Index == (ChannelList.Num() - 1)
+				? FText::Format(LOCTEXT("ListEnd", "{0}."), PropertyName)
+				: FText::Format(LOCTEXT("ListEntry", "{0},"), PropertyName);
+
+			ChannelPresets->AddSlot()
+				.Padding(TextPadding)
+				[
+					SNew(STextBlock)
+					.TextStyle(FAppStyle::Get(), "TinyText")
+					.Text(Name)
+				];
 		}
 	}
 
@@ -359,28 +357,57 @@ TSharedRef<SWidget> SDMMaterialWizard::CreateNewTemplate_ChannelList()
 
 TSharedRef<SWidget> SDMMaterialWizard::CreateNewInstanceLayout()
 {
-	FAssetPickerConfig PickerConfig;
-	PickerConfig.SelectionMode = ESelectionMode::Single;
-	PickerConfig.Filter.ClassPaths.Add(FTopLevelAssetPath(UDynamicMaterialInstance::StaticClass()));
-	PickerConfig.Filter.bIncludeOnlyOnDiskAssets = true;
-	PickerConfig.ThumbnailLabel = EThumbnailLabel::AssetName;
-	PickerConfig.OnShouldFilterAsset.BindSP(this, &SDMMaterialWizard::ShouldFilterOutAsset);
-	PickerConfig.InitialAssetViewType = EAssetViewType::Tile;
-	PickerConfig.bAutohideSearchBar = false;
-	PickerConfig.bFocusSearchBoxWhenOpened = false;
-	PickerConfig.bShowBottomToolbar = false;
-	PickerConfig.bAllowDragging = false;
-	PickerConfig.bAllowRename = false;
-	PickerConfig.bCanShowClasses = false;
-	PickerConfig.bCanShowFolders = true;
-	PickerConfig.bCanShowReadOnlyFolders = true;
-	PickerConfig.bCanShowRealTimeThumbnails = true;
-	PickerConfig.bCanShowDevelopersFolder = true;
-	PickerConfig.bForceShowEngineContent = true;
-	PickerConfig.bForceShowPluginContent = true;
-	PickerConfig.bAddFilterUI = false;
-
-	PickerConfig.Filter.PackagePaths.Reset();
+	return SNew(SVerticalBox)
+
+		+ SVerticalBox::Slot()
+		.AutoHeight()
+		[		
+			SNew(SBorder)
+			.BorderImage(FAppStyle::Get().GetBrush("Menu.Background"))
+			.Padding(3.f)
+			.VAlign(EVerticalAlignment::VAlign_Fill)
+			[
+				SNew(SVerticalBox)
+
+				+ SVerticalBox::Slot()
+				.AutoHeight()
+				.Padding(0.f, 0.f, 0.f, 3.f)
+				[
+					CreateNewInstance_SearchBox()
+				]
+
+				+ SVerticalBox::Slot()
+				.AutoHeight()
+				[
+					CreateNewInstance_Picker()
+				]
+			]
+		]
+
+		+ SVerticalBox::Slot()
+		.AutoHeight()
+		.HAlign(EHorizontalAlignment::HAlign_Left)
+		.Padding(0.f, 10.f, 0.f, 0.f)
+		[
+			CreateNewInstance_AcceptButton()
+		];
+}
+
+TSharedRef<SWidget> SDMMaterialWizard::CreateNewInstance_SearchBox()
+{
+	return SAssignNew(AssetSearchBox, SAssetSearchBox)
+		.HintText(NSLOCTEXT("ContentBrowser", "SearchBoxHint", "Search Assets"))
+		.OnTextChanged(this, &SDMMaterialWizard::OnSearchBoxChanged)
+		.OnTextCommitted(this, &SDMMaterialWizard::OnSearchBoxCommitted)
+		.DelayChangeNotificationsWhileTyping(true);
+}
+
+TSharedRef<SWidget> SDMMaterialWizard::CreateNewInstance_Picker()
+{
+	FARFilter Filter;
+	Filter.ClassPaths.Add(FTopLevelAssetPath(UDynamicMaterialInstance::StaticClass()));
+	Filter.bIncludeOnlyOnDiskAssets = true;
+	Filter.PackagePaths.Reset();
 
 	TextFilter = MakeShared<FAssetTextFilter>();
 	TextFilter->SetIncludeClassName(true);
@@ -389,68 +416,53 @@ TSharedRef<SWidget> SDMMaterialWizard::CreateNewInstanceLayout()
 
 	AssetView = SNew(SAssetView)
 		.InitialCategoryFilter(EContentBrowserItemCategoryFilter::IncludeAssets)
-		.SelectionMode(PickerConfig.SelectionMode)
-		.OnShouldFilterAsset(PickerConfig.OnShouldFilterAsset)
+		.SelectionMode(ESelectionMode::Single)
+		.OnShouldFilterAsset(this, &SDMMaterialWizard::ShouldFilterOutAsset)
 		.OnItemsActivated(this, &SDMMaterialWizard::OnAssetsActivated)
-		.OnIsAssetValidForCustomToolTip(PickerConfig.OnIsAssetValidForCustomToolTip)
-		.OnGetCustomAssetToolTip(PickerConfig.OnGetCustomAssetToolTip)
-		.OnVisualizeAssetToolTip(PickerConfig.OnVisualizeAssetToolTip)
-		.OnAssetToolTipClosing(PickerConfig.OnAssetToolTipClosing)
 		.InitialSourcesData(FSourcesData())
-		.InitialBackendFilter(PickerConfig.Filter)
-		.InitialViewType(PickerConfig.InitialAssetViewType)
-		.InitialAssetSelection(PickerConfig.InitialAssetSelection)
-		.ShowBottomToolbar(PickerConfig.bShowBottomToolbar)
-		.OnAssetTagWantsToBeDisplayed(PickerConfig.OnAssetTagWantsToBeDisplayed)
-		.OnGetCustomSourceAssets(PickerConfig.OnGetCustomSourceAssets)
-		.AllowDragging(PickerConfig.bAllowDragging)
-		.CanShowClasses(PickerConfig.bCanShowClasses)
-		.CanShowFolders(PickerConfig.bCanShowFolders)
-		.CanShowReadOnlyFolders(PickerConfig.bCanShowReadOnlyFolders)
-		.ShowPathInColumnView(PickerConfig.bShowPathInColumnView)
-		.ShowTypeInColumnView(PickerConfig.bShowTypeInColumnView)
+		.InitialBackendFilter(Filter)
+		.InitialViewType(EAssetViewType::Tile)
+		.InitialAssetSelection({})
+		.ShowBottomToolbar(false)
+		.AllowDragging(false)
+		.CanShowClasses(false)
+		.CanShowFolders(true)
+		.CanShowReadOnlyFolders(true)
 		.ShowViewOptions(false)
-		.SortByPathInColumnView(PickerConfig.bSortByPathInColumnView)
+		.bShowPathViewFilters(false)
 		.FilterRecursivelyWithBackendFilter(false)
-		.CanShowRealTimeThumbnails(PickerConfig.bCanShowRealTimeThumbnails)
-		.CanShowDevelopersFolder(PickerConfig.bCanShowDevelopersFolder)
-		.ForceShowEngineContent(PickerConfig.bForceShowEngineContent)
-		.ForceShowPluginContent(PickerConfig.bForceShowPluginContent)
+		.CanShowRealTimeThumbnails(true)
+		.CanShowDevelopersFolder(true)
+		.ForceShowEngineContent(true)
+		.ForceShowPluginContent(true)
 		.HighlightedText(TAttribute<FText>(this, &SDMMaterialWizard::GetSearchText))
-		.ThumbnailLabel(PickerConfig.ThumbnailLabel)
-		.AssetShowWarningText(PickerConfig.AssetShowWarningText)
+		.ThumbnailLabel(EThumbnailLabel::AssetName)
 		.AllowFocusOnSync(false)
-		.HiddenColumnNames(PickerConfig.HiddenColumnNames)
-		.CustomColumns(PickerConfig.CustomColumns)
 		.InitialThumbnailSize(EThumbnailSize::Small)
 		.ShowTypeInTileView(false)
 		.TextFilter(TextFilter);
 
+	AssetView->OverrideShowEngineContent();
+	AssetView->OverrideShowPluginContent();
+	AssetView->OverrideShowDeveloperContent();
 	AssetView->RequestSlowFullListRefresh();
 
-	return SNew(SBorder)
-		.BorderImage(FAppStyle::Get().GetBrush("Menu.Background"))
-		.Padding(3.f, 3.f, 3.f, 3.f)
-		.VAlign(EVerticalAlignment::VAlign_Fill)
-		[
-			SNew(SVerticalBox)
+	return AssetView.ToSharedRef();
+}
 
-			+ SVerticalBox::Slot()
-			.AutoHeight()
-			.Padding(0.f, 0.f, 0.f, 3.f)
-			[
-				SAssignNew(AssetSearchBox, SAssetSearchBox)
-				.HintText(NSLOCTEXT("ContentBrowser", "SearchBoxHint", "Search Assets"))
-				.OnTextChanged(this, &SDMMaterialWizard::OnSearchBoxChanged)
-				.OnTextCommitted(this, &SDMMaterialWizard::OnSearchBoxCommitted)
-				.DelayChangeNotificationsWhileTyping(true)
-			]
+TSharedRef<SWidget> SDMMaterialWizard::CreateNewInstance_AcceptButton()
+{
+	using namespace UE::DynamicMaterialDesigner::Private;
 
-			+ SVerticalBox::Slot()
-			.FillHeight(1.f)
-			[
-				AssetView.ToSharedRef()
-			]
+	return SNew(SButton)
+		.IsEnabled(this, &SDMMaterialWizard::NewInstance_CanAccept)
+		.ButtonStyle(FAppStyle::Get(), "PrimaryButton")
+		.ContentPadding(ButtonPadding)
+		.OnClicked(this, &SDMMaterialWizard::NewInstanceAccept_OnClick)
+		[
+			SNew(STextBlock)
+			.TextStyle(FDynamicMaterialEditorStyle::Get(), "RegularFont")
+			.Text(LOCTEXT("Continue", "Continue"))
 		];
 }
 
@@ -458,18 +470,14 @@ TSharedRef<SWidget> SDMMaterialWizard::CreateNewTemplate_AcceptButton()
 {
 	using namespace UE::DynamicMaterialDesigner::Private;
 
-	return SNew(SBox)
-		.HAlign(HAlign_Center)
+	return SNew(SButton)
+		.ButtonStyle(FAppStyle::Get(), "PrimaryButton")
+		.ContentPadding(ButtonPadding)
+		.OnClicked(this, &SDMMaterialWizard::NewTemplateAccept_OnClick)
 		[
-			SNew(SButton)
-			.ButtonStyle(FAppStyle::Get(), "PrimaryButton")
-			.ContentPadding(ButtonPadding)
-			.OnClicked(this, &SDMMaterialWizard::Accept_OnClick)
-			[
-				SNew(STextBlock)
-				.TextStyle(FDynamicMaterialEditorStyle::Get(), "RegularFont")
-				.Text(LOCTEXT("Continue", "Continue"))
-			]
+			SNew(STextBlock)
+			.TextStyle(FDynamicMaterialEditorStyle::Get(), "RegularFont")
+			.Text(LOCTEXT("Continue", "Continue"))
 		];
 }
 
@@ -493,7 +501,7 @@ void SDMMaterialWizard::Preset_OnChange(ECheckBoxState InState, FName InPresetNa
 	}
 }
 
-FReply SDMMaterialWizard::Accept_OnClick()
+FReply SDMMaterialWizard::NewTemplateAccept_OnClick()
 {
 	if (CurrentPreset.IsNone())
 	{
@@ -679,6 +687,11 @@ bool SDMMaterialWizard::ShouldFilterOutAsset(const FAssetData& InAsset) const
 	return false;
 }
 
+bool SDMMaterialWizard::NewInstance_CanAccept() const
+{
+	return AssetView.IsValid() && !AssetView->GetSelectedAssets().IsEmpty();
+}
+
 void SDMMaterialWizard::OnAssetsActivated(TArrayView<const FContentBrowserItem> InSelectedItems, EAssetTypeActivationMethod::Type InActivationMethod)
 {
 	if (InSelectedItems.IsEmpty() || InActivationMethod == EAssetTypeActivationMethod::Previewed)
@@ -710,6 +723,19 @@ void SDMMaterialWizard::OnAssetsActivated(TArrayView<const FContentBrowserItem>
 	SelectTemplate(MaterialModel);
 }
 
+FReply SDMMaterialWizard::NewInstanceAccept_OnClick()
+{
+	if (!AssetView.IsValid())
+	{
+		return FReply::Handled();
+	}
+
+	TArray<FContentBrowserItem> SelectedItems = AssetView->GetSelectedItems();
+	OnAssetsActivated(SelectedItems, EAssetTypeActivationMethod::DoubleClicked);
+
+	return FReply::Handled();
+}
+
 void SDMMaterialWizard::OnEnginePreExit()
 {
 	TextFilter.Reset();
diff --git a/Engine/Plugins/Experimental/DynamicMaterial/Source/DynamicMaterialEditor/Private/UI/Widgets/SDMMaterialWizard.h b/Engine/Plugins/Experimental/DynamicMaterial/Source/DynamicMaterialEditor/Private/UI/Widgets/SDMMaterialWizard.h
--- a/Engine/Plugins/Experimental/DynamicMaterial/Source/DynamicMaterialEditor/Private/UI/Widgets/SDMMaterialWizard.h
+++ b/Engine/Plugins/Experimental/DynamicMaterial/Source/DynamicMaterialEditor/Private/UI/Widgets/SDMMaterialWizard.h
@@ -72,12 +72,15 @@ protected:
 	TSharedRef<SWidget> CreateNewTemplate_ChannelList();
 	TSharedRef<SWidget> CreateNewTemplate_AcceptButton();
 	TSharedRef<SWidget> CreateNewInstanceLayout();
+	TSharedRef<SWidget> CreateNewInstance_SearchBox();
+	TSharedRef<SWidget> CreateNewInstance_Picker();
+	TSharedRef<SWidget> CreateNewInstance_AcceptButton();
 
 	/** Attributes and Events */
 	ECheckBoxState Preset_GetState(FName InPresetName) const;
 	void Preset_OnChange(ECheckBoxState InState, FName InPresetName);
 
-	FReply Accept_OnClick();
+	FReply NewTemplateAccept_OnClick();
 
 	void OnMaterialBuilt(UDynamicMaterialModelBase* InMaterialModel);
 
@@ -95,8 +98,12 @@ protected:
 
 	bool ShouldFilterOutAsset(const FAssetData& InAsset) const;
 
+	bool NewInstance_CanAccept() const;
+
 	void OnAssetsActivated(TArrayView<const FContentBrowserItem> InSelectedItems, EAssetTypeActivationMethod::Type InActivationMethod);
 
+	FReply NewInstanceAccept_OnClick();
+
 	void OnEnginePreExit();
 
 	/** Operations */
