Material Designer
- Mask channel thumbnail should now update when the the source is changed via the base stage texture being changed.
- When built in debug, ctrl+shift+right click will open the stage preview materials in the material editor for debugging.

diff --git a/Engine/Plugins/Experimental/DynamicMaterial/Source/DynamicMaterial/Private/Components/MaterialValues/DMMaterialValueBool.cpp b/Engine/Plugins/Experimental/DynamicMaterial/Source/DynamicMaterial/Private/Components/MaterialValues/DMMaterialValueBool.cpp
--- a/Engine/Plugins/Experimental/DynamicMaterial/Source/DynamicMaterial/Private/Components/MaterialValues/DMMaterialValueBool.cpp
+++ b/Engine/Plugins/Experimental/DynamicMaterial/Source/DynamicMaterial/Private/Components/MaterialValues/DMMaterialValueBool.cpp
@@ -137,7 +137,7 @@ void UDMMaterialValueBool::SetValue(bool InValue)
 
 	Value = InValue;
  
-	OnValueChanged(EDMUpdateType::Value);
+	OnValueChanged(EDMUpdateType::Value | EDMUpdateType::AllowParentUpdate);
 }
  
 #undef LOCTEXT_NAMESPACE
diff --git a/Engine/Plugins/Experimental/DynamicMaterial/Source/DynamicMaterial/Private/Components/MaterialValues/DMMaterialValueColorAtlas.cpp b/Engine/Plugins/Experimental/DynamicMaterial/Source/DynamicMaterial/Private/Components/MaterialValues/DMMaterialValueColorAtlas.cpp
--- a/Engine/Plugins/Experimental/DynamicMaterial/Source/DynamicMaterial/Private/Components/MaterialValues/DMMaterialValueColorAtlas.cpp
+++ b/Engine/Plugins/Experimental/DynamicMaterial/Source/DynamicMaterial/Private/Components/MaterialValues/DMMaterialValueColorAtlas.cpp
@@ -50,7 +50,7 @@ void UDMMaterialValueColorAtlas::SetValue(float InValue)
 
 	Value = InValue;
 
-	OnValueChanged(EDMUpdateType::Value);
+	OnValueChanged(EDMUpdateType::Value | EDMUpdateType::AllowParentUpdate);
 }
 
 void UDMMaterialValueColorAtlas::SetMIDParameter(UMaterialInstanceDynamic* InMID) const
diff --git a/Engine/Plugins/Experimental/DynamicMaterial/Source/DynamicMaterial/Private/Components/MaterialValues/DMMaterialValueFloat1.cpp b/Engine/Plugins/Experimental/DynamicMaterial/Source/DynamicMaterial/Private/Components/MaterialValues/DMMaterialValueFloat1.cpp
--- a/Engine/Plugins/Experimental/DynamicMaterial/Source/DynamicMaterial/Private/Components/MaterialValues/DMMaterialValueFloat1.cpp
+++ b/Engine/Plugins/Experimental/DynamicMaterial/Source/DynamicMaterial/Private/Components/MaterialValues/DMMaterialValueFloat1.cpp
@@ -123,7 +123,7 @@ void UDMMaterialValueFloat1::SetValue(float InValue)
  
 	Value = InValue;
  
-	OnValueChanged(EDMUpdateType::Value);
+	OnValueChanged(EDMUpdateType::Value | EDMUpdateType::AllowParentUpdate);
 }
  
 void UDMMaterialValueFloat1::SetMIDParameter(UMaterialInstanceDynamic* InMID) const
diff --git a/Engine/Plugins/Experimental/DynamicMaterial/Source/DynamicMaterial/Private/Components/MaterialValues/DMMaterialValueFloat2.cpp b/Engine/Plugins/Experimental/DynamicMaterial/Source/DynamicMaterial/Private/Components/MaterialValues/DMMaterialValueFloat2.cpp
--- a/Engine/Plugins/Experimental/DynamicMaterial/Source/DynamicMaterial/Private/Components/MaterialValues/DMMaterialValueFloat2.cpp
+++ b/Engine/Plugins/Experimental/DynamicMaterial/Source/DynamicMaterial/Private/Components/MaterialValues/DMMaterialValueFloat2.cpp
@@ -136,7 +136,7 @@ void UDMMaterialValueFloat2::SetValue(const FVector2D& InValue)
  
 	Value = ValueClamped;
  
-	OnValueChanged(EDMUpdateType::Value);
+	OnValueChanged(EDMUpdateType::Value | EDMUpdateType::AllowParentUpdate);
 }
  
 #if WITH_EDITOR
diff --git a/Engine/Plugins/Experimental/DynamicMaterial/Source/DynamicMaterial/Private/Components/MaterialValues/DMMaterialValueFloat3RGB.cpp b/Engine/Plugins/Experimental/DynamicMaterial/Source/DynamicMaterial/Private/Components/MaterialValues/DMMaterialValueFloat3RGB.cpp
--- a/Engine/Plugins/Experimental/DynamicMaterial/Source/DynamicMaterial/Private/Components/MaterialValues/DMMaterialValueFloat3RGB.cpp
+++ b/Engine/Plugins/Experimental/DynamicMaterial/Source/DynamicMaterial/Private/Components/MaterialValues/DMMaterialValueFloat3RGB.cpp
@@ -136,7 +136,7 @@ void UDMMaterialValueFloat3RGB::SetValue(const FLinearColor& InValue)
 
 	Value = ValueClamped;
 
-	OnValueChanged(EDMUpdateType::Value);
+	OnValueChanged(EDMUpdateType::Value | EDMUpdateType::AllowParentUpdate);
 }
  
 #if WITH_EDITOR
diff --git a/Engine/Plugins/Experimental/DynamicMaterial/Source/DynamicMaterial/Private/Components/MaterialValues/DMMaterialValueFloat3RPY.cpp b/Engine/Plugins/Experimental/DynamicMaterial/Source/DynamicMaterial/Private/Components/MaterialValues/DMMaterialValueFloat3RPY.cpp
--- a/Engine/Plugins/Experimental/DynamicMaterial/Source/DynamicMaterial/Private/Components/MaterialValues/DMMaterialValueFloat3RPY.cpp
+++ b/Engine/Plugins/Experimental/DynamicMaterial/Source/DynamicMaterial/Private/Components/MaterialValues/DMMaterialValueFloat3RPY.cpp
@@ -138,7 +138,7 @@ void UDMMaterialValueFloat3RPY::SetValue(const FRotator& InValue)
  
 	Value = ValueClamped;
  
-	OnValueChanged(EDMUpdateType::Value);
+	OnValueChanged(EDMUpdateType::Value | EDMUpdateType::AllowParentUpdate);
 }
  
 #if WITH_EDITOR
diff --git a/Engine/Plugins/Experimental/DynamicMaterial/Source/DynamicMaterial/Private/Components/MaterialValues/DMMaterialValueFloat3XYZ.cpp b/Engine/Plugins/Experimental/DynamicMaterial/Source/DynamicMaterial/Private/Components/MaterialValues/DMMaterialValueFloat3XYZ.cpp
--- a/Engine/Plugins/Experimental/DynamicMaterial/Source/DynamicMaterial/Private/Components/MaterialValues/DMMaterialValueFloat3XYZ.cpp
+++ b/Engine/Plugins/Experimental/DynamicMaterial/Source/DynamicMaterial/Private/Components/MaterialValues/DMMaterialValueFloat3XYZ.cpp
@@ -135,7 +135,7 @@ void UDMMaterialValueFloat3XYZ::SetValue(const FVector& InValue)
  
 	Value = ValueClamped;
  
-	OnValueChanged(EDMUpdateType::Value);
+	OnValueChanged(EDMUpdateType::Value | EDMUpdateType::AllowParentUpdate);
 }
  
 #if WITH_EDITOR
diff --git a/Engine/Plugins/Experimental/DynamicMaterial/Source/DynamicMaterial/Private/Components/MaterialValues/DMMaterialValueFloat4.cpp b/Engine/Plugins/Experimental/DynamicMaterial/Source/DynamicMaterial/Private/Components/MaterialValues/DMMaterialValueFloat4.cpp
--- a/Engine/Plugins/Experimental/DynamicMaterial/Source/DynamicMaterial/Private/Components/MaterialValues/DMMaterialValueFloat4.cpp
+++ b/Engine/Plugins/Experimental/DynamicMaterial/Source/DynamicMaterial/Private/Components/MaterialValues/DMMaterialValueFloat4.cpp
@@ -138,7 +138,7 @@ void UDMMaterialValueFloat4::SetValue(const FLinearColor& InValue)
  
 	Value = ValueClamped;
  
-	OnValueChanged(EDMUpdateType::Value);
+	OnValueChanged(EDMUpdateType::Value | EDMUpdateType::AllowParentUpdate);
 }
  
 #if WITH_EDITOR
diff --git a/Engine/Plugins/Experimental/DynamicMaterial/Source/DynamicMaterial/Private/Components/MaterialValues/DMMaterialValueTexture.cpp b/Engine/Plugins/Experimental/DynamicMaterial/Source/DynamicMaterial/Private/Components/MaterialValues/DMMaterialValueTexture.cpp
--- a/Engine/Plugins/Experimental/DynamicMaterial/Source/DynamicMaterial/Private/Components/MaterialValues/DMMaterialValueTexture.cpp
+++ b/Engine/Plugins/Experimental/DynamicMaterial/Source/DynamicMaterial/Private/Components/MaterialValues/DMMaterialValueTexture.cpp
@@ -233,7 +233,7 @@ void UDMMaterialValueTexture::SetValue(UTexture* InValue)
 
 	Value = InValue;
 
-	OnValueChanged(EDMUpdateType::Value);
+	OnValueChanged(EDMUpdateType::Value | EDMUpdateType::AllowParentUpdate);
 }
 
 void UDMMaterialValueTexture::SetMIDParameter(UMaterialInstanceDynamic* InMID) const
diff --git a/Engine/Plugins/Experimental/DynamicMaterial/Source/DynamicMaterialEditor/Private/Components/DMMaterialStage.cpp b/Engine/Plugins/Experimental/DynamicMaterial/Source/DynamicMaterialEditor/Private/Components/DMMaterialStage.cpp
--- a/Engine/Plugins/Experimental/DynamicMaterial/Source/DynamicMaterialEditor/Private/Components/DMMaterialStage.cpp
+++ b/Engine/Plugins/Experimental/DynamicMaterial/Source/DynamicMaterialEditor/Private/Components/DMMaterialStage.cpp
@@ -845,7 +845,7 @@ void UDMMaterialStage::UpdateInputMap(int32 InInputIdx, int32 InSourceIndex, int
 
 	RemoveUnusedInputs();
 
-	Source->Update(this, EDMUpdateType::Structure);
+	Source->Update(this, EDMUpdateType::Structure | EDMUpdateType::AllowParentUpdate);
 }
 
 int32 UDMMaterialStage::FindIndex() const
diff --git a/Engine/Plugins/Experimental/DynamicMaterial/Source/DynamicMaterialEditor/Private/Components/MaterialStageExpressions/DMMSETextureSampleBase.cpp b/Engine/Plugins/Experimental/DynamicMaterial/Source/DynamicMaterialEditor/Private/Components/MaterialStageExpressions/DMMSETextureSampleBase.cpp
--- a/Engine/Plugins/Experimental/DynamicMaterial/Source/DynamicMaterialEditor/Private/Components/MaterialStageExpressions/DMMSETextureSampleBase.cpp
+++ b/Engine/Plugins/Experimental/DynamicMaterial/Source/DynamicMaterialEditor/Private/Components/MaterialStageExpressions/DMMSETextureSampleBase.cpp
@@ -409,12 +409,12 @@ void UDMMaterialStageExpressionTextureSampleBase::UpdateMask()
 								{
 									MaskStage->Modify();
 								}
-
-								// Set output to alpha
-								MaskStage->UpdateInputMap(2, MaskStageInputConnections[2].Channels[0].SourceIndex,
-									FDMMaterialStageConnectorChannel::WHOLE_CHANNEL, 1, FDMMaterialStageConnectorChannel::WHOLE_CHANNEL, EDMMaterialPropertyType::None);
 							}
 
+							// Set output to alpha
+							MaskStage->UpdateInputMap(2, MaskStageInputConnections[2].Channels[0].SourceIndex,
+								FDMMaterialStageConnectorChannel::WHOLE_CHANNEL, 1, FDMMaterialStageConnectorChannel::FOURTH_CHANNEL, EDMMaterialPropertyType::None);
+
 							return;
 						}
 					}
diff --git a/Engine/Plugins/Experimental/DynamicMaterial/Source/DynamicMaterialEditor/Private/UI/Widgets/Editor/SlotEditor/SDMMaterialStage.cpp b/Engine/Plugins/Experimental/DynamicMaterial/Source/DynamicMaterialEditor/Private/UI/Widgets/Editor/SlotEditor/SDMMaterialStage.cpp
--- a/Engine/Plugins/Experimental/DynamicMaterial/Source/DynamicMaterialEditor/Private/UI/Widgets/Editor/SlotEditor/SDMMaterialStage.cpp
+++ b/Engine/Plugins/Experimental/DynamicMaterial/Source/DynamicMaterialEditor/Private/UI/Widgets/Editor/SlotEditor/SDMMaterialStage.cpp
@@ -31,6 +31,11 @@
 #include "Widgets/SToolTip.h"
 #include "Widgets/Text/STextBlock.h"
 
+#if UE_BUILD_DEBUG
+#include "AssetToolsModule.h"
+#include "IAssetTools.h"
+#endif
+
 #define LOCTEXT_NAMESPACE "SDMMaterialStage"
 
 void SDMMaterialStage::PrivateRegisterAttributes(FSlateAttributeDescriptor::FInitializer&)
@@ -176,6 +181,23 @@ FReply SDMMaterialStage::OnMouseButtonDown(const FGeometry& InMyGeometry, const
 		return FReply::Handled();
 	}
 
+#if UE_BUILD_DEBUG
+	else if (InMouseEvent.GetEffectingButton() == EKeys::RightMouseButton
+		&& InMouseEvent.IsShiftDown() && InMouseEvent.IsControlDown())
+	{
+		if (PreviewImage.IsValid())
+		{
+			if (UMaterial* PreviewMaterial = PreviewImage->GetPreviewMaterial())
+			{
+				IAssetTools& AssetTools = FModuleManager::LoadModuleChecked<FAssetToolsModule>("AssetTools").Get();
+				AssetTools.OpenEditorForAssets({PreviewMaterial});
+
+				return FReply::Handled();
+			}
+		}
+	}
+#endif
+
 	return FReply::Unhandled();
 }
 
