Material Designer
- More stringent layer type checks are done when adding default stage layers.

diff --git a/Engine/Plugins/Experimental/DynamicMaterial/Source/DynamicMaterialEditor/Private/Components/DMMaterialProperty.cpp b/Engine/Plugins/Experimental/DynamicMaterial/Source/DynamicMaterialEditor/Private/Components/DMMaterialProperty.cpp
--- a/Engine/Plugins/Experimental/DynamicMaterial/Source/DynamicMaterialEditor/Private/Components/DMMaterialProperty.cpp
+++ b/Engine/Plugins/Experimental/DynamicMaterial/Source/DynamicMaterialEditor/Private/Components/DMMaterialProperty.cpp
@@ -175,6 +175,11 @@ void UDMMaterialProperty::AddDefaultBaseStage(UDMMaterialLayerObject* InLayer)
 		return;
 	}
 
+	if (InLayer->GetMaterialProperty() == EDMMaterialPropertyType::None)
+	{
+		InLayer->SetMaterialProperty(MaterialProperty);
+	}
+
 	UDMMaterialStage* DefaultStage = UDMMaterialStageBlend::CreateStage(UDMMaterialStageBlendNormal::StaticClass());
 	check(DefaultStage);
 
diff --git a/Engine/Plugins/Experimental/DynamicMaterial/Source/DynamicMaterialEditor/Private/Components/DMMaterialStageBlend.cpp b/Engine/Plugins/Experimental/DynamicMaterial/Source/DynamicMaterialEditor/Private/Components/DMMaterialStageBlend.cpp
--- a/Engine/Plugins/Experimental/DynamicMaterial/Source/DynamicMaterialEditor/Private/Components/DMMaterialStageBlend.cpp
+++ b/Engine/Plugins/Experimental/DynamicMaterial/Source/DynamicMaterialEditor/Private/Components/DMMaterialStageBlend.cpp
@@ -488,19 +488,22 @@ void UDMMaterialStageBlend::AddDefaultInput(int32 InInputIndex) const
 			}
 			else
 			{
-				EDMMaterialPropertyType DefaultProperty = EDMMaterialPropertyType::None;
+				EDMMaterialPropertyType DefaultProperty = StageProperty;
 
-				if (UDMMaterialSlot* Slot = Layer->GetSlot())
+				if (DefaultProperty != EDMMaterialPropertyType::None)
 				{
-					if (UDynamicMaterialModelEditorOnlyData* ModelEditorOnlyData = Slot->GetMaterialModelEditorOnlyData())
+					if (UDMMaterialSlot* Slot = Layer->GetSlot())
 					{
-						if (ModelEditorOnlyData->GetSlotForMaterialProperty(EDMMaterialPropertyType::BaseColor))
+						if (UDynamicMaterialModelEditorOnlyData* ModelEditorOnlyData = Slot->GetMaterialModelEditorOnlyData())
 						{
-							DefaultProperty = EDMMaterialPropertyType::BaseColor;
-						}
-						else if (ModelEditorOnlyData->GetSlotForMaterialProperty(EDMMaterialPropertyType::EmissiveColor))
-						{
-							DefaultProperty = EDMMaterialPropertyType::EmissiveColor;
+							if (ModelEditorOnlyData->GetSlotForMaterialProperty(EDMMaterialPropertyType::BaseColor))
+							{
+								DefaultProperty = EDMMaterialPropertyType::BaseColor;
+							}
+							else if (ModelEditorOnlyData->GetSlotForMaterialProperty(EDMMaterialPropertyType::EmissiveColor))
+							{
+								DefaultProperty = EDMMaterialPropertyType::EmissiveColor;
+							}
 						}
 					}
 				}
