Material Designer
- The command list for slot operations is now present.

diff --git a/Engine/Plugins/Experimental/DynamicMaterial/Source/DynamicMaterialEditor/Private/UI/Menus/DMMaterialSlotLayerMenus.cpp b/Engine/Plugins/Experimental/DynamicMaterial/Source/DynamicMaterialEditor/Private/UI/Menus/DMMaterialSlotLayerMenus.cpp
--- a/Engine/Plugins/Experimental/DynamicMaterial/Source/DynamicMaterialEditor/Private/UI/Menus/DMMaterialSlotLayerMenus.cpp
+++ b/Engine/Plugins/Experimental/DynamicMaterial/Source/DynamicMaterialEditor/Private/UI/Menus/DMMaterialSlotLayerMenus.cpp
@@ -71,6 +71,11 @@ TSharedRef<SWidget> FDMMaterialSlotLayerMenus::GenerateSlotLayerMenu(const TShar
 
 	FToolMenuContext MenuContext(UDMMenuContext::CreateLayer(InSlotWidget->GetEditorWidget(), InLayer));
 
+	if (TSharedPtr<SDMMaterialEditor> EditorWidget = InSlotWidget->GetEditorWidget())
+	{
+		MenuContext.AppendCommandList(EditorWidget->GetCommandList());
+	}
+
 	return UToolMenus::Get()->GenerateWidget(SlotLayerMenuName, MenuContext);
 }
 
