Material Designer
- Now correctly uses the world subsystem in all cases to set materials.
- Separated out the model and material function libraries.

diff --git a/Engine/Plugins/Experimental/DynamicMaterial/Source/DynamicMaterialEditor/Private/DMObjectMaterialProperty.cpp b/Engine/Plugins/Experimental/DynamicMaterial/Source/DynamicMaterialEditor/Private/DMObjectMaterialProperty.cpp
--- a/Engine/Plugins/Experimental/DynamicMaterial/Source/DynamicMaterialEditor/Private/DMObjectMaterialProperty.cpp
+++ b/Engine/Plugins/Experimental/DynamicMaterial/Source/DynamicMaterialEditor/Private/DMObjectMaterialProperty.cpp
@@ -113,7 +113,7 @@ UDynamicMaterialInstance* FDMObjectMaterialProperty::GetMaterial() const
 	return nullptr;
 }
 
-void FDMObjectMaterialProperty::SetMaterial(UDynamicMaterialInstance* DynamicMaterial)
+void FDMObjectMaterialProperty::SetMaterial(UDynamicMaterialInstance* InDynamicMaterial)
 {
 	UObject* Outer = OuterWeak.Get();
 
@@ -136,7 +136,7 @@ void FDMObjectMaterialProperty::SetMaterial(UDynamicMaterialInstance* DynamicMat
 
 				if (ArrayHelper.IsValidIndex(Index))
 				{
-					*reinterpret_cast<UMaterialInterface**>(ArrayHelper.GetRawPtr(Index)) = DynamicMaterial;
+					*reinterpret_cast<UMaterialInterface**>(ArrayHelper.GetRawPtr(Index)) = InDynamicMaterial;
 				}
 			}
 		}
@@ -146,7 +146,7 @@ void FDMObjectMaterialProperty::SetMaterial(UDynamicMaterialInstance* DynamicMat
 
 			if (ObjectProperty && ObjectProperty->PropertyClass->IsChildOf(UMaterialInterface::StaticClass()))
 			{
-				Property->SetValue_InContainer(Outer, &DynamicMaterial);
+				Property->SetValue_InContainer(Outer, &InDynamicMaterial);
 			}
 		}
 
@@ -159,7 +159,7 @@ void FDMObjectMaterialProperty::SetMaterial(UDynamicMaterialInstance* DynamicMat
 		{
 			if (Index >= 0 && Index < Component->GetNumMaterials())
 			{
-				Component->SetMaterial(Index, DynamicMaterial);
+				Component->SetMaterial(Index, InDynamicMaterial);
 			}
 		}
 	}
diff --git a/Engine/Plugins/Experimental/DynamicMaterial/Source/DynamicMaterialEditor/Private/UI/Widgets/Editor/SDMToolBar.cpp b/Engine/Plugins/Experimental/DynamicMaterial/Source/DynamicMaterialEditor/Private/UI/Widgets/Editor/SDMToolBar.cpp
--- a/Engine/Plugins/Experimental/DynamicMaterial/Source/DynamicMaterialEditor/Private/UI/Widgets/Editor/SDMToolBar.cpp
+++ b/Engine/Plugins/Experimental/DynamicMaterial/Source/DynamicMaterialEditor/Private/UI/Widgets/Editor/SDMToolBar.cpp
@@ -22,6 +22,7 @@
 #include "UI/Widgets/SDMMaterialDesigner.h"
 #include "UI/Widgets/SDMMaterialEditor.h"
 #include "UObject/Object.h"
+#include "Utils/DMMaterialInstanceFunctionLibrary.h"
 #include "Utils/DMMaterialModelFunctionLibrary.h"
 #include "Widgets/Images/SImage.h"
 #include "Widgets/Input/SButton.h"
@@ -305,7 +306,7 @@ void SDMToolBar::SetActorPropertySelected(AActor* InActor)
 		ActorNameWidget->SetText(GetActorName());
 		ActorRowWidget->SetVisibility(EVisibility::Visible);
 
-		TArray<FDMObjectMaterialProperty> ActorProperties = UDMMaterialModelFunctionLibrary::GetActorMaterialProperties(InActor);
+		TArray<FDMObjectMaterialProperty> ActorProperties = UDMMaterialInstanceFunctionLibrary::GetActorMaterialProperties(InActor);
 		UDynamicMaterialModelBase* MaterialModelBase = GetMaterialModelBase();
 
 		const int32 ActorPropertyCount = ActorProperties.Num();
@@ -487,7 +488,7 @@ void SDMToolBar::OnMaterialSlotChanged(TSharedPtr<FDMObjectMaterialProperty> InS
 	}
 	else if (InSelectedSlot->GetOuter())
 	{
-		if (UDynamicMaterialModel* NewModel = UDMMaterialModelFunctionLibrary::CreateMaterialInObject(*InSelectedSlot.Get()))
+		if (UDynamicMaterialModel* NewModel = UDMMaterialInstanceFunctionLibrary::CreateMaterialInObject(*InSelectedSlot.Get()))
 		{
 			DesignerWidget->OpenObjectMaterialProperty(*InSelectedSlot);
 		}
@@ -676,27 +677,18 @@ FReply SDMToolBar::OnUseClicked()
 		return FReply::Handled();
 	}
 
-	TSharedPtr<FDMObjectMaterialProperty> CurrentActorProperty = ActorMaterialProperties[SelectedMaterialElementIndex];
+	FDMObjectMaterialProperty& CurrentActorProperty = *ActorMaterialProperties[SelectedMaterialElementIndex];
 
-	if (DMSubsystem && DMSubsystem->GetMaterialValueSetterDelegate().IsBound())
+	if (!UDMMaterialInstanceFunctionLibrary::SetMaterialInObject(CurrentActorProperty, SelectedInstance))
 	{
-		if (!DMSubsystem->ExecuteIsValidDelegate(CurrentModelBase))
-		{
-			return FReply::Handled();
-		}
-
-		DMSubsystem->ExecuteMaterialValueSetterDelegate(*CurrentActorProperty, SelectedInstance);
-	}
-	else
-	{
-		ActorMaterialProperties[SelectedMaterialElementIndex]->SetMaterial(SelectedInstance);
+		return FReply::Handled();
 	}
 
 	if (TSharedPtr<SDMMaterialEditor> EditorWidget = GetEditorWidget())
 	{
 		if (TSharedPtr<SDMMaterialDesigner> DesignerWidget = EditorWidget->GetDesignerWidget())
 		{
-			DesignerWidget->OpenObjectMaterialProperty(*CurrentActorProperty);
+			DesignerWidget->OpenObjectMaterialProperty(CurrentActorProperty);
 		}
 	}
 
@@ -915,29 +907,8 @@ FReply SDMToolBar::OnConvertToEditableClicked()
 	// If it was in an actor, set it on the actor
 	if (NewInstance && CurrentActorProperty.IsValid())
 	{
-		if (DMSubsystem && DMSubsystem->GetMaterialValueSetterDelegate().IsBound())
-		{
-			if (DMSubsystem->ExecuteIsValidDelegate(CurrentModelDynamic))
-			{
-				DMSubsystem->ExecuteMaterialValueSetterDelegate(*CurrentActorProperty, NewInstance);
-			}
-			else
-			{
-				UE_LOG(LogDynamicMaterialEditor, Error, TEXT("Asset is not valid for current world, not assigning to actor."));
-			}
-		}
-		else
-		{
-			ActorMaterialProperties[SelectedMaterialElementIndex]->SetMaterial(NewInstance);
-		}
-
-		if (TSharedPtr<SDMMaterialEditor> EditorWidget = GetEditorWidget())
-		{
-			if (TSharedPtr<SDMMaterialDesigner> DesignerWidget = EditorWidget->GetDesignerWidget())
-			{
-				DesignerWidget->OpenObjectMaterialProperty(*CurrentActorProperty);
-			}
-		}
+		// Setting it on the actor will automatically open it if the actor property is currently active.
+		UDMMaterialInstanceFunctionLibrary::SetMaterialInObject(*CurrentActorProperty, NewInstance);
 	}
 	else
 	{
diff --git a/Engine/Plugins/Experimental/DynamicMaterial/Source/DynamicMaterialEditor/Private/UI/Widgets/SDMActorMaterialSelector.cpp b/Engine/Plugins/Experimental/DynamicMaterial/Source/DynamicMaterialEditor/Private/UI/Widgets/SDMActorMaterialSelector.cpp
--- a/Engine/Plugins/Experimental/DynamicMaterial/Source/DynamicMaterialEditor/Private/UI/Widgets/SDMActorMaterialSelector.cpp
+++ b/Engine/Plugins/Experimental/DynamicMaterial/Source/DynamicMaterialEditor/Private/UI/Widgets/SDMActorMaterialSelector.cpp
@@ -13,7 +13,7 @@
 #include "Styling/StyleColors.h"
 #include "ThumbnailRendering/ThumbnailManager.h"
 #include "UI/Widgets/SDMMaterialDesigner.h"
-#include "Utils/DMMaterialModelFunctionLibrary.h"
+#include "Utils/DMMaterialInstanceFunctionLibrary.h"
 #include "Widgets/Input/SButton.h"
 #include "Widgets/Layout/SBox.h"
 #include "Widgets/Layout/SScrollBox.h"
@@ -213,7 +213,7 @@ FReply SDMActorMaterialSelector::OnCreateMaterialButtonClicked(int32 InActorProp
 		return FReply::Handled();
 	}
 
-	UDynamicMaterialModel* NewMaterialModel = UDMMaterialModelFunctionLibrary::CreateMaterialInObject(ActorProperties[InActorPropertyIndex]);
+	UDynamicMaterialModel* NewMaterialModel = UDMMaterialInstanceFunctionLibrary::CreateMaterialInObject(ActorProperties[InActorPropertyIndex]);
 
 	if (NewMaterialModel)
 	{
diff --git a/Engine/Plugins/Experimental/DynamicMaterial/Source/DynamicMaterialEditor/Private/UI/Widgets/SDMMaterialDesigner.cpp b/Engine/Plugins/Experimental/DynamicMaterial/Source/DynamicMaterialEditor/Private/UI/Widgets/SDMMaterialDesigner.cpp
--- a/Engine/Plugins/Experimental/DynamicMaterial/Source/DynamicMaterialEditor/Private/UI/Widgets/SDMMaterialDesigner.cpp
+++ b/Engine/Plugins/Experimental/DynamicMaterial/Source/DynamicMaterialEditor/Private/UI/Widgets/SDMMaterialDesigner.cpp
@@ -16,14 +16,13 @@
 #include "SAssetDropTarget.h"
 #include "UI/Utils/DMDropTargetPrivateSetter.h"
 #include "UI/Widgets/Editor/EditorLayouts/SDMMaterialEditor_Left.h"
-#include "UI/Widgets/Editor/EditorLayouts/SDMMaterialEditor_LeftSlim.h"
-#include "UI/Widgets/Editor/EditorLayouts/SDMMaterialEditor_TopHorizontal.h"
 #include "UI/Widgets/Editor/EditorLayouts/SDMMaterialEditor_TopSlim.h"
 #include "UI/Widgets/Editor/EditorLayouts/SDMMaterialEditor_TopVertical.h"
 #include "UI/Widgets/SDMActorMaterialSelector.h"
 #include "UI/Widgets/SDMMaterialEditor.h"
 #include "UI/Widgets/SDMMaterialSelectPrompt.h"
 #include "UI/Widgets/SDMMaterialWizard.h"
+#include "Utils/DMMaterialInstanceFunctionLibrary.h"
 #include "Utils/DMMaterialModelFunctionLibrary.h"
 #include "Widgets/SNullWidget.h"
 
@@ -211,7 +210,7 @@ void SDMMaterialDesigner::OpenActor_Internal(AActor* InActor)
 {
 	SetWidget(SNullWidget::NullWidget, /* Include Drop Target */ true);
 
-	TArray<FDMObjectMaterialProperty> ActorProperties = UDMMaterialModelFunctionLibrary::GetActorMaterialProperties(InActor);
+	TArray<FDMObjectMaterialProperty> ActorProperties = UDMMaterialInstanceFunctionLibrary::GetActorMaterialProperties(InActor);
 
 	if (ActorProperties.IsEmpty())
 	{
diff --git a/Engine/Plugins/Experimental/DynamicMaterial/Source/DynamicMaterialEditor/Private/UI/Widgets/SDMMaterialWizard.cpp b/Engine/Plugins/Experimental/DynamicMaterial/Source/DynamicMaterialEditor/Private/UI/Widgets/SDMMaterialWizard.cpp
--- a/Engine/Plugins/Experimental/DynamicMaterial/Source/DynamicMaterialEditor/Private/UI/Widgets/SDMMaterialWizard.cpp
+++ b/Engine/Plugins/Experimental/DynamicMaterial/Source/DynamicMaterialEditor/Private/UI/Widgets/SDMMaterialWizard.cpp
@@ -20,6 +20,7 @@
 #include "Styling/SlateTypes.h"
 #include "Styling/StyleColors.h"
 #include "UI/Widgets/SDMMaterialDesigner.h"
+#include "Utils/DMMaterialInstanceFunctionLibrary.h"
 #include "Utils/DMMaterialModelFunctionLibrary.h"
 #include "Utils/DMPrivate.h"
 #include "Widgets/Input/SButton.h"
@@ -853,7 +854,11 @@ void SDMMaterialWizard::CreateNewDynamicInstanceInActor(UDynamicMaterialModel* I
 		return;
 	}
 
-	InMaterialObjectProperty.SetMaterial(NewInstance);
+	if (!UDMMaterialInstanceFunctionLibrary::SetMaterialInObject(InMaterialObjectProperty, NewInstance))
+	{
+		UE::DynamicMaterialEditor::Private::LogError(TEXT("Failed to set material instance on object."));
+		return;
+	}
 
 	DesignerWidget->OpenObjectMaterialProperty(InMaterialObjectProperty);
 }
@@ -927,7 +932,11 @@ void SDMMaterialWizard::CreateTemplateMaterialInActor(FName InChannelList, FDMOb
 		EditorOnlyData->OnWizardComplete();
 	}
 
-	InMaterialObjectProperty.SetMaterial(NewInstance);
+	if (!UDMMaterialInstanceFunctionLibrary::SetMaterialInObject(InMaterialObjectProperty, NewInstance))
+	{
+		UE::DynamicMaterialEditor::Private::LogError(TEXT("Failed to set material template on object."));
+		return;
+	}
 
 	DesignerWidget->OpenMaterialModelBase(NewInstance->GetMaterialModelBase());
 }
diff --git a/Engine/Plugins/Experimental/DynamicMaterial/Source/DynamicMaterialEditor/Private/Utils/DMMaterialInstanceFunctionLibrary.cpp b/Engine/Plugins/Experimental/DynamicMaterial/Source/DynamicMaterialEditor/Private/Utils/DMMaterialInstanceFunctionLibrary.cpp
new file mode 100644
--- /dev/null
+++ b/Engine/Plugins/Experimental/DynamicMaterial/Source/DynamicMaterialEditor/Private/Utils/DMMaterialInstanceFunctionLibrary.cpp
@@ -0,0 +1,130 @@
+// Copyright Epic Games, Inc. All Rights Reserved.
+
+#include "Utils/DMMaterialInstanceFunctionLibrary.h"
+
+#include "AssetRegistry/AssetRegistryModule.h"
+#include "Components/DMMaterialLayer.h"
+#include "Components/PrimitiveComponent.h"
+#include "ContentBrowserModule.h"
+#include "DMObjectMaterialProperty.h"
+#include "DMWorldSubsystem.h"
+#include "DynamicMaterialEditorModule.h"
+#include "Engine/World.h"
+#include "GameFramework/Actor.h"
+#include "IContentBrowserSingleton.h"
+#include "Material/DynamicMaterialInstance.h"
+#include "Material/DynamicMaterialInstanceFactory.h"
+#include "Materials/Material.h"
+#include "Model/DynamicMaterialModel.h"
+#include "Model/DynamicMaterialModelDynamic.h"
+#include "Model/DynamicMaterialModelEditorOnlyData.h"
+#include "Model/IDynamicMaterialModelEditorOnlyDataInterface.h"
+#include "UObject/Package.h"
+#include "UObject/UObjectGlobals.h"
+#include "Utils/DMPrivate.h"
+
+#define LOCTEXT_NAMESPACE "DMMaterialInstanceFunctionLibrary"
+
+TArray<FDMObjectMaterialProperty> UDMMaterialInstanceFunctionLibrary::GetActorMaterialProperties(AActor* InActor)
+{
+	TArray<FDMObjectMaterialProperty> ActorProperties;
+
+	if (!IsValid(InActor))
+	{
+		return ActorProperties;
+	}
+
+	FDMGetObjectMaterialPropertiesDelegate PropertyGenerator = FDynamicMaterialEditorModule::GetCustomMaterialPropertyGenerator(InActor->GetClass());
+
+	if (PropertyGenerator.IsBound())
+	{
+		ActorProperties = PropertyGenerator.Execute(InActor);
+
+		if (!ActorProperties.IsEmpty())
+		{
+			return ActorProperties;
+		}
+	}
+
+	InActor->ForEachComponent<UPrimitiveComponent>(/* Include Child Actors */ false, [&ActorProperties](UPrimitiveComponent* InComp)
+		{
+			for (int32 MaterialIdx = 0; MaterialIdx < InComp->GetNumMaterials(); ++MaterialIdx)
+			{
+				ActorProperties.Add({InComp, MaterialIdx});
+			}
+		});
+
+	return ActorProperties;
+}
+
+bool UDMMaterialInstanceFunctionLibrary::SetMaterialInObject(FDMObjectMaterialProperty& InMaterialProperty,
+	UDynamicMaterialInstance* InInstance)
+{
+	if (!InMaterialProperty.IsValid())
+	{
+		return false;
+	}
+
+	bool bSubsystemTakenOver = false;
+
+	if (UObject* const Outer = InMaterialProperty.GetOuter())
+	{
+		if (const UWorld* const World = Outer->GetWorld())
+		{
+			if (IsValid(World))
+			{
+				if (UDMWorldSubsystem* const WorldSubsystem = World->GetSubsystem<UDMWorldSubsystem>())
+				{
+					bSubsystemTakenOver = WorldSubsystem->ExecuteMaterialValueSetterDelegate(InMaterialProperty, InInstance);
+				}
+			}
+		}
+	}
+
+	if (!bSubsystemTakenOver)
+	{
+		InMaterialProperty.SetMaterial(InInstance);
+	}
+
+	return InMaterialProperty.GetMaterial() == InInstance;
+}
+
+UDynamicMaterialModel* UDMMaterialInstanceFunctionLibrary::CreateMaterialInObject(FDMObjectMaterialProperty& InMaterialProperty)
+{
+	if (!InMaterialProperty.IsValid())
+	{
+		return nullptr;
+	}
+
+	UObject* const Outer = InMaterialProperty.GetOuter();
+
+	if (!Outer)
+	{
+		return nullptr;
+	}
+
+	UDynamicMaterialInstanceFactory* const InstanceFactory = NewObject<UDynamicMaterialInstanceFactory>();
+	check(InstanceFactory);
+
+	UDynamicMaterialInstance* const NewInstance = Cast<UDynamicMaterialInstance>(InstanceFactory->FactoryCreateNew(UDynamicMaterialInstance::StaticClass(),
+		Outer, NAME_None, RF_Transactional, nullptr, GWarn));
+	check(NewInstance);
+
+	UDynamicMaterialModel* MaterialModel = NewInstance->GetMaterialModel();
+
+	if (!MaterialModel)
+	{
+		return nullptr;
+	}
+
+	SetMaterialInObject(InMaterialProperty, NewInstance);
+
+	if (IDynamicMaterialModelEditorOnlyDataInterface* EditorOnlyData = MaterialModel->GetEditorOnlyData())
+	{
+		EditorOnlyData->RequestMaterialBuild();
+	}
+
+	return MaterialModel;
+}
+
+#undef LOCTEXT_NAMESPACE
diff --git a/Engine/Plugins/Experimental/DynamicMaterial/Source/DynamicMaterialEditor/Private/Utils/DMMaterialModelFunctionLibrary.cpp b/Engine/Plugins/Experimental/DynamicMaterial/Source/DynamicMaterialEditor/Private/Utils/DMMaterialModelFunctionLibrary.cpp
--- a/Engine/Plugins/Experimental/DynamicMaterial/Source/DynamicMaterialEditor/Private/Utils/DMMaterialModelFunctionLibrary.cpp
+++ b/Engine/Plugins/Experimental/DynamicMaterial/Source/DynamicMaterialEditor/Private/Utils/DMMaterialModelFunctionLibrary.cpp
@@ -7,7 +7,6 @@
 #include "Components/DMMaterialLayer.h"
 #include "Components/PrimitiveComponent.h"
 #include "ContentBrowserModule.h"
-#include "DMObjectMaterialProperty.h"
 #include "DMWorldSubsystem.h"
 #include "DynamicMaterialEditorModule.h"
 #include "Engine/Texture.h"
@@ -29,85 +28,6 @@
 
 #define LOCTEXT_NAMESPACE "DMMaterialModelFunctionLibrary"
 
-TArray<FDMObjectMaterialProperty> UDMMaterialModelFunctionLibrary::GetActorMaterialProperties(AActor* InActor)
-{
-	TArray<FDMObjectMaterialProperty> ActorProperties;
-
-	if (!IsValid(InActor))
-	{
-		return ActorProperties;
-	}
-
-	FDMGetObjectMaterialPropertiesDelegate PropertyGenerator = FDynamicMaterialEditorModule::GetCustomMaterialPropertyGenerator(InActor->GetClass());
-
-	if (PropertyGenerator.IsBound())
-	{
-		ActorProperties = PropertyGenerator.Execute(InActor);
-
-		if (!ActorProperties.IsEmpty())
-		{
-			return ActorProperties;
-		}
-	}
-
-	InActor->ForEachComponent<UPrimitiveComponent>(false, [&ActorProperties](UPrimitiveComponent* InComp)
-		{
-			for (int32 MaterialIdx = 0; MaterialIdx < InComp->GetNumMaterials(); ++MaterialIdx)
-			{
-				ActorProperties.Add({InComp, MaterialIdx});
-			}
-		});
-
-	return ActorProperties;
-}
-
-UDynamicMaterialModel* UDMMaterialModelFunctionLibrary::CreateMaterialInObject(FDMObjectMaterialProperty& InMaterialProperty)
-{
-	if (!InMaterialProperty.IsValid())
-	{
-		return nullptr;
-	}
-
-	UObject* const Outer = InMaterialProperty.GetOuter();
-
-	UDynamicMaterialInstanceFactory* const InstanceFactory = NewObject<UDynamicMaterialInstanceFactory>();
-	check(InstanceFactory);
-
-	UDynamicMaterialInstance* const NewInstance = Cast<UDynamicMaterialInstance>(InstanceFactory->FactoryCreateNew(UDynamicMaterialInstance::StaticClass(),
-		Outer, NAME_None, RF_Transactional, nullptr, GWarn));
-	check(NewInstance);
-
-	bool bSubsystemTakenOver = false;
-
-	if (const UWorld* const World = Outer->GetWorld())
-	{
-		if (IsValid(World))
-		{
-			if (UDMWorldSubsystem* const WorldSubsystem = World->GetSubsystem<UDMWorldSubsystem>())
-			{
-				bSubsystemTakenOver = WorldSubsystem->ExecuteMaterialValueSetterDelegate(InMaterialProperty, NewInstance);
-			}
-		}
-	}
-
-	if (!bSubsystemTakenOver)
-	{
-		InMaterialProperty.SetMaterial(NewInstance);
-	}
-
-	if (UDynamicMaterialModel* MaterialModel = NewInstance->GetMaterialModel())
-	{
-		if (IDynamicMaterialModelEditorOnlyDataInterface* EditorOnlyData = MaterialModel->GetEditorOnlyData())
-		{
-			EditorOnlyData->RequestMaterialBuild();
-		}
-
-		return MaterialModel;
-	}
-
-	return nullptr;
-}
-
 UDynamicMaterialInstance* UDMMaterialModelFunctionLibrary::ExportMaterial(UDynamicMaterialModelBase* InMaterialModelBase)
 {
 	if (!IsValid(InMaterialModelBase))
diff --git a/Engine/Plugins/Experimental/DynamicMaterial/Source/DynamicMaterialEditor/Public/DMObjectMaterialProperty.h b/Engine/Plugins/Experimental/DynamicMaterial/Source/DynamicMaterialEditor/Public/DMObjectMaterialProperty.h
--- a/Engine/Plugins/Experimental/DynamicMaterial/Source/DynamicMaterialEditor/Public/DMObjectMaterialProperty.h
+++ b/Engine/Plugins/Experimental/DynamicMaterial/Source/DynamicMaterialEditor/Public/DMObjectMaterialProperty.h
@@ -23,6 +23,8 @@ struct FDMObjectMaterialProperty
 {
 	GENERATED_BODY()
 
+	friend class UDMMaterialInstanceFunctionLibrary;
+
 	FDMObjectMaterialProperty();
 
 	/** UPrimitiveComponent Material Index */
@@ -41,8 +43,6 @@ struct FDMObjectMaterialProperty
 
 	DYNAMICMATERIALEDITOR_API UDynamicMaterialInstance* GetMaterial() const;
 
-	DYNAMICMATERIALEDITOR_API void SetMaterial(UDynamicMaterialInstance* DynamicMaterial);
-
 	DYNAMICMATERIALEDITOR_API bool IsValid() const;
 
 	DYNAMICMATERIALEDITOR_API FText GetPropertyName(bool bInIgnoreNewStatus) const;
@@ -80,4 +80,6 @@ protected:
 
 	/** Component or array property index. */
 	int32 Index = INDEX_NONE;
+
+	void SetMaterial(UDynamicMaterialInstance* InDynamicMaterial);
 };
diff --git a/Engine/Plugins/Experimental/DynamicMaterial/Source/DynamicMaterialEditor/Public/Utils/DMMaterialInstanceFunctionLibrary.h b/Engine/Plugins/Experimental/DynamicMaterial/Source/DynamicMaterialEditor/Public/Utils/DMMaterialInstanceFunctionLibrary.h
new file mode 100644
--- /dev/null
+++ b/Engine/Plugins/Experimental/DynamicMaterial/Source/DynamicMaterialEditor/Public/Utils/DMMaterialInstanceFunctionLibrary.h
@@ -0,0 +1,33 @@
+// Copyright Epic Games, Inc. All Rights Reserved.
+
+#pragma once
+
+#include "Kismet/BlueprintFunctionLibrary.h"
+
+#include "Containers/ContainersFwd.h"
+#include "DMEDefs.h"
+
+#include "DMMaterialInstanceFunctionLibrary.generated.h"
+
+class AActor;
+class UDynamicMaterialInstance;
+class UDynamicMaterialModel;
+class UDynamicMaterialModelBase;
+class UMaterial;
+struct FDMObjectMaterialProperty;
+
+/**
+ * Material Instance Function Library
+ */
+UCLASS()
+class UDMMaterialInstanceFunctionLibrary : public UBlueprintFunctionLibrary
+{
+	GENERATED_BODY()
+
+public:
+	DYNAMICMATERIALEDITOR_API static TArray<FDMObjectMaterialProperty> GetActorMaterialProperties(AActor* InActor);
+
+	DYNAMICMATERIALEDITOR_API static bool SetMaterialInObject(FDMObjectMaterialProperty& InMaterialProperty, UDynamicMaterialInstance* InInstance);
+
+	DYNAMICMATERIALEDITOR_API static UDynamicMaterialModel* CreateMaterialInObject(FDMObjectMaterialProperty& InMaterialProperty);
+};
diff --git a/Engine/Plugins/Experimental/DynamicMaterial/Source/DynamicMaterialEditor/Public/Utils/DMMaterialModelFunctionLibrary.h b/Engine/Plugins/Experimental/DynamicMaterial/Source/DynamicMaterialEditor/Public/Utils/DMMaterialModelFunctionLibrary.h
--- a/Engine/Plugins/Experimental/DynamicMaterial/Source/DynamicMaterialEditor/Public/Utils/DMMaterialModelFunctionLibrary.h
+++ b/Engine/Plugins/Experimental/DynamicMaterial/Source/DynamicMaterialEditor/Public/Utils/DMMaterialModelFunctionLibrary.h
@@ -8,12 +8,10 @@
 
 #include "DMMaterialModelFunctionLibrary.generated.h"
 
-class AActor;
 class FString;
 class UDynamicMaterialModelBase;
 class UDynamicMaterialModelDynamic;
 class UMaterial;
-struct FDMObjectMaterialProperty;
 
 /**
  * Material / Model Function Library
@@ -24,10 +22,6 @@ class UDMMaterialModelFunctionLibrary : public UBlueprintFunctionLibrary
 	GENERATED_BODY()
 
 public:
-	DYNAMICMATERIALEDITOR_API static TArray<FDMObjectMaterialProperty> GetActorMaterialProperties(AActor* InActor);
-
-	DYNAMICMATERIALEDITOR_API static UDynamicMaterialModel* CreateMaterialInObject(FDMObjectMaterialProperty& InMaterialProperty);
-
 	DYNAMICMATERIALEDITOR_API static UDynamicMaterialInstance* ExportMaterial(UDynamicMaterialModelBase* InMaterialModelBase);
 
 	DYNAMICMATERIALEDITOR_API static UDynamicMaterialInstance* ExportMaterial(UDynamicMaterialModelBase* InMaterialModel, const FString& InSavePath);
