Material Designer
- When applying texture sets to MD assets, the channels are now enabled if they are disabled.

diff --git a/Engine/Plugins/Experimental/DynamicMaterial/Source/DynamicMaterialEditor/Private/Model/DynamicMaterialModelEditorOnlyData.cpp b/Engine/Plugins/Experimental/DynamicMaterial/Source/DynamicMaterialEditor/Private/Model/DynamicMaterialModelEditorOnlyData.cpp
--- a/Engine/Plugins/Experimental/DynamicMaterial/Source/DynamicMaterialEditor/Private/Model/DynamicMaterialModelEditorOnlyData.cpp
+++ b/Engine/Plugins/Experimental/DynamicMaterial/Source/DynamicMaterialEditor/Private/Model/DynamicMaterialModelEditorOnlyData.cpp
@@ -638,6 +638,11 @@ bool UDynamicMaterialModelEditorOnlyData::AddTextureSet(UDMTextureSet* InTexture
 
 	for (const TPair<EDMTextureSetMaterialProperty, FDMMaterialTexture>& MaterialTexture : InTextureSet->GetTextures())
 	{
+		if (MaterialTexture.Value.Texture.IsNull())
+		{
+			continue;
+		}
+
 		const EDMMaterialPropertyType PropertyType = FDMUtils::TextureSetMaterialPropertyToMaterialPropertyType(MaterialTexture.Key);
 
 		if (PropertyType == EDMMaterialPropertyType::None)
@@ -645,11 +650,28 @@ bool UDynamicMaterialModelEditorOnlyData::AddTextureSet(UDMTextureSet* InTexture
 			continue;
 		}
 
+		UDMMaterialProperty* Property = GetMaterialProperty(PropertyType);
+
+		if (!Property)
+		{
+			continue;
+		}
+
+		if (!Property->IsEnabled())
+		{
+			Property->SetEnabled(true);
+		}
+
 		UDMMaterialSlot* Slot = GetSlotForMaterialProperty(PropertyType);
 
 		if (!Slot)
 		{
-			continue;
+			Slot = AddSlotForMaterialProperty(PropertyType);
+
+			if (!Slot)
+			{
+				continue;
+			}
 		}
 
 		UTexture* Texture = MaterialTexture.Value.Texture.LoadSynchronous();
