Material Designer
- When selecting the Alpha channel override, preview images will now compile correctly.
- When disabling mask stages in the base and emissive channels, they will correctly be ignored by the opacity channel's automatic multiplier.
- Blends and Functions that attempt to use the base and mask properties will now only use those properties if they are enabled.

diff --git a/Engine/Plugins/Experimental/DynamicMaterial/Source/DynamicMaterialEditor/Private/Components/DMMaterialProperty.cpp b/Engine/Plugins/Experimental/DynamicMaterial/Source/DynamicMaterialEditor/Private/Components/DMMaterialProperty.cpp
--- a/Engine/Plugins/Experimental/DynamicMaterial/Source/DynamicMaterialEditor/Private/Components/DMMaterialProperty.cpp
+++ b/Engine/Plugins/Experimental/DynamicMaterial/Source/DynamicMaterialEditor/Private/Components/DMMaterialProperty.cpp
@@ -407,6 +407,11 @@ void UDMMaterialProperty::GenerateOpacityExpressions(const TSharedRef<FDMMateria
 		UDMMaterialStage* BaseStage = Layer->GetStage(EDMMaterialLayerStage::Base);
 		UDMMaterialStage* MaskStage = Layer->GetStage(EDMMaterialLayerStage::Mask);
 
+		if (!MaskStage->IsEnabled())
+		{
+			continue;
+		}
+
 		MaskStage->GenerateExpressions(InBuildState);
 		UDMMaterialStageThroughputLayerBlend* LayerBlend = Cast<UDMMaterialStageThroughputLayerBlend>(MaskStage->GetSource());
 
diff --git a/Engine/Plugins/Experimental/DynamicMaterial/Source/DynamicMaterialEditor/Private/Components/DMMaterialStageBlend.cpp b/Engine/Plugins/Experimental/DynamicMaterial/Source/DynamicMaterialEditor/Private/Components/DMMaterialStageBlend.cpp
--- a/Engine/Plugins/Experimental/DynamicMaterial/Source/DynamicMaterialEditor/Private/Components/DMMaterialStageBlend.cpp
+++ b/Engine/Plugins/Experimental/DynamicMaterial/Source/DynamicMaterialEditor/Private/Components/DMMaterialStageBlend.cpp
@@ -490,17 +490,17 @@ void UDMMaterialStageBlend::AddDefaultInput(int32 InInputIndex) const
 			{
 				EDMMaterialPropertyType DefaultProperty = StageProperty;
 
-				if (DefaultProperty != EDMMaterialPropertyType::None)
+				if (DefaultProperty == EDMMaterialPropertyType::None)
 				{
 					if (UDMMaterialSlot* Slot = Layer->GetSlot())
 					{
 						if (UDynamicMaterialModelEditorOnlyData* ModelEditorOnlyData = Slot->GetMaterialModelEditorOnlyData())
 						{
-							if (ModelEditorOnlyData->GetSlotForMaterialProperty(EDMMaterialPropertyType::BaseColor))
+							if (ModelEditorOnlyData->GetSlotForEnabledMaterialProperty(EDMMaterialPropertyType::BaseColor))
 							{
 								DefaultProperty = EDMMaterialPropertyType::BaseColor;
 							}
-							else if (ModelEditorOnlyData->GetSlotForMaterialProperty(EDMMaterialPropertyType::EmissiveColor))
+							else if (ModelEditorOnlyData->GetSlotForEnabledMaterialProperty(EDMMaterialPropertyType::EmissiveColor))
 							{
 								DefaultProperty = EDMMaterialPropertyType::EmissiveColor;
 							}
diff --git a/Engine/Plugins/Experimental/DynamicMaterial/Source/DynamicMaterialEditor/Private/Components/DMMaterialStageFunction.cpp b/Engine/Plugins/Experimental/DynamicMaterial/Source/DynamicMaterialEditor/Private/Components/DMMaterialStageFunction.cpp
--- a/Engine/Plugins/Experimental/DynamicMaterial/Source/DynamicMaterialEditor/Private/Components/DMMaterialStageFunction.cpp
+++ b/Engine/Plugins/Experimental/DynamicMaterial/Source/DynamicMaterialEditor/Private/Components/DMMaterialStageFunction.cpp
@@ -153,22 +153,25 @@ void UDMMaterialStageFunction::AddDefaultInput(int32 InInputIndex) const
 	}
 	else
 	{
-		EDMMaterialPropertyType DefaultProperty = EDMMaterialPropertyType::None;
+		EDMMaterialPropertyType DefaultProperty = StageProperty;
 
-		if (UDMMaterialSlot* Slot = Layer->GetSlot())
+		if (DefaultProperty == EDMMaterialPropertyType::None)
 		{
-			if (UDynamicMaterialModelEditorOnlyData* ModelEditorOnlyData = Slot->GetMaterialModelEditorOnlyData())
+			if (UDMMaterialSlot* Slot = Layer->GetSlot())
 			{
-				if (ModelEditorOnlyData->GetSlotForMaterialProperty(EDMMaterialPropertyType::BaseColor))
+				if (UDynamicMaterialModelEditorOnlyData* ModelEditorOnlyData = Slot->GetMaterialModelEditorOnlyData())
 				{
-					DefaultProperty = EDMMaterialPropertyType::BaseColor;
-				}
-				else if (ModelEditorOnlyData->GetSlotForMaterialProperty(EDMMaterialPropertyType::EmissiveColor))
-				{
-					DefaultProperty = EDMMaterialPropertyType::EmissiveColor;
+					if (ModelEditorOnlyData->GetSlotForEnabledMaterialProperty(EDMMaterialPropertyType::BaseColor))
+					{
+						DefaultProperty = EDMMaterialPropertyType::BaseColor;
+					}
+					else if (ModelEditorOnlyData->GetSlotForEnabledMaterialProperty(EDMMaterialPropertyType::EmissiveColor))
+					{
+						DefaultProperty = EDMMaterialPropertyType::EmissiveColor;
+					}
 				}
 			}
-		}		
+		}
 
 		Stage->ChangeInput_PreviousStage(
 			InInputIndex,
diff --git a/Engine/Plugins/Experimental/DynamicMaterial/Source/DynamicMaterialEditor/Private/Components/DMMaterialStageThroughputLayerBlend.cpp b/Engine/Plugins/Experimental/DynamicMaterial/Source/DynamicMaterialEditor/Private/Components/DMMaterialStageThroughputLayerBlend.cpp
--- a/Engine/Plugins/Experimental/DynamicMaterial/Source/DynamicMaterialEditor/Private/Components/DMMaterialStageThroughputLayerBlend.cpp
+++ b/Engine/Plugins/Experimental/DynamicMaterial/Source/DynamicMaterialEditor/Private/Components/DMMaterialStageThroughputLayerBlend.cpp
@@ -373,12 +373,12 @@ void UDMMaterialStageThroughputLayerBlend::GenerateMainExpressions(const TShared
 
 					if (ModelEditorOnlyData)
 					{
-						if (UDMMaterialSlot* BaseColorSlot = ModelEditorOnlyData->GetSlotForMaterialProperty(EDMMaterialPropertyType::BaseColor))
+						if (UDMMaterialSlot* BaseColorSlot = ModelEditorOnlyData->GetSlotForEnabledMaterialProperty(EDMMaterialPropertyType::BaseColor))
 						{
 							RGBProperty = EDMMaterialPropertyType::BaseColor;
 							RGBSlot = BaseColorSlot;
 						}
-						else if (UDMMaterialSlot* EmissiveColorSlot = ModelEditorOnlyData->GetSlotForMaterialProperty(EDMMaterialPropertyType::EmissiveColor))
+						else if (UDMMaterialSlot* EmissiveColorSlot = ModelEditorOnlyData->GetSlotForEnabledMaterialProperty(EDMMaterialPropertyType::EmissiveColor))
 						{
 							RGBProperty = EDMMaterialPropertyType::EmissiveColor;
 							RGBSlot = EmissiveColorSlot;
@@ -598,6 +598,7 @@ bool UDMMaterialStageThroughputLayerBlend::GenerateStagePreviewMaterial(UDMMater
 	UMaterialExpression* ComponentMask = SourceExpressions.Last();
 	OutputIndex = InputConnectionMap[InputMaskSource].Channels[0].OutputIndex;
 	int32 Channel;
+	bool bGenerateBitmask = true;
 
 	{
 		const UDMMaterialStageSource* Source = nullptr;
@@ -619,41 +620,52 @@ bool UDMMaterialStageThroughputLayerBlend::GenerateStagePreviewMaterial(UDMMater
 		if (Source && Source->GetOutputConnectors().IsValidIndex(OutputIndex))
 		{
 			OutputIndex = Source->GetOutputConnectors()[OutputIndex].Index;
+
+			const int32 InnateOutputIndex = Source->GetInnateMaskOutput(OutputIndex, InputConnectionMap[InputMaskSource].Channels[0].OutputChannel);
+
+			if (InnateOutputIndex != INDEX_NONE)
+			{
+				OutputIndex = InnateOutputIndex;
+				bGenerateBitmask = false;
+			}
 		}
 	}
 
-	// Pick the first selected channel - this is always a single channel.
-	if (InputConnectionMap[InputMaskSource].Channels[0].OutputChannel & FDMMaterialStageConnectorChannel::FIRST_CHANNEL)
-	{
-		Channel = FDMMaterialStageConnectorChannel::FIRST_CHANNEL;
-	}
-	else if (InputConnectionMap[InputMaskSource].Channels[0].OutputChannel & FDMMaterialStageConnectorChannel::SECOND_CHANNEL)
-	{
-		Channel = FDMMaterialStageConnectorChannel::SECOND_CHANNEL;
-	}
-	else if (InputConnectionMap[InputMaskSource].Channels[0].OutputChannel & FDMMaterialStageConnectorChannel::THIRD_CHANNEL)
-	{
-		Channel = FDMMaterialStageConnectorChannel::THIRD_CHANNEL;
-	}
-	else if (InputConnectionMap[InputMaskSource].Channels[0].OutputChannel & FDMMaterialStageConnectorChannel::FOURTH_CHANNEL)
+	if (bGenerateBitmask)
 	{
-		Channel = FDMMaterialStageConnectorChannel::FOURTH_CHANNEL;
-	}
-	else
-	{
-		Channel = FDMMaterialStageConnectorChannel::FIRST_CHANNEL;
-	}
+		// Pick the first selected channel - this is always a single channel.
+		if (InputConnectionMap[InputMaskSource].Channels[0].OutputChannel & FDMMaterialStageConnectorChannel::FIRST_CHANNEL)
+		{
+			Channel = FDMMaterialStageConnectorChannel::FIRST_CHANNEL;
+		}
+		else if (InputConnectionMap[InputMaskSource].Channels[0].OutputChannel & FDMMaterialStageConnectorChannel::SECOND_CHANNEL)
+		{
+			Channel = FDMMaterialStageConnectorChannel::SECOND_CHANNEL;
+		}
+		else if (InputConnectionMap[InputMaskSource].Channels[0].OutputChannel & FDMMaterialStageConnectorChannel::THIRD_CHANNEL)
+		{
+			Channel = FDMMaterialStageConnectorChannel::THIRD_CHANNEL;
+		}
+		else if (InputConnectionMap[InputMaskSource].Channels[0].OutputChannel & FDMMaterialStageConnectorChannel::FOURTH_CHANNEL)
+		{
+			Channel = FDMMaterialStageConnectorChannel::FOURTH_CHANNEL;
+		}
+		else
+		{
+			Channel = FDMMaterialStageConnectorChannel::FIRST_CHANNEL;
+		}
 
-	ComponentMask = BuildState->GetBuildUtils().CreateExpressionBitMask(
-		ComponentMask,
-		OutputIndex,
-		Channel
-	);
+		ComponentMask = BuildState->GetBuildUtils().CreateExpressionBitMask(
+			ComponentMask,
+			OutputIndex,
+			Channel
+		);
 
-	SourceExpressions.Add(ComponentMask);
+		SourceExpressions.Add(ComponentMask);
 
-	OutputIndex = 0;
-	Channel = FDMMaterialStageConnectorChannel::WHOLE_CHANNEL;
+		OutputIndex = 0;
+		Channel = FDMMaterialStageConnectorChannel::WHOLE_CHANNEL;
+	}
 
 	Layer->ApplyEffects(BuildState, InStage, SourceExpressions, Channel, OutputIndex);
 
