Material Designer
- Fixed failing to check for invalid values when creating instanced material.

diff --git a/Engine/Plugins/Experimental/DynamicMaterial/Source/DynamicMaterial/Private/Model/DynamicMaterialModelDynamic.cpp b/Engine/Plugins/Experimental/DynamicMaterial/Source/DynamicMaterial/Private/Model/DynamicMaterialModelDynamic.cpp
--- a/Engine/Plugins/Experimental/DynamicMaterial/Source/DynamicMaterial/Private/Model/DynamicMaterialModelDynamic.cpp
+++ b/Engine/Plugins/Experimental/DynamicMaterial/Source/DynamicMaterial/Private/Model/DynamicMaterialModelDynamic.cpp
@@ -10,6 +10,7 @@
 #if WITH_EDITOR
 #include "Components/DMMaterialValue.h"
 #include "Components/DMTextureUV.h"
+#include "DynamicMaterialModule.h"
 #endif
 
 const FString UDynamicMaterialModelDynamic::ParentModelPathToken = TEXT("ParentModel");
@@ -254,13 +255,19 @@ void UDynamicMaterialModelDynamic::EnsureComponents()
 	const TArray<UDMMaterialValue*>& ParentValues = ParentModelLocal->GetValues();
 	const TSet<TObjectPtr<UDMMaterialComponent>>& RuntimeComponents = ParentModelLocal->GetRuntimeComponents();
 
-	const int32 RequiredComponentCount = GlobalParamCount + ParentValues.Num() + RuntimeComponents.Num();
+	int32 RequiredComponentCount = GlobalParamCount + ParentValues.Num() + RuntimeComponents.Num();
 
 	DynamicComponents.Reserve(RequiredComponentCount);
 
 	ParentModelLocal->ForEachGlobalParameter(
-		[this](UDMMaterialValue* InValue)
+		[this, &RequiredComponentCount](UDMMaterialValue* InValue)
 		{
+			if (!IsValid(InValue))
+			{
+				--RequiredComponentCount;
+				return;
+			}
+
 			const FName ValueName = InValue->GetFName();
 
 			if (!DynamicComponents.Contains(ValueName))
@@ -272,6 +279,12 @@ void UDynamicMaterialModelDynamic::EnsureComponents()
 
 	for (UDMMaterialValue* ParentValue : ParentValues)
 	{
+		if (!IsValid(ParentValue))
+		{
+			--RequiredComponentCount;
+			continue;
+		}
+
 		const FName ParentValueName = ParentValue->GetFName();
 
 		if (!DynamicComponents.Contains(ParentValueName))
@@ -282,6 +295,12 @@ void UDynamicMaterialModelDynamic::EnsureComponents()
 
 	for (const TObjectPtr<UDMMaterialComponent>& RuntimeComponent : RuntimeComponents)
 	{
+		if (!IsValid(RuntimeComponent))
+		{
+			--RequiredComponentCount;
+			continue;
+		}
+
 		if (UDMMaterialValue* ParentValue = Cast<UDMMaterialValue>(RuntimeComponent))
 		{
 			const FName ParentValueName = ParentValue->GetFName();
@@ -300,11 +319,16 @@ void UDynamicMaterialModelDynamic::EnsureComponents()
 				DynamicComponents.Add(ParentValueName, TextureUV->ToDynamic(this));
 			}
 		}
+		else
+		{
+			--RequiredComponentCount;
+			UE_LOG(LogDynamicMaterial, Error, TEXT("Invalid Component type while creating Material Designer Instance [%s]"), *RuntimeComponent->GetName());
+		}
 	}
 
 	const int32 RemovedComponentCount = DynamicComponents.Num() - RequiredComponentCount;
 
-	if (RemovedComponentCount == 0)
+	if (RemovedComponentCount <= 0)
 	{
 		return;
 	}
