Material Designer
- Value and texture parameter names are now not changed upon load.

diff --git a/Engine/Plugins/Experimental/DynamicMaterial/Source/DynamicMaterial/Private/Components/DMMaterialValue.cpp b/Engine/Plugins/Experimental/DynamicMaterial/Source/DynamicMaterial/Private/Components/DMMaterialValue.cpp
--- a/Engine/Plugins/Experimental/DynamicMaterial/Source/DynamicMaterial/Private/Components/DMMaterialValue.cpp
+++ b/Engine/Plugins/Experimental/DynamicMaterial/Source/DynamicMaterial/Private/Components/DMMaterialValue.cpp
@@ -204,7 +204,7 @@ bool UDMMaterialValue::SetParameterName(FName InBaseName)
 		Parameter->SetParentComponent(this);
 	}
 
-	UpdateCachedParameterName();
+	UpdateCachedParameterName(/* Reset Name */ false);
  
 	return true;
 }
@@ -247,7 +247,8 @@ void UDMMaterialValue::OnComponentAdded()
 			MaterialModel->AddRuntimeComponentReference(this);
 		}
 
-		UpdateCachedParameterName();
+		CachedParameterName = NAME_None;
+		UpdateCachedParameterName(/* Reset Name */ true);
 	}
 }
 
@@ -293,7 +294,7 @@ void UDMMaterialValue::PostEditorDuplicate(UDynamicMaterialModel* InMaterialMode
 	if (GetOuter() == InMaterialModel)
 	{
 		Super::PostEditorDuplicate(InMaterialModel, InParent);
-		UpdateCachedParameterName();
+		UpdateCachedParameterName(/* Reset Name */ false);
 		return;
 	}
 
@@ -320,7 +321,7 @@ void UDMMaterialValue::PostEditorDuplicate(UDynamicMaterialModel* InMaterialMode
 		SetParameterName(OldParameterName);
 	}
 
-	UpdateCachedParameterName();
+	UpdateCachedParameterName(/* Reset Name */ false);
 }
 
 bool UDMMaterialValue::Modify(bool bInAlwaysMarkDirty)
@@ -402,13 +403,13 @@ FName UDMMaterialValue::GenerateAutomaticParameterName() const
 	return *GetComponentPath();
 }
 
-void UDMMaterialValue::UpdateCachedParameterName()
+void UDMMaterialValue::UpdateCachedParameterName(bool bInResetName)
 {
 	if (Parameter)
 	{
 		CachedParameterName = Parameter->GetParameterName();
 	}
-	else
+	else if (bInResetName || CachedParameterName.IsNone())
 	{
 		CachedParameterName = GenerateAutomaticParameterName();
 	}
@@ -437,7 +438,7 @@ void UDMMaterialValue::Update(UDMMaterialComponent* InSource, EDMUpdateType InUp
 
 	if (InUpdateType == EDMUpdateType::Structure)
 	{
-		UpdateCachedParameterName();
+		UpdateCachedParameterName(/* Reset Name */ false);
 	}
 #endif
 
diff --git a/Engine/Plugins/Experimental/DynamicMaterial/Source/DynamicMaterial/Private/Components/DMTextureUV.cpp b/Engine/Plugins/Experimental/DynamicMaterial/Source/DynamicMaterial/Private/Components/DMTextureUV.cpp
--- a/Engine/Plugins/Experimental/DynamicMaterial/Source/DynamicMaterial/Private/Components/DMTextureUV.cpp
+++ b/Engine/Plugins/Experimental/DynamicMaterial/Source/DynamicMaterial/Private/Components/DMTextureUV.cpp
@@ -337,6 +337,7 @@ void UDMTextureUV::PostEditorDuplicate(UDynamicMaterialModel* InMaterialModel, U
 	if (GetOuter() == InMaterialModel)
 	{
 		Super::PostEditorDuplicate(InMaterialModel, InParent);
+		UpdateCachedParameterNames(/* Reset names */ false);
 		return;
 	}
 
@@ -352,12 +353,13 @@ void UDMTextureUV::PostEditorDuplicate(UDynamicMaterialModel* InMaterialModel, U
 	}
 
 	MaterialParameters.Empty();
-	CachedParameterNames.Empty();
 
 	Super::PostEditorDuplicate(InMaterialModel, InParent);
 
 	Rename(nullptr, InMaterialModel, UE::DynamicMaterial::RenameFlags);
 
+	UpdateCachedParameterNames(/* Reset names */ false);
+
 	using namespace UE::DynamicMaterial;
 
 	for (const TPair<int32, FName>& OldParameterPair : OldParameterNames)
@@ -396,8 +398,6 @@ void UDMTextureUV::PostEditorDuplicate(UDynamicMaterialModel* InMaterialModel, U
 			}
 		}
 	}
-
-	UpdateCachedParameterNames();
 }
 #endif
 
@@ -411,12 +411,12 @@ void UDMTextureUV::SetMIDParameters(UMaterialInstanceDynamic* InMID)
 	check(InMID);
 
 	auto UpdateMID = [InMID](FName InParamName, float InValue)
-	{
-		if (FMath::IsNearlyEqual(InValue, InMID->K2_GetScalarParameterValue(InParamName)) == false)
 		{
-			InMID->SetScalarParameterValue(InParamName, InValue);
-		}
-	};
+			if (FMath::IsNearlyEqual(InValue, InMID->K2_GetScalarParameterValue(InParamName)) == false)
+			{
+				InMID->SetScalarParameterValue(InParamName, InValue);
+			}
+		};
 
 	UpdateMID(GetMaterialParameterName(UDMTextureUV::NAME_Offset,   0), GetOffset().X);
 	UpdateMID(GetMaterialParameterName(UDMTextureUV::NAME_Offset,   1), GetOffset().Y);
@@ -468,7 +468,7 @@ void UDMTextureUV::Update(UDMMaterialComponent* InSource, EDMUpdateType InUpdate
 
 	if (InUpdateType == EDMUpdateType::Structure)
 	{
-		UpdateCachedParameterNames();
+		UpdateCachedParameterNames(/* Reset Names */ false);
 	}
 #endif
 
@@ -686,14 +686,19 @@ void UDMTextureUV::UpdateCachedParameterName(FName InPropertyName, int32 InCompo
 	{
 		CachedParameterNames.FindOrAdd(ParamId) = (*ParameterPtr)->GetParameterName();
 	}
-	else
+	else if (!CachedParameterNames.Contains(ParamId))
 	{
-		CachedParameterNames.FindOrAdd(ParamId) = GenerateAutomaticParameterName(InPropertyName, InComponent);
+		CachedParameterNames.Add(ParamId, GenerateAutomaticParameterName(InPropertyName, InComponent));
 	}
 }
 
-void UDMTextureUV::UpdateCachedParameterNames()
+void UDMTextureUV::UpdateCachedParameterNames(bool bInResetNames)
 {
+	if (bInResetNames)
+	{
+		CachedParameterNames.Empty(7);
+	}
+
 	UpdateCachedParameterName(UDMTextureUV::NAME_Offset, 0);
 	UpdateCachedParameterName(UDMTextureUV::NAME_Offset, 1);
 	UpdateCachedParameterName(UDMTextureUV::NAME_Rotation, 0);
@@ -710,7 +715,7 @@ void UDMTextureUV::OnComponentAdded()
 		return;
 	}
 
-	UpdateCachedParameterNames();
+	UpdateCachedParameterNames(/* Reset Names */ true);
 
 	if (UDynamicMaterialModel* MaterialModel = GetMaterialModel())
 	{
@@ -800,7 +805,7 @@ void UDMTextureUV::PostLoad()
 		MaterialModel->AddRuntimeComponentReference(this);
 	}
 
-	UpdateCachedParameterNames();
+	UpdateCachedParameterNames(/* Reset Names */ false);
 
 	for (const TPair<int32, TObjectPtr<UDMMaterialParameter>>& Pair : MaterialParameters)
 	{
@@ -835,7 +840,7 @@ void UDMTextureUV::PostEditImport()
 		return;
 	}
 
-	UpdateCachedParameterNames();
+	UpdateCachedParameterNames(/* Reset Names */ false);
 
 	for (const TPair<int32, TObjectPtr<UDMMaterialParameter>>& Pair : MaterialParameters)
 	{
diff --git a/Engine/Plugins/Experimental/DynamicMaterial/Source/DynamicMaterial/Public/Components/DMMaterialValue.h b/Engine/Plugins/Experimental/DynamicMaterial/Source/DynamicMaterial/Public/Components/DMMaterialValue.h
--- a/Engine/Plugins/Experimental/DynamicMaterial/Source/DynamicMaterial/Public/Components/DMMaterialValue.h
+++ b/Engine/Plugins/Experimental/DynamicMaterial/Source/DynamicMaterial/Public/Components/DMMaterialValue.h
@@ -200,7 +200,7 @@ protected:
 	DYNAMICMATERIAL_API FName GenerateAutomaticParameterName() const;
 
 	/** Updates the cached parameter name based on the Parameter object or the above method. */
-	void UpdateCachedParameterName();
+	void UpdateCachedParameterName(bool bInResetName);
 #endif
  
 	//~ Begin UDMMaterialComponent
diff --git a/Engine/Plugins/Experimental/DynamicMaterial/Source/DynamicMaterial/Public/Components/DMTextureUV.h b/Engine/Plugins/Experimental/DynamicMaterial/Source/DynamicMaterial/Public/Components/DMTextureUV.h
--- a/Engine/Plugins/Experimental/DynamicMaterial/Source/DynamicMaterial/Public/Components/DMTextureUV.h
+++ b/Engine/Plugins/Experimental/DynamicMaterial/Source/DynamicMaterial/Public/Components/DMTextureUV.h
@@ -263,7 +263,7 @@ protected:
 	void UpdateCachedParameterName(FName InPropertyName, int32 InComponent);
 
 	/** Updates the cached parameter name based on the Parameter object or the above method. */
-	void UpdateCachedParameterNames();
+	void UpdateCachedParameterNames(bool bInResetNames);
 #endif
 
 	//~ Begin UDMMaterialComponent
