Material Designer
- Texture names in the asset creator now have a translucent background so the texture can be seen through it.

diff --git a/Engine/Plugins/Experimental/DynamicMaterial/Source/DynamicMaterialTextureSetEditor/Private/DMTextureSetStyle.cpp b/Engine/Plugins/Experimental/DynamicMaterial/Source/DynamicMaterialTextureSetEditor/Private/DMTextureSetStyle.cpp
--- a/Engine/Plugins/Experimental/DynamicMaterial/Source/DynamicMaterialTextureSetEditor/Private/DMTextureSetStyle.cpp
+++ b/Engine/Plugins/Experimental/DynamicMaterial/Source/DynamicMaterialTextureSetEditor/Private/DMTextureSetStyle.cpp
@@ -22,7 +22,7 @@ FDMTextureSetStyle::FDMTextureSetStyle()
 		FStyleColors::Recessed.GetSpecifiedColor(), 6.0f,
 		FStyleColors::Header.GetSpecifiedColor(), 2.0f));
 
-	Set("TextureSetConfig.TextureNameHighlight.Background", new FSlateColorBrush(FStyleColors::Black.GetSpecifiedColor()));
+	Set("TextureSetConfig.TextureNameHighlight.Background", new FSlateColorBrush(FStyleColors::Black.GetSpecifiedColor() * FLinearColor(1.f, 1.f, 1.f, 0.5f)));
 
 	Set("TextureSetConfig.TextureNameHighlight.Color", FStyleColors::AccentWhite);
 }
