Material Designer
- The brown/yellow color of the wizard preset selection buttons is now green.

diff --git a/Engine/Plugins/Experimental/DynamicMaterial/Source/DynamicMaterialEditor/Private/DynamicMaterialEditorStyle.cpp b/Engine/Plugins/Experimental/DynamicMaterial/Source/DynamicMaterialEditor/Private/DynamicMaterialEditorStyle.cpp
--- a/Engine/Plugins/Experimental/DynamicMaterial/Source/DynamicMaterialEditor/Private/DynamicMaterialEditorStyle.cpp
+++ b/Engine/Plugins/Experimental/DynamicMaterial/Source/DynamicMaterialEditor/Private/DynamicMaterialEditorStyle.cpp
@@ -126,9 +126,9 @@ void FDynamicMaterialEditorStyle::SetupGeneralStyles()
 		.SetPressedPadding(FMargin(0, 1, 0, 0)));
 
 	Set("DulledSectionButton", FCheckBoxStyle(FStarshipCoreStyle::GetCoreStyle().GetWidgetStyle<FCheckBoxStyle>("FilterBar.BasicFilterButton"))
-		.SetCheckedImage(FSlateRoundedBoxBrush(FStyleColors::AccentBrown, 4.0f, FStyleColors::Input, 1.0f))
-		.SetCheckedHoveredImage(FSlateRoundedBoxBrush(FStyleColors::AccentFolder, 4.0f, FStyleColors::Input, 1.0f))
-		.SetCheckedPressedImage(FSlateRoundedBoxBrush(FStyleColors::AccentFolder, 4.0f, FStyleColors::Input, 1.0f)));
+		.SetCheckedImage(FSlateRoundedBoxBrush(FStyleColors::AccentGreen.GetSpecifiedColor() * FLinearColor(0.5f, 0.5f, 0.5f, 1.f), 4.0f, FStyleColors::Input, 1.0f))
+		.SetCheckedHoveredImage(FSlateRoundedBoxBrush(FStyleColors::AccentGreen, 4.0f, FStyleColors::Input, 1.0f))
+		.SetCheckedPressedImage(FSlateRoundedBoxBrush(FStyleColors::AccentGreen, 4.0f, FStyleColors::Input, 1.0f)));
 }
 
 void FDynamicMaterialEditorStyle::SetupStageStyles()
