Material Designer
- Material object names are now called "Material Designer" instead of "Dynamic Material Instance"
- Model object names are now "Material Designer Model" and "Material Designer Model Instance".

diff --git a/Engine/Plugins/Experimental/DynamicMaterial/Source/DynamicMaterial/Private/Model/DynamicMaterialModelDynamic.cpp b/Engine/Plugins/Experimental/DynamicMaterial/Source/DynamicMaterial/Private/Model/DynamicMaterialModelDynamic.cpp
--- a/Engine/Plugins/Experimental/DynamicMaterial/Source/DynamicMaterial/Private/Model/DynamicMaterialModelDynamic.cpp
+++ b/Engine/Plugins/Experimental/DynamicMaterial/Source/DynamicMaterial/Private/Model/DynamicMaterialModelDynamic.cpp
@@ -21,7 +21,7 @@ UDynamicMaterialModelDynamic* UDynamicMaterialModelDynamic::Create(UObject* InOu
 {
 	check(InParentModel);
 
-	UDynamicMaterialModelDynamic* NewModelDynamic = NewObject<UDynamicMaterialModelDynamic>(InOuter, NAME_None, RF_Transactional);
+	UDynamicMaterialModelDynamic* NewModelDynamic = NewObject<UDynamicMaterialModelDynamic>(InOuter, TEXT("MaterialDesignerModelInstance"), RF_Transactional);
 	NewModelDynamic->ParentModelSoft = InParentModel;
 	NewModelDynamic->ParentModel = InParentModel;
 	NewModelDynamic->InitComponents();
diff --git a/Engine/Plugins/Experimental/DynamicMaterial/Source/DynamicMaterialEditor/Private/DMContentBrowserIntegration.cpp b/Engine/Plugins/Experimental/DynamicMaterial/Source/DynamicMaterialEditor/Private/DMContentBrowserIntegration.cpp
--- a/Engine/Plugins/Experimental/DynamicMaterial/Source/DynamicMaterialEditor/Private/DMContentBrowserIntegration.cpp
+++ b/Engine/Plugins/Experimental/DynamicMaterial/Source/DynamicMaterialEditor/Private/DMContentBrowserIntegration.cpp
@@ -422,7 +422,7 @@ void FDMContentBrowserIntegration::CreateMaterialInstance(UDynamicMaterialInstan
 		return;
 	}
 
-	UDynamicMaterialInstance* Instance = NewObject<UDynamicMaterialInstance>(GetTransientPackage());
+	UDynamicMaterialInstance* Instance = NewObject<UDynamicMaterialInstance>(GetTransientPackage(), TEXT("MaterialDesigner"));
 
 	if (!Instance)
 	{
diff --git a/Engine/Plugins/Experimental/DynamicMaterial/Source/DynamicMaterialEditor/Private/Material/DynamicMaterialInstanceFactory.cpp b/Engine/Plugins/Experimental/DynamicMaterial/Source/DynamicMaterialEditor/Private/Material/DynamicMaterialInstanceFactory.cpp
--- a/Engine/Plugins/Experimental/DynamicMaterial/Source/DynamicMaterialEditor/Private/Material/DynamicMaterialInstanceFactory.cpp
+++ b/Engine/Plugins/Experimental/DynamicMaterial/Source/DynamicMaterialEditor/Private/Material/DynamicMaterialInstanceFactory.cpp
@@ -23,18 +23,23 @@ UDynamicMaterialInstanceFactory::UDynamicMaterialInstanceFactory()
 	bText = false;
 }
 
-UObject* UDynamicMaterialInstanceFactory::FactoryCreateNew(UClass* Class, UObject* InParent, FName Name, 
-	EObjectFlags Flags, UObject* Context, FFeedbackContext* Warn)
+UObject* UDynamicMaterialInstanceFactory::FactoryCreateNew(UClass* InClass, UObject* InParent, FName InName, 
+	EObjectFlags InFlags, UObject* InContext, FFeedbackContext* InWarn)
 {
-	check(Class->IsChildOf(UDynamicMaterialInstance::StaticClass()));
+	check(InClass->IsChildOf(UDynamicMaterialInstance::StaticClass()));
 
-	UDynamicMaterialInstance* NewInstance = NewObject<UDynamicMaterialInstance>(InParent, Class, Name, Flags | RF_Transactional);
+	if (InName.IsNone())
+	{
+		InName = TEXT("MaterialDesigner");
+	}
+
+	UDynamicMaterialInstance* NewInstance = NewObject<UDynamicMaterialInstance>(InParent, InClass, InName, InFlags | RF_Transactional);
 	check(NewInstance);
 
 	UDynamicMaterialModelFactory* EditorFactory = NewObject<UDynamicMaterialModelFactory>();
 	check(EditorFactory);
 
-	UDynamicMaterialModelBase* ModelBase = Cast<UDynamicMaterialModelBase>(Context);
+	UDynamicMaterialModelBase* ModelBase = Cast<UDynamicMaterialModelBase>(InContext);
 
 	if (!ModelBase)
 	{
diff --git a/Engine/Plugins/Experimental/DynamicMaterial/Source/DynamicMaterialEditor/Private/Model/DynamicMateriaModelFactory.cpp b/Engine/Plugins/Experimental/DynamicMaterial/Source/DynamicMaterialEditor/Private/Model/DynamicMateriaModelFactory.cpp
--- a/Engine/Plugins/Experimental/DynamicMaterial/Source/DynamicMaterialEditor/Private/Model/DynamicMateriaModelFactory.cpp
+++ b/Engine/Plugins/Experimental/DynamicMaterial/Source/DynamicMaterialEditor/Private/Model/DynamicMateriaModelFactory.cpp
@@ -21,12 +21,17 @@ UDynamicMaterialModelFactory::UDynamicMaterialModelFactory()
 	bText = false;
 }
 
-UObject* UDynamicMaterialModelFactory::FactoryCreateNew(UClass* Class, UObject* InParent, FName Name, 
-	EObjectFlags Flags, UObject* Context, FFeedbackContext* Warn)
+UObject* UDynamicMaterialModelFactory::FactoryCreateNew(UClass* InClass, UObject* InParent, FName InName, 
+	EObjectFlags InFlags, UObject* InContext, FFeedbackContext* InWarn)
 {
-	check(Class->IsChildOf(UDynamicMaterialModel::StaticClass()));
+	check(InClass->IsChildOf(UDynamicMaterialModel::StaticClass()));
 
-	UDynamicMaterialModel* NewModel = NewObject<UDynamicMaterialModel>(InParent, Class, Name, Flags | RF_Transactional);
+	if (InName.IsNone())
+	{
+		InName = TEXT("MaterialDesignerModel");
+	}
+
+	UDynamicMaterialModel* NewModel = NewObject<UDynamicMaterialModel>(InParent, InClass, InName, InFlags | RF_Transactional);
 
 	UDynamicMaterialModelEditorOnlyData* ModelEditorOnlyData = NewObject<UDynamicMaterialModelEditorOnlyData>(
 		NewModel,
diff --git a/Engine/Plugins/Experimental/DynamicMaterial/Source/DynamicMaterialEditor/Private/Model/DynamicMaterialModelFactory.h b/Engine/Plugins/Experimental/DynamicMaterial/Source/DynamicMaterialEditor/Private/Model/DynamicMaterialModelFactory.h
--- a/Engine/Plugins/Experimental/DynamicMaterial/Source/DynamicMaterialEditor/Private/Model/DynamicMaterialModelFactory.h
+++ b/Engine/Plugins/Experimental/DynamicMaterial/Source/DynamicMaterialEditor/Private/Model/DynamicMaterialModelFactory.h
@@ -17,7 +17,8 @@ public:
 
 	UDynamicMaterialModelFactory();
 
-	virtual UObject* FactoryCreateNew(UClass* Class, UObject* InParent, FName Name, EObjectFlags Flags, UObject* Context, FFeedbackContext* Warn) override;
+	virtual UObject* FactoryCreateNew(UClass* InClass, UObject* InParent, FName InName, EObjectFlags InFlags, UObject* InContext, 
+		FFeedbackContext* InWarn) override;
 	virtual FText GetDisplayName() const override;
 	virtual FText GetToolTip() const override;
 };
\ No newline at end of file
diff --git a/Engine/Plugins/Experimental/DynamicMaterial/Source/DynamicMaterialEditor/Private/UI/Widgets/SDMMaterialWizard.cpp b/Engine/Plugins/Experimental/DynamicMaterial/Source/DynamicMaterialEditor/Private/UI/Widgets/SDMMaterialWizard.cpp
--- a/Engine/Plugins/Experimental/DynamicMaterial/Source/DynamicMaterialEditor/Private/UI/Widgets/SDMMaterialWizard.cpp
+++ b/Engine/Plugins/Experimental/DynamicMaterial/Source/DynamicMaterialEditor/Private/UI/Widgets/SDMMaterialWizard.cpp
@@ -831,7 +831,7 @@ void SDMMaterialWizard::CreateNewDynamicInstanceInActor(UDynamicMaterialModel* I
 		return;
 	}
 
-	UDynamicMaterialInstanceFactory* Factory = NewObject<UDynamicMaterialInstanceFactory>(GetTransientPackage());
+	UDynamicMaterialInstanceFactory* Factory = NewObject<UDynamicMaterialInstanceFactory>(GetTransientPackage(), TEXT("MaterialDesigner"));
 
 	UDynamicMaterialInstance* NewInstance = Cast<UDynamicMaterialInstance>(Factory->FactoryCreateNew(
 		UDynamicMaterialInstance::StaticClass(),
@@ -909,7 +909,7 @@ void SDMMaterialWizard::CreateTemplateMaterialInActor(FName InChannelList, FDMOb
 		return;
 	}
 
-	UDynamicMaterialInstanceFactory* Factory = NewObject<UDynamicMaterialInstanceFactory>(GetTransientPackage());
+	UDynamicMaterialInstanceFactory* Factory = NewObject<UDynamicMaterialInstanceFactory>(GetTransientPackage(), TEXT("MaterialDesigner"));
 
 	UDynamicMaterialInstance* NewInstance = Cast<UDynamicMaterialInstance>(Factory->FactoryCreateNew(
 		UDynamicMaterialInstance::StaticClass(),
@@ -930,6 +930,7 @@ void SDMMaterialWizard::CreateTemplateMaterialInActor(FName InChannelList, FDMOb
 	{
 		EditorOnlyData->SetChannelListPreset(InChannelList);
 		EditorOnlyData->OnWizardComplete();
+		EditorOnlyData->RequestMaterialBuild(/* Immediate */ true);
 	}
 
 	if (!UDMMaterialInstanceFunctionLibrary::SetMaterialInObject(InMaterialObjectProperty, NewInstance))
diff --git a/Engine/Plugins/Experimental/DynamicMaterial/Source/DynamicMaterialEditor/Private/Utils/DMMaterialInstanceFunctionLibrary.cpp b/Engine/Plugins/Experimental/DynamicMaterial/Source/DynamicMaterialEditor/Private/Utils/DMMaterialInstanceFunctionLibrary.cpp
--- a/Engine/Plugins/Experimental/DynamicMaterial/Source/DynamicMaterialEditor/Private/Utils/DMMaterialInstanceFunctionLibrary.cpp
+++ b/Engine/Plugins/Experimental/DynamicMaterial/Source/DynamicMaterialEditor/Private/Utils/DMMaterialInstanceFunctionLibrary.cpp
@@ -103,7 +103,7 @@ UDynamicMaterialModel* UDMMaterialInstanceFunctionLibrary::CreateMaterialInObjec
 		return nullptr;
 	}
 
-	UDynamicMaterialInstanceFactory* const InstanceFactory = NewObject<UDynamicMaterialInstanceFactory>();
+	UDynamicMaterialInstanceFactory* const InstanceFactory = NewObject<UDynamicMaterialInstanceFactory>(Outer, TEXT("MaterialDesigner"));
 	check(InstanceFactory);
 
 	UDynamicMaterialInstance* const NewInstance = Cast<UDynamicMaterialInstance>(InstanceFactory->FactoryCreateNew(UDynamicMaterialInstance::StaticClass(),
diff --git a/Engine/Plugins/Experimental/DynamicMaterial/Source/DynamicMaterialEditor/Public/Material/DynamicMaterialInstanceFactory.h b/Engine/Plugins/Experimental/DynamicMaterial/Source/DynamicMaterialEditor/Public/Material/DynamicMaterialInstanceFactory.h
--- a/Engine/Plugins/Experimental/DynamicMaterial/Source/DynamicMaterialEditor/Public/Material/DynamicMaterialInstanceFactory.h
+++ b/Engine/Plugins/Experimental/DynamicMaterial/Source/DynamicMaterialEditor/Public/Material/DynamicMaterialInstanceFactory.h
@@ -14,7 +14,8 @@ public:
 	DYNAMICMATERIALEDITOR_API UDynamicMaterialInstanceFactory();
 
 	//~ Begin UFactory
-	DYNAMICMATERIALEDITOR_API virtual UObject* FactoryCreateNew(UClass* Class, UObject* InParent, FName Name, EObjectFlags Flags, UObject* Context, FFeedbackContext* Warn) override;
+	DYNAMICMATERIALEDITOR_API virtual UObject* FactoryCreateNew(UClass* InClass, UObject* InParent, FName InName, EObjectFlags InFlags, 
+		UObject* InContext, FFeedbackContext* InWarn) override;
 	DYNAMICMATERIALEDITOR_API virtual FText GetDisplayName() const override;
 	DYNAMICMATERIALEDITOR_API virtual FText GetToolTip() const override;
 	//~ End UFactory
