Material Designer
- Fix for object names being non-unique.

diff --git a/Engine/Plugins/Experimental/DynamicMaterial/Source/DynamicMaterial/Private/Model/DynamicMaterialModelDynamic.cpp b/Engine/Plugins/Experimental/DynamicMaterial/Source/DynamicMaterial/Private/Model/DynamicMaterialModelDynamic.cpp
--- a/Engine/Plugins/Experimental/DynamicMaterial/Source/DynamicMaterial/Private/Model/DynamicMaterialModelDynamic.cpp
+++ b/Engine/Plugins/Experimental/DynamicMaterial/Source/DynamicMaterial/Private/Model/DynamicMaterialModelDynamic.cpp
@@ -21,7 +21,9 @@ UDynamicMaterialModelDynamic* UDynamicMaterialModelDynamic::Create(UObject* InOu
 {
 	check(InParentModel);
 
-	UDynamicMaterialModelDynamic* NewModelDynamic = NewObject<UDynamicMaterialModelDynamic>(InOuter, TEXT("MaterialDesignerModelInstance"), RF_Transactional);
+	const FName ModelName = MakeUniqueObjectName(InOuter, UDynamicMaterialModelDynamic::StaticClass(), TEXT("MaterialDesignerModelInstance"));
+
+	UDynamicMaterialModelDynamic* NewModelDynamic = NewObject<UDynamicMaterialModelDynamic>(InOuter, ModelName, RF_Transactional);
 	NewModelDynamic->ParentModelSoft = InParentModel;
 	NewModelDynamic->ParentModel = InParentModel;
 	NewModelDynamic->InitComponents();
diff --git a/Engine/Plugins/Experimental/DynamicMaterial/Source/DynamicMaterialEditor/Private/DMContentBrowserIntegration.cpp b/Engine/Plugins/Experimental/DynamicMaterial/Source/DynamicMaterialEditor/Private/DMContentBrowserIntegration.cpp
--- a/Engine/Plugins/Experimental/DynamicMaterial/Source/DynamicMaterialEditor/Private/DMContentBrowserIntegration.cpp
+++ b/Engine/Plugins/Experimental/DynamicMaterial/Source/DynamicMaterialEditor/Private/DMContentBrowserIntegration.cpp
@@ -422,7 +422,10 @@ void FDMContentBrowserIntegration::CreateMaterialInstance(UDynamicMaterialInstan
 		return;
 	}
 
-	UDynamicMaterialInstance* Instance = NewObject<UDynamicMaterialInstance>(GetTransientPackage(), TEXT("MaterialDesigner"));
+	UDynamicMaterialInstance* Instance = NewObject<UDynamicMaterialInstance>(
+		GetTransientPackage(), 
+		MakeUniqueObjectName(GetTransientPackage(), UDynamicMaterialInstance::StaticClass(), TEXT("MaterialDesigner"))
+	);
 
 	if (!Instance)
 	{
diff --git a/Engine/Plugins/Experimental/DynamicMaterial/Source/DynamicMaterialEditor/Private/Material/DynamicMaterialInstanceFactory.cpp b/Engine/Plugins/Experimental/DynamicMaterial/Source/DynamicMaterialEditor/Private/Material/DynamicMaterialInstanceFactory.cpp
--- a/Engine/Plugins/Experimental/DynamicMaterial/Source/DynamicMaterialEditor/Private/Material/DynamicMaterialInstanceFactory.cpp
+++ b/Engine/Plugins/Experimental/DynamicMaterial/Source/DynamicMaterialEditor/Private/Material/DynamicMaterialInstanceFactory.cpp
@@ -30,7 +30,7 @@ UObject* UDynamicMaterialInstanceFactory::FactoryCreateNew(UClass* InClass, UObj
 
 	if (InName.IsNone())
 	{
-		InName = TEXT("MaterialDesigner");
+		InName = MakeUniqueObjectName(InParent, UDynamicMaterialInstance::StaticClass(), TEXT("MaterialDesigner"));
 	}
 
 	UDynamicMaterialInstance* NewInstance = NewObject<UDynamicMaterialInstance>(InParent, InClass, InName, InFlags | RF_Transactional);
diff --git a/Engine/Plugins/Experimental/DynamicMaterial/Source/DynamicMaterialEditor/Private/Model/DynamicMateriaModelFactory.cpp b/Engine/Plugins/Experimental/DynamicMaterial/Source/DynamicMaterialEditor/Private/Model/DynamicMateriaModelFactory.cpp
--- a/Engine/Plugins/Experimental/DynamicMaterial/Source/DynamicMaterialEditor/Private/Model/DynamicMateriaModelFactory.cpp
+++ b/Engine/Plugins/Experimental/DynamicMaterial/Source/DynamicMaterialEditor/Private/Model/DynamicMateriaModelFactory.cpp
@@ -28,7 +28,7 @@ UObject* UDynamicMaterialModelFactory::FactoryCreateNew(UClass* InClass, UObject
 
 	if (InName.IsNone())
 	{
-		InName = TEXT("MaterialDesignerModel");
+		InName = MakeUniqueObjectName(InParent, UDynamicMaterialModel::StaticClass(), TEXT("MaterialDesignerModel"));
 	}
 
 	UDynamicMaterialModel* NewModel = NewObject<UDynamicMaterialModel>(InParent, InClass, InName, InFlags | RF_Transactional);
diff --git a/Engine/Plugins/Experimental/DynamicMaterial/Source/DynamicMaterialEditor/Private/Utils/DMMaterialInstanceFunctionLibrary.cpp b/Engine/Plugins/Experimental/DynamicMaterial/Source/DynamicMaterialEditor/Private/Utils/DMMaterialInstanceFunctionLibrary.cpp
--- a/Engine/Plugins/Experimental/DynamicMaterial/Source/DynamicMaterialEditor/Private/Utils/DMMaterialInstanceFunctionLibrary.cpp
+++ b/Engine/Plugins/Experimental/DynamicMaterial/Source/DynamicMaterialEditor/Private/Utils/DMMaterialInstanceFunctionLibrary.cpp
@@ -103,7 +103,7 @@ UDynamicMaterialModel* UDMMaterialInstanceFunctionLibrary::CreateMaterialInObjec
 		return nullptr;
 	}
 
-	UDynamicMaterialInstanceFactory* const InstanceFactory = NewObject<UDynamicMaterialInstanceFactory>(Outer, TEXT("MaterialDesigner"));
+	UDynamicMaterialInstanceFactory* const InstanceFactory = NewObject<UDynamicMaterialInstanceFactory>();
 	check(InstanceFactory);
 
 	UDynamicMaterialInstance* const NewInstance = Cast<UDynamicMaterialInstance>(InstanceFactory->FactoryCreateNew(UDynamicMaterialInstance::StaticClass(),
