Material Designer
- Effects can no longer be toggled or removed from material instances.

diff --git a/Engine/Plugins/Experimental/DynamicMaterial/Source/DynamicMaterialEditor/Private/UI/Widgets/Editor/SlotEditor/SDMMaterialSlotLayerEffectItem.cpp b/Engine/Plugins/Experimental/DynamicMaterial/Source/DynamicMaterialEditor/Private/UI/Widgets/Editor/SlotEditor/SDMMaterialSlotLayerEffectItem.cpp
--- a/Engine/Plugins/Experimental/DynamicMaterial/Source/DynamicMaterialEditor/Private/UI/Widgets/Editor/SlotEditor/SDMMaterialSlotLayerEffectItem.cpp
+++ b/Engine/Plugins/Experimental/DynamicMaterial/Source/DynamicMaterialEditor/Private/UI/Widgets/Editor/SlotEditor/SDMMaterialSlotLayerEffectItem.cpp
@@ -4,8 +4,12 @@
 
 #include "Components/DMMaterialEffectStack.h"
 #include "DynamicMaterialEditorStyle.h"
+#include "Model/DynamicMaterialModel.h"
 #include "UI/DragDrop/DMLayerEffectsDragDropOperation.h"
 #include "UI/Widgets/Editor/SlotEditor/SDMMaterialSlotLayerEffectView.h"
+#include "UI/Widgets/Editor/SlotEditor/SDMMaterialSlotLayerItem.h"
+#include "UI/Widgets/Editor/SlotEditor/SDMMaterialSlotLayerView.h"
+#include "UI/Widgets/SDMMaterialEditor.h"
 #include "Utils/DMPrivate.h"
 #include "Widgets/Input/SButton.h"
 
@@ -114,6 +118,7 @@ TSharedRef<SWidget> SDMMaterialSlotLayerEffectItem::CreateLayerBypassButton()
 		.ButtonStyle(FDynamicMaterialEditorStyle::Get(), "HoverHintOnly")
 		.ToolTipText(LOCTEXT("LayerBypassTooltip", "Toggle the bypassing of this layer."))
 		.Cursor(EMouseCursor::Default)
+		.IsEnabled(CanModifyMaterialModel())
 		.OnClicked(this, &SDMMaterialSlotLayerEffectItem::OnLayerBypassButtonClick)
 		[
 			SNew(SImage)
@@ -129,6 +134,7 @@ TSharedRef<SWidget> SDMMaterialSlotLayerEffectItem::CreateLayerRemoveButton()
 		.ButtonStyle(FDynamicMaterialEditorStyle::Get(), "HoverHintOnly")
 		.ToolTipText(LOCTEXT("RemoveEffectTooltip", "Remove Effect"))
 		.Cursor(EMouseCursor::Default)
+		.IsEnabled(CanModifyMaterialModel())
 		.OnClicked(this, &SDMMaterialSlotLayerEffectItem::OnLayerRemoveButtonClick)
 		[
 			SNew(SImage)
@@ -157,6 +163,31 @@ FText SDMMaterialSlotLayerEffectItem::GetLayerHeaderText() const
 	return FText::GetEmpty();
 }
 
+bool SDMMaterialSlotLayerEffectItem::CanModifyMaterialModel() const
+{
+	if (TSharedPtr<SDMMaterialSlotLayerEffectView> EffectView = GetEffectView())
+	{
+		if (TSharedPtr<SDMMaterialSlotLayerItem> LayerItem = EffectView->GetLayerItem())
+		{
+			if (TSharedPtr<SDMMaterialSlotLayerView> LayerView = LayerItem->GetSlotLayerView())
+			{
+				if (TSharedPtr<SDMMaterialSlotEditor> SlotEditorWidget = LayerView->GetSlotEditorWidget())
+				{
+					if (TSharedPtr<SDMMaterialEditor> EditorWidget = SlotEditorWidget->GetEditorWidget())
+					{
+						if (UDynamicMaterialModelBase* MaterialModelBase = EditorWidget->GetMaterialModelBase())
+						{
+							return MaterialModelBase->IsA<UDynamicMaterialModel>();
+						}
+					}
+				}
+			}
+		}
+	}
+
+	return false;
+}
+
 const FSlateBrush* SDMMaterialSlotLayerEffectItem::GetLayerBypassButtonImage() const
 {
 	static const FSlateIcon VisibleIcon = FSlateIcon(FAppStyle::GetAppStyleSetName(), "Icons.Visible");
diff --git a/Engine/Plugins/Experimental/DynamicMaterial/Source/DynamicMaterialEditor/Private/UI/Widgets/Editor/SlotEditor/SDMMaterialSlotLayerEffectItem.h b/Engine/Plugins/Experimental/DynamicMaterial/Source/DynamicMaterialEditor/Private/UI/Widgets/Editor/SlotEditor/SDMMaterialSlotLayerEffectItem.h
--- a/Engine/Plugins/Experimental/DynamicMaterial/Source/DynamicMaterialEditor/Private/UI/Widgets/Editor/SlotEditor/SDMMaterialSlotLayerEffectItem.h
+++ b/Engine/Plugins/Experimental/DynamicMaterial/Source/DynamicMaterialEditor/Private/UI/Widgets/Editor/SlotEditor/SDMMaterialSlotLayerEffectItem.h
@@ -47,6 +47,8 @@ protected:
 
 	FText GetLayerHeaderText() const;
 
+	bool CanModifyMaterialModel() const;
+
 	const FSlateBrush* GetLayerBypassButtonImage() const;
 
 	FReply OnLayerRemoveButtonClick();
