Material Designer
- When viewing a template material, rather than an instance, the word "Material" is displayed at the bottom-left of the material preview.

diff --git a/Engine/Plugins/Experimental/DynamicMaterial/Source/DynamicMaterialEditor/Private/UI/Widgets/SDMMaterialEditor.cpp b/Engine/Plugins/Experimental/DynamicMaterial/Source/DynamicMaterialEditor/Private/UI/Widgets/SDMMaterialEditor.cpp
--- a/Engine/Plugins/Experimental/DynamicMaterial/Source/DynamicMaterialEditor/Private/UI/Widgets/SDMMaterialEditor.cpp
+++ b/Engine/Plugins/Experimental/DynamicMaterial/Source/DynamicMaterialEditor/Private/UI/Widgets/SDMMaterialEditor.cpp
@@ -1123,27 +1123,24 @@ TSharedRef<SDMMaterialProperties> SDMMaterialEditor::CreateSlot_MaterialProperti
 
 TSharedRef<SWidget> SDMMaterialEditor::CreateSlot_Preview()
 {
-	if (IsDynamicModel())
-	{
-		return SNew(SOverlay)
-			+ SOverlay::Slot()
-			[
-				SNew(SDMMaterialPreview, SharedThis(this), GetMaterialModelBase())
-			]
-			+ SOverlay::Slot()
-			.HAlign(EHorizontalAlignment::HAlign_Left)
-			.VAlign(EVerticalAlignment::VAlign_Bottom)
-			.Padding(3.f, 2.f)
-			[
-				SNew(STextBlock)
+	return SNew(SOverlay)
+		+ SOverlay::Slot()
+		[
+			SNew(SDMMaterialPreview, SharedThis(this), GetMaterialModelBase())
+		]
+		+ SOverlay::Slot()
+		.HAlign(EHorizontalAlignment::HAlign_Left)
+		.VAlign(EVerticalAlignment::VAlign_Bottom)
+		.Padding(3.f, 2.f)
+		[
+			SNew(STextBlock)
 				.Font(FAppStyle::GetFontStyle("TinyText"))
-				.Text(LOCTEXT("MaterialInstance", "Instance"))
+				.Text(IsDynamicModel() 
+					? LOCTEXT("MaterialInstance", "Instance")
+					: LOCTEXT("MaterialTemplate", "Material"))
 				.ShadowColorAndOpacity(FLinearColor::Black)
 				.ShadowOffset(FVector2D(1.0))
-			];
-	}
-
-	return SNew(SDMMaterialPreview, SharedThis(this), GetMaterialModelBase());
+		];
 }
 
 TSharedRef<SDMMaterialPropertySelector> SDMMaterialEditor::CreateSlot_PropertySelector()
