Material Designer
- The Light Function material domain is now available.

diff --git a/Engine/Plugins/Experimental/DynamicMaterial/Source/DynamicMaterialEditor/Private/Model/DynamicMaterialModelEditorOnlyData.cpp b/Engine/Plugins/Experimental/DynamicMaterial/Source/DynamicMaterialEditor/Private/Model/DynamicMaterialModelEditorOnlyData.cpp
--- a/Engine/Plugins/Experimental/DynamicMaterial/Source/DynamicMaterialEditor/Private/Model/DynamicMaterialModelEditorOnlyData.cpp
+++ b/Engine/Plugins/Experimental/DynamicMaterial/Source/DynamicMaterialEditor/Private/Model/DynamicMaterialModelEditorOnlyData.cpp
@@ -110,7 +110,8 @@ const TArray<EMaterialDomain> UDynamicMaterialModelEditorOnlyData::SupportedDoma
 {
 	EMaterialDomain::MD_Surface,
 	EMaterialDomain::MD_PostProcess,
-	EMaterialDomain::MD_DeferredDecal
+	EMaterialDomain::MD_DeferredDecal,
+	EMaterialDomain::MD_LightFunction
 };
 
 const TArray<EBlendMode> UDynamicMaterialModelEditorOnlyData::SupportedBlendModes =
@@ -242,33 +243,38 @@ void UDynamicMaterialModelEditorOnlyData::AssignPropertyAlphaValues()
 
 void UDynamicMaterialModelEditorOnlyData::OnDomainChanged()
 {
-	if (Domain == EMaterialDomain::MD_PostProcess)
+	switch (Domain)
 	{
-		const FDMUpdateGuard Guard;
+		case EMaterialDomain::MD_PostProcess:
+		case EMaterialDomain::MD_LightFunction:
+		{
+			const FDMUpdateGuard Guard;
 
-		// Post process only supports emissive.
-		UDMMaterialSlot* BaseColorSlot = GetSlotForMaterialProperty(EDMMaterialPropertyType::BaseColor);
-		UDMMaterialSlot* EmissiveSlot = GetSlotForMaterialProperty(EDMMaterialPropertyType::EmissiveColor);
+			// Post process only supports emissive.
+			UDMMaterialSlot* BaseColorSlot = GetSlotForMaterialProperty(EDMMaterialPropertyType::BaseColor);
+			UDMMaterialSlot* EmissiveSlot = GetSlotForMaterialProperty(EDMMaterialPropertyType::EmissiveColor);
 
-		if (!EmissiveSlot)
-		{
-			if (BaseColorSlot)
-			{
-				EnsureSwapSlotMaterialProperty(EDMMaterialPropertyType::BaseColor, EDMMaterialPropertyType::EmissiveColor);
-			}
-			else
+			if (!EmissiveSlot)
 			{
-				AddSlotForMaterialProperty(EDMMaterialPropertyType::EmissiveColor);
+				if (BaseColorSlot)
+				{
+					EnsureSwapSlotMaterialProperty(EDMMaterialPropertyType::BaseColor, EDMMaterialPropertyType::EmissiveColor);
+				}
+				else
+				{
+					AddSlotForMaterialProperty(EDMMaterialPropertyType::EmissiveColor);
+				}
 			}
+
+			SetShadingModel(EDMMaterialShadingModel::Unlit);
+			SetBlendMode(EBlendMode::BLEND_Opaque);
+			break;
 		}
 
-		SetShadingModel(EDMMaterialShadingModel::Unlit);
-		SetBlendMode(EBlendMode::BLEND_Opaque);
-	}
-	else if (Domain == EMaterialDomain::MD_DeferredDecal)
-	{
-		SetShadingModel(EDMMaterialShadingModel::DefaultLit);
-		SetBlendMode(EBlendMode::BLEND_Translucent);
+		case EMaterialDomain::MD_DeferredDecal:
+			SetShadingModel(EDMMaterialShadingModel::DefaultLit);
+			SetBlendMode(EBlendMode::BLEND_Translucent);
+			break;
 	}
 
 	RequestMaterialBuild();
@@ -450,7 +456,7 @@ void UDynamicMaterialModelEditorOnlyData::BuildMaterial(bool bInDirtyAssets)
 		}
 	);
 
-	if (Domain != EMaterialDomain::MD_PostProcess)
+	if (Domain != EMaterialDomain::MD_PostProcess && Domain != EMaterialDomain::MD_LightFunction)
 	{
 		/**
 		 * Generate opacity input based on base/emissive if it doesn't already have an input.
@@ -1262,7 +1268,8 @@ UDMMaterialSlot* UDynamicMaterialModelEditorOnlyData::AddSlot()
 
 UDMMaterialSlot* UDynamicMaterialModelEditorOnlyData::AddSlotForMaterialProperty(EDMMaterialPropertyType InType)
 {
-	if (InType == EDMMaterialPropertyType::EmissiveColor && Domain == EMaterialDomain::MD_PostProcess)
+	if (InType != EDMMaterialPropertyType::EmissiveColor 
+		&& (Domain == EMaterialDomain::MD_PostProcess || Domain == EMaterialDomain::MD_LightFunction))
 	{
 		return nullptr;
 	}
diff --git a/Engine/Plugins/Experimental/DynamicMaterial/Source/DynamicMaterialEditor/Public/Model/DynamicMaterialModelEditorOnlyData.h b/Engine/Plugins/Experimental/DynamicMaterial/Source/DynamicMaterialEditor/Public/Model/DynamicMaterialModelEditorOnlyData.h
--- a/Engine/Plugins/Experimental/DynamicMaterial/Source/DynamicMaterialEditor/Public/Model/DynamicMaterialModelEditorOnlyData.h
+++ b/Engine/Plugins/Experimental/DynamicMaterial/Source/DynamicMaterialEditor/Public/Model/DynamicMaterialModelEditorOnlyData.h
@@ -261,7 +261,7 @@ protected:
 	EDMState State;
 
 	UPROPERTY(EditInstanceOnly, BlueprintReadOnly, Category = "Material Designer", 
-		meta = (ValidEnumValues = "MD_Surface,MD_PostProcess,MD_DeferredDecal"))
+		meta = (ValidEnumValues = "MD_Surface,MD_PostProcess,MD_DeferredDecal,MD_LightFunction"))
 	TEnumAsByte<EMaterialDomain> Domain;
 
 	UPROPERTY(EditInstanceOnly, BlueprintReadOnly, Category = "Material Designer",
