Material Designer
- Textures (multiple) and texture sets can now be dropped directly into the wizard.

diff --git a/Engine/Plugins/Experimental/DynamicMaterial/Source/DynamicMaterialEditor/Private/UI/Widgets/SDMMaterialDesigner.cpp b/Engine/Plugins/Experimental/DynamicMaterial/Source/DynamicMaterialEditor/Private/UI/Widgets/SDMMaterialDesigner.cpp
--- a/Engine/Plugins/Experimental/DynamicMaterial/Source/DynamicMaterialEditor/Private/UI/Widgets/SDMMaterialDesigner.cpp
+++ b/Engine/Plugins/Experimental/DynamicMaterial/Source/DynamicMaterialEditor/Private/UI/Widgets/SDMMaterialDesigner.cpp
@@ -361,7 +361,10 @@ bool SDMMaterialDesigner::OnAssetDraggedOver(TArrayView<FAssetData> InAssets)
 	const bool bIsEditor = Content.IsValid()
 		&& Content->GetWidgetClass().GetWidgetType() == SDMMaterialEditor::StaticWidgetClass().GetWidgetType();
 
-	if (bIsEditor)
+	const bool bIsWizard = Content.IsValid()
+		&& Content->GetWidgetClass().GetWidgetType() == SDMMaterialWizard::StaticWidgetClass().GetWidgetType();
+
+	if (bIsEditor || bIsWizard)
 	{
 		AllowedClasses.Add(UDMTextureSet::StaticClass());
 	}
@@ -391,7 +394,7 @@ bool SDMMaterialDesigner::OnAssetDraggedOver(TArrayView<FAssetData> InAssets)
 		}
 	}
 
-	if (bIsEditor && DroppedTextures.Num() > 1)
+	if ((bIsEditor || bIsWizard) && DroppedTextures.Num() > 1)
 	{
 		return true;
 	}
@@ -403,6 +406,12 @@ void SDMMaterialDesigner::OnAssetsDropped(const FDragDropEvent& InDragDropEvent,
 {
 	TArray<FAssetData> DroppedTextures;
 
+	const bool bIsEditor = Content.IsValid()
+		&& Content->GetWidgetClass().GetWidgetType() == SDMMaterialEditor::StaticWidgetClass().GetWidgetType();
+
+	const bool bIsWizard = Content.IsValid()
+		&& Content->GetWidgetClass().GetWidgetType() == SDMMaterialWizard::StaticWidgetClass().GetWidgetType();
+
 	for (const FAssetData& Asset : InAssets)
 	{
 		UClass* AssetClass = Asset.GetClass(EResolveClass::Yes);
@@ -439,20 +448,30 @@ void SDMMaterialDesigner::OnAssetsDropped(const FDragDropEvent& InDragDropEvent,
 		}
 		else if (AssetClass->IsChildOf(UDMTextureSet::StaticClass()))
 		{
-			if (Content.IsValid() && Content->GetWidgetClass().GetWidgetType() == SDMMaterialEditor::StaticWidgetClass().GetWidgetType())
+			if (bIsEditor)
 			{
 				StaticCastSharedPtr<SDMMaterialEditor>(Content)->HandleDrop_TextureSet(Cast<UDMTextureSet>(Asset.GetAsset()));
 				return;
 			}
+
+			if (bIsWizard)
+			{
+				StaticCastSharedPtr<SDMMaterialWizard>(Content)->HandleDrop_TextureSet(Cast<UDMTextureSet>(Asset.GetAsset()));
+				return;
+			}
 		}
 	}
 
 	if (DroppedTextures.Num() > 1)
 	{
-		if (Content.IsValid() && Content->GetWidgetClass().GetWidgetType() == SDMMaterialEditor::StaticWidgetClass().GetWidgetType())
+		if (bIsEditor)
 		{
 			StaticCastSharedPtr<SDMMaterialEditor>(Content)->HandleDrop_CreateTextureSet(DroppedTextures);
 		}
+		else if (bIsWizard)
+		{
+			StaticCastSharedPtr<SDMMaterialWizard>(Content)->HandleDrop_CreateTextureSet(DroppedTextures);
+		}
 	}
 }
 
diff --git a/Engine/Plugins/Experimental/DynamicMaterial/Source/DynamicMaterialEditor/Private/UI/Widgets/SDMMaterialWizard.cpp b/Engine/Plugins/Experimental/DynamicMaterial/Source/DynamicMaterialEditor/Private/UI/Widgets/SDMMaterialWizard.cpp
--- a/Engine/Plugins/Experimental/DynamicMaterial/Source/DynamicMaterialEditor/Private/UI/Widgets/SDMMaterialWizard.cpp
+++ b/Engine/Plugins/Experimental/DynamicMaterial/Source/DynamicMaterialEditor/Private/UI/Widgets/SDMMaterialWizard.cpp
@@ -6,6 +6,7 @@
 #include "AssetTextFilter.h"
 #include "AssetThumbnail.h"
 #include "DMDefs.h"
+#include "DMTextureSetBlueprintFunctionLibrary.h"
 #include "DynamicMaterialEditorSettings.h"
 #include "DynamicMaterialEditorStyle.h"
 #include "DynamicMaterialModule.h"
@@ -127,6 +128,88 @@ UDynamicMaterialModel* SDMMaterialWizard::GetMaterialModel() const
 	return MaterialModelWeak.Get();
 }
 
+void SDMMaterialWizard::HandleDrop_CreateTextureSet(const TArray<FAssetData>& InTextureAssets)
+{
+	if (InTextureAssets.Num() < 2)
+	{
+		return;
+	}
+
+	UDMTextureSetBlueprintFunctionLibrary::CreateTextureSetFromAssetsInteractive(
+		InTextureAssets,
+		FDMTextureSetBuilderOnComplete::CreateSPLambda(
+			this,
+			[this](UDMTextureSet* InTextureSet, bool bInWasAccepted)
+			{
+				if (bInWasAccepted)
+				{
+					HandleDrop_TextureSet(InTextureSet);
+				}
+			}
+		)
+	);
+}
+
+void SDMMaterialWizard::HandleDrop_TextureSet(UDMTextureSet* InTextureSet)
+{
+	if (!InTextureSet)
+	{
+		return;
+	}
+
+	UDynamicMaterialModel* MaterialModel = GetMaterialModel();
+
+	if (!MaterialModel)
+	{
+		if (!MaterialObjectProperty.IsSet())
+		{
+			return;
+		}
+
+		if (!MaterialObjectProperty->IsValid())
+		{
+			return;
+		}
+
+		UDynamicMaterialInstance* Material = MaterialObjectProperty->GetMaterial();
+
+		if (!Material)
+		{
+			MaterialModel = UDMMaterialInstanceFunctionLibrary::CreateMaterialInObject(*MaterialObjectProperty);
+		}
+		else
+		{
+			UDynamicMaterialModelBase* MaterialModelBase = Material->GetMaterialModelBase();
+
+			if (!MaterialModelBase)
+			{
+				return;
+			}
+
+			MaterialModel = Cast<UDynamicMaterialModel>(MaterialModelBase);
+		}
+
+		if (!MaterialModel)
+		{
+			return;
+		}
+	}
+
+	UDynamicMaterialModelEditorOnlyData* EditorOnlyData = UDynamicMaterialModelEditorOnlyData::Get(MaterialModel);
+
+	if (!EditorOnlyData)
+	{
+		return;
+	}
+
+	FDMScopedUITransaction Transaction(LOCTEXT("DropTextureSet", "Drop Texture Set"));
+
+	EditorOnlyData->Modify();
+	EditorOnlyData->AddTextureSet(InTextureSet, /* Replace */ true);
+
+	OpenMaterialInEditor();
+}
+
 TSharedRef<SWidget> SDMMaterialWizard::CreateLayout()
 {
 	using namespace UE::DynamicMaterialDesigner::Private;
@@ -580,16 +663,22 @@ void SDMMaterialWizard::OpenMaterialInEditor()
 		return;
 	}
 
-	DesignerWidget->Empty();
-
 	if (MaterialObjectProperty.IsSet())
 	{
-		DesignerWidget->OpenObjectMaterialProperty(MaterialObjectProperty.GetValue());
+		if (DesignerWidget->OpenObjectMaterialProperty(*MaterialObjectProperty))
+		{
+			return;
+		}
 	}
-	else if (UDynamicMaterialModel* MaterialModel = GetMaterialModel())
+	else if (UDynamicMaterialModel* MaterialModel = MaterialModelWeak.Get())
 	{
-		DesignerWidget->OpenMaterialModelBase(MaterialModel);
+		if (DesignerWidget->OpenMaterialModelBase(MaterialModel))
+		{
+			return;
+		}
 	}
+
+	DesignerWidget->Empty();
 }
 
 ECheckBoxState SDMMaterialWizard::IsModeSelected(EDMMaterialWizardMode InMode) const
diff --git a/Engine/Plugins/Experimental/DynamicMaterial/Source/DynamicMaterialEditor/Private/UI/Widgets/SDMMaterialWizard.h b/Engine/Plugins/Experimental/DynamicMaterial/Source/DynamicMaterialEditor/Private/UI/Widgets/SDMMaterialWizard.h
--- a/Engine/Plugins/Experimental/DynamicMaterial/Source/DynamicMaterialEditor/Private/UI/Widgets/SDMMaterialWizard.h
+++ b/Engine/Plugins/Experimental/DynamicMaterial/Source/DynamicMaterialEditor/Private/UI/Widgets/SDMMaterialWizard.h
@@ -20,8 +20,10 @@ class SDMMaterialDesigner;
 class SToolInputAssetPicker;
 class SWidget;
 class SWidgetSwitcher;
+class UDMTextureSet;
 class UDynamicMaterialModel;
 enum class ECheckBoxState : uint8;
+struct FAssetData;
 struct FContentBrowserItem;
 
 class SDMMaterialWizard : public SCompoundWidget
@@ -45,6 +47,10 @@ public:
 
 	UDynamicMaterialModel* GetMaterialModel() const;
 
+	void HandleDrop_CreateTextureSet(const TArray<FAssetData>& InTextureAssets);
+
+	void HandleDrop_TextureSet(UDMTextureSet* InTextureSet);
+
 protected:
 	enum EDMMaterialWizardMode : uint8
 	{
