Reconciled offline work

diff --git a/Engine/Plugins/Experimental/DynamicMaterial/Source/DynamicMaterialEditor/Private/UI/Widgets/Editor/SlotEditor/SDMMaterialLayerBlendMode.cpp b/Engine/Plugins/Experimental/DynamicMaterial/Source/DynamicMaterialEditor/Private/UI/Widgets/Editor/SlotEditor/SDMMaterialLayerBlendMode.cpp
--- a/Engine/Plugins/Experimental/DynamicMaterial/Source/DynamicMaterialEditor/Private/UI/Widgets/Editor/SlotEditor/SDMMaterialLayerBlendMode.cpp
+++ b/Engine/Plugins/Experimental/DynamicMaterial/Source/DynamicMaterialEditor/Private/UI/Widgets/Editor/SlotEditor/SDMMaterialLayerBlendMode.cpp
@@ -35,6 +35,7 @@
 #include "ToolMenu.h"
 #include "ToolMenuDelegates.h"
 #include "ToolMenus.h"
+#include "UI/Widgets/Editor/SDMMaterialSlotEditor.h"
 #include "UI/Widgets/Editor/SlotEditor/SDMMaterialSlotLayerItem.h"
 #include "UI/Widgets/Editor/SlotEditor/SDMMaterialSlotLayerView.h"
 #include "Widgets/Input/SComboButton.h"
@@ -300,6 +301,17 @@ void SDMMaterialLayerBlendMode::OnBlendModeSelected(UClass* InBlendClass)
 	FDMScopedUITransaction Transaction(LOCTEXT("SetStageBlendMode", "Set Blend Mode"));
 	BaseStage->Modify();
 	BaseStage->ChangeSource<UDMMaterialStageBlend>(InBlendClass);
+
+	if (TSharedPtr<SDMMaterialSlotLayerItem> LayerItem = LayerItemWidgetWeak.Pin())
+	{
+		if (TSharedPtr<SDMMaterialSlotLayerView> LayerView = LayerItem->GetSlotLayerView())
+		{
+			if (TSharedPtr<SDMMaterialSlotEditor> SlotEditorWidget = LayerView->GetSlotEditorWidget())
+			{
+				SlotEditorWidget->InvalidateSlotSettings();
+			}
+		}
+	}
 }
 
 bool SDMMaterialLayerBlendMode::CanSelectBlendMode(UClass* InBlendClass)
