Material Designer
- Fix crash on bad rename on material replace.

diff --git a/Engine/Plugins/Experimental/DynamicMaterial/Source/DynamicMaterialEditor/Private/Utils/DMMaterialModelFunctionLibrary.cpp b/Engine/Plugins/Experimental/DynamicMaterial/Source/DynamicMaterialEditor/Private/Utils/DMMaterialModelFunctionLibrary.cpp
--- a/Engine/Plugins/Experimental/DynamicMaterial/Source/DynamicMaterialEditor/Private/Utils/DMMaterialModelFunctionLibrary.cpp
+++ b/Engine/Plugins/Experimental/DynamicMaterial/Source/DynamicMaterialEditor/Private/Utils/DMMaterialModelFunctionLibrary.cpp
@@ -497,7 +497,8 @@ bool UDMMaterialModelFunctionLibrary::CreateModelInstanceInMaterial(UDynamicMate
 		CurrentName = CurrentModel->GetName();
 
 		const FString NewName = CurrentName + TEXT("_OLD");
-		CurrentModel->Rename(*NewName, GetTransientPackage(), UE::DynamicMaterial::RenameFlags);
+		const FName NewNameUnique = MakeUniqueObjectName(GetTransientPackage(), CurrentModel->GetClass(), *NewName);
+		CurrentModel->Rename(*NewNameUnique.ToString(), GetTransientPackage(), UE::DynamicMaterial::RenameFlags);
 
 		InToInstance->SetMaterialModel(nullptr);
 	}
