Material Designer
- Fixed a crash when creating a material designer instance from the details on a Motion Design shape.

diff --git a/Engine/Plugins/Experimental/Avalanche/Source/AvalancheEditor/Private/MaterialDesigner/AvaMaterialDesignerExtension.cpp b/Engine/Plugins/Experimental/Avalanche/Source/AvalancheEditor/Private/MaterialDesigner/AvaMaterialDesignerExtension.cpp
--- a/Engine/Plugins/Experimental/Avalanche/Source/AvalancheEditor/Private/MaterialDesigner/AvaMaterialDesignerExtension.cpp
+++ b/Engine/Plugins/Experimental/Avalanche/Source/AvalancheEditor/Private/MaterialDesigner/AvaMaterialDesignerExtension.cpp
@@ -65,7 +65,7 @@ bool FAvaMaterialDesignerExtension::IsDynamicMaterialModelValid(UDynamicMaterial
 	{
 		if (UDynamicMaterialInstance* const MeshInstance = Cast<UDynamicMaterialInstance>(DynamicMesh->GetMaterial(MeshIndex)))
 		{
-			if (MeshInstance->GetMaterialModel() == InMaterialModel)
+			if (MeshInstance->GetMaterialModelBase() == InMaterialModel)
 			{
 				return true;
 			}
diff --git a/Engine/Plugins/Experimental/Avalanche/Source/AvalancheShapes/Private/DynamicMeshes/AvaShapeDynMeshBase.cpp b/Engine/Plugins/Experimental/Avalanche/Source/AvalancheShapes/Private/DynamicMeshes/AvaShapeDynMeshBase.cpp
--- a/Engine/Plugins/Experimental/Avalanche/Source/AvalancheShapes/Private/DynamicMeshes/AvaShapeDynMeshBase.cpp
+++ b/Engine/Plugins/Experimental/Avalanche/Source/AvalancheShapes/Private/DynamicMeshes/AvaShapeDynMeshBase.cpp
@@ -27,6 +27,7 @@
 #include "Subsystems/ActorModifierCoreSubsystem.h"
 
 #if WITH_EDITOR
+#include "DMObjectMaterialProperty.h"
 #include "Editor.h"
 #include "Elements/Actor/ActorElementEditorViewportInteractionCustomization.h"
 #include "IDynamicMaterialEditorModule.h"
@@ -1197,8 +1198,10 @@ void UAvaShapeDynamicMeshBase::OnMaterialTypeChanged(int32 MaterialIndex)
 					SetMaterial(MaterialIndex, NewInstance);
 
 #if WITH_EDITOR
+					const FDMObjectMaterialProperty ObjectMaterialProperty = FDMObjectMaterialProperty(GetShapeMeshComponent(), MaterialIndex);
+
 					const IDynamicMaterialEditorModule& MaterialDesignerModule = IDynamicMaterialEditorModule::Get();
-					MaterialDesignerModule.OpenMaterial(NewInstance, GetWorld(), /* Invoke Tab */ true);
+					MaterialDesignerModule.OpenMaterialObjectProperty(ObjectMaterialProperty, GetWorld(), /* Invoke Tab */ true);
 #endif
 				}
 				break;
diff --git a/Engine/Plugins/Experimental/DynamicMaterial/Source/DynamicMaterialEditor/Private/UI/Widgets/SDMMaterialWizard.cpp b/Engine/Plugins/Experimental/DynamicMaterial/Source/DynamicMaterialEditor/Private/UI/Widgets/SDMMaterialWizard.cpp
--- a/Engine/Plugins/Experimental/DynamicMaterial/Source/DynamicMaterialEditor/Private/UI/Widgets/SDMMaterialWizard.cpp
+++ b/Engine/Plugins/Experimental/DynamicMaterial/Source/DynamicMaterialEditor/Private/UI/Widgets/SDMMaterialWizard.cpp
@@ -845,35 +845,28 @@ void SDMMaterialWizard::SelectTemplate(UDynamicMaterialModel* InTemplateModel)
 		return;
 	}
 
-	if (UDynamicMaterialModel* MaterialModel = GetMaterialModel())
+	if (MaterialObjectProperty.IsSet())
 	{
-		if (UDynamicMaterialInstance* Instance = MaterialModel->GetDynamicMaterialInstance())
+		if (MaterialObjectProperty->IsValid())
 		{
-			CreateDynamicMaterialInInstance(InTemplateModel, Instance);
+			CreateNewDynamicInstanceInActor(InTemplateModel, MaterialObjectProperty.GetValue());
 		}
 		else
 		{
-			UE::DynamicMaterialEditor::Private::LogError(TEXT("Unable to find material instance to create new dynamic material in."));
+			UE::DynamicMaterialEditor::Private::LogError(TEXT("Invalid actor property to create new dynamic material in."));
 
 			DesignerWidget->ShowSelectPrompt();
 		}
 	}
-	else if (MaterialObjectProperty.IsSet())
+	else if (UDynamicMaterialModel* MaterialModel = GetMaterialModel())
 	{
-		if (MaterialObjectProperty->IsValid())
+		if (UDynamicMaterialInstance* Instance = MaterialModel->GetDynamicMaterialInstance())
 		{
-			if (UDynamicMaterialInstance* Instance = MaterialObjectProperty->GetMaterial())
-			{
-				CreateDynamicMaterialInInstance(InTemplateModel, Instance);
-			}
-			else
-			{
-				CreateNewDynamicInstanceInActor(InTemplateModel, MaterialObjectProperty.GetValue());
-			}
+			CreateDynamicMaterialInInstance(InTemplateModel, Instance);
 		}
 		else
 		{
-			UE::DynamicMaterialEditor::Private::LogError(TEXT("Invalid actor property to create new dynamic material in."));
+			UE::DynamicMaterialEditor::Private::LogError(TEXT("Unable to find material instance to create new dynamic material in."));
 
 			DesignerWidget->ShowSelectPrompt();
 		}
