Material Designer
- Undoing after changing the opened material designer asset will no longer cause a crash.

diff --git a/Engine/Plugins/Experimental/Avalanche/Source/AvalancheMedia/Private/Playable/AvaPlayableGroup.cpp b/Engine/Plugins/Experimental/Avalanche/Source/AvalancheMedia/Private/Playable/AvaPlayableGroup.cpp
--- a/Engine/Plugins/Experimental/Avalanche/Source/AvalancheMedia/Private/Playable/AvaPlayableGroup.cpp
+++ b/Engine/Plugins/Experimental/Avalanche/Source/AvalancheMedia/Private/Playable/AvaPlayableGroup.cpp
@@ -435,7 +435,7 @@ void UAvaPlayableGroup::ForEachPlayableTransition(TFunctionRef<bool(UAvaPlayable
 
 UAvaPlayableGroup* UAvaPlayableGroup::FindPlayableGroupForWorld(const UWorld* InWorld, bool bInFallbackToGlobalSearch)
 {
-	if (!InWorld)
+	if (!InWorld || !InWorld->PersistentLevel)
 	{
 		return nullptr;
 	}
diff --git a/Engine/Plugins/Experimental/DynamicMaterial/Source/DynamicMaterialEditor/Private/UI/Widgets/Editor/DMMaterialPreviewViewportClient.cpp b/Engine/Plugins/Experimental/DynamicMaterial/Source/DynamicMaterialEditor/Private/UI/Widgets/Editor/DMMaterialPreviewViewportClient.cpp
--- a/Engine/Plugins/Experimental/DynamicMaterial/Source/DynamicMaterialEditor/Private/UI/Widgets/Editor/DMMaterialPreviewViewportClient.cpp
+++ b/Engine/Plugins/Experimental/DynamicMaterial/Source/DynamicMaterialEditor/Private/UI/Widgets/Editor/DMMaterialPreviewViewportClient.cpp
@@ -8,10 +8,12 @@
 #include "UnrealWidget.h"
 
 FDMMaterialPreviewViewportClient::FDMMaterialPreviewViewportClient(const TSharedRef<SDMMaterialPreview>& InPreviewWidget,
-	FAdvancedPreviewScene& InPreviewScene)
-	: FEditorViewportClient(nullptr, &InPreviewScene, StaticCastSharedRef<SEditorViewport>(InPreviewWidget))
+	FAdvancedPreviewScene& InPreviewScene, TSharedRef<FEditorModeTools> InPreviewModeTools)
+	: FEditorViewportClient(&*InPreviewModeTools, &InPreviewScene,
+		StaticCastSharedRef<SEditorViewport>(InPreviewWidget))
 {
 	PreviewWidget = InPreviewWidget;
+	PreviewModeTools = InPreviewModeTools;
 
 	bDrawAxes = false;
 
@@ -124,7 +126,10 @@ void FDMMaterialPreviewViewportClient::Tick(float InDeltaSeconds)
 	// Tick the preview scene world.
 	if (UWorld* World = PreviewScene->GetWorld())
 	{
-		World->Tick(LEVELTICK_All, InDeltaSeconds);
+		if (IsValid(World))
+		{
+			World->Tick(LEVELTICK_All, InDeltaSeconds);
+		}
 	}
 }
 
diff --git a/Engine/Plugins/Experimental/DynamicMaterial/Source/DynamicMaterialEditor/Private/UI/Widgets/Editor/DMMaterialPreviewViewportClient.h b/Engine/Plugins/Experimental/DynamicMaterial/Source/DynamicMaterialEditor/Private/UI/Widgets/Editor/DMMaterialPreviewViewportClient.h
--- a/Engine/Plugins/Experimental/DynamicMaterial/Source/DynamicMaterialEditor/Private/UI/Widgets/Editor/DMMaterialPreviewViewportClient.h
+++ b/Engine/Plugins/Experimental/DynamicMaterial/Source/DynamicMaterialEditor/Private/UI/Widgets/Editor/DMMaterialPreviewViewportClient.h
@@ -7,13 +7,15 @@
 #include "Templates/SharedPointer.h"
 
 class FAdvancedPreviewScene;
+class FEditorModeTools;
 class SDMMaterialPreview;
 
 /** Based on FMaterialEditorViewportClient (private) */
 class FDMMaterialPreviewViewportClient : public FEditorViewportClient
 {
 public:
-	FDMMaterialPreviewViewportClient(const TSharedRef<SDMMaterialPreview>& InPreviewWidget, FAdvancedPreviewScene& InPreviewScene);
+	FDMMaterialPreviewViewportClient(const TSharedRef<SDMMaterialPreview>& InPreviewWidget, FAdvancedPreviewScene& InPreviewScene,
+		TSharedRef<FEditorModeTools> InPreviewModeTools);
 
 	/**
 	* Focuses the viewport to the center of the bounding box/sphere ensuring that the entire bounds are in view
@@ -35,4 +37,5 @@ public:
 private:
 	TWeakPtr<SDMMaterialPreview> PreviewWidget;
 	FAdvancedPreviewScene* AdvancedPreviewScene;
+	TSharedPtr<FEditorModeTools> PreviewModeTools;
 };
diff --git a/Engine/Plugins/Experimental/DynamicMaterial/Source/DynamicMaterialEditor/Private/UI/Widgets/Editor/SDMMaterialPreview.cpp b/Engine/Plugins/Experimental/DynamicMaterial/Source/DynamicMaterialEditor/Private/UI/Widgets/Editor/SDMMaterialPreview.cpp
--- a/Engine/Plugins/Experimental/DynamicMaterial/Source/DynamicMaterialEditor/Private/UI/Widgets/Editor/SDMMaterialPreview.cpp
+++ b/Engine/Plugins/Experimental/DynamicMaterial/Source/DynamicMaterialEditor/Private/UI/Widgets/Editor/SDMMaterialPreview.cpp
@@ -9,6 +9,7 @@
 #include "DynamicMaterialEditorSettings.h"
 #include "DynamicMaterialModule.h"
 #include "Editor/UnrealEdEngine.h"
+#include "EditorModeManager.h"
 #include "Engine/Engine.h"
 #include "Engine/PostProcessVolume.h"
 #include "Engine/World.h"
@@ -559,7 +560,7 @@ void SDMMaterialPreview::OpenMaterialPreviewTab()
 
 TSharedRef<FEditorViewportClient> SDMMaterialPreview::MakeEditorViewportClient()
 {
-	EditorViewportClient = MakeShared<FDMMaterialPreviewViewportClient>(SharedThis(this), *PreviewScene.Get());
+	EditorViewportClient = MakeShared<FDMMaterialPreviewViewportClient>(SharedThis(this), *PreviewScene.Get(), MakeShared<FEditorModeTools>());
 	EditorViewportClient->SetViewLocation(FVector::ZeroVector);
 	EditorViewportClient->SetViewRotation(FRotator(-15.0f, -90.0f, 0.0f));
 	EditorViewportClient->SetViewLocationForOrbiting(FVector::ZeroVector);
