Material Designer
- When updating input maps, if the input map will be unchanged, no update event will be triggered.

diff --git a/Engine/Plugins/Experimental/DynamicMaterial/Source/DynamicMaterialEditor/Private/Components/DMMaterialStage.cpp b/Engine/Plugins/Experimental/DynamicMaterial/Source/DynamicMaterialEditor/Private/Components/DMMaterialStage.cpp
--- a/Engine/Plugins/Experimental/DynamicMaterial/Source/DynamicMaterialEditor/Private/Components/DMMaterialStage.cpp
+++ b/Engine/Plugins/Experimental/DynamicMaterial/Source/DynamicMaterialEditor/Private/Components/DMMaterialStage.cpp
@@ -18,6 +18,7 @@
 #include "Model/DynamicMaterialModel.h"
 #include "Model/DynamicMaterialModelEditorOnlyData.h"
 #include "UObject/Package.h"
+#include "Utils/DMPrivate.h"
 
 #define LOCTEXT_NAMESPACE "DMMaterialStage"
 
@@ -771,7 +772,7 @@ void UDMMaterialStage::UpdateInputMap(int32 InInputIdx, int32 InSourceIndex, int
 	check(InSourceIndex != FDMMaterialStageConnectorChannel::PREVIOUS_STAGE || InStageProperty != EDMMaterialPropertyType::None);
 	check(InSourceIndex >= 0 && InSourceIndex <= Inputs.Num());
 
-	VerifyAllInputMaps();
+	const bool bMadeChanges = VerifyAllInputMaps();
 	check(InputConnectionMap.IsValidIndex(InInputIdx));
 
 	check(Source);
@@ -823,24 +824,39 @@ void UDMMaterialStage::UpdateInputMap(int32 InInputIdx, int32 InSourceIndex, int
 		check(Throughput->CanInputConnectTo(InInputIdx, StageInputOutputConnectors[InOutputIdx], InOutputChannel));
 	}
 
+	TArray<FDMMaterialStageConnectorChannel>& Channels = InputConnectionMap[InInputIdx].Channels;
+	const int32 ChannelIndex = UE::DynamicMaterialEditor::Private::ChannelBitToChannelIndex(InInputChannel);
+	const FDMMaterialStageConnectorChannel NewChannel = {InSourceIndex, InStageProperty, InOutputIdx, InOutputChannel};
+
 	// Remove all mapping for this input and replace them with a whole channel mapping
-	if (InInputChannel == FDMMaterialStageConnectorChannel::WHOLE_CHANNEL)
+	if (ChannelIndex == 0)
 	{
-		InputConnectionMap[InInputIdx].Channels.Empty();
-		InputConnectionMap[InInputIdx].Channels.Add({InSourceIndex, InStageProperty, InOutputIdx, InOutputChannel});
+		if (!bMadeChanges && Channels.Num() == 1 && Channels[0] == NewChannel)
+		{
+			return;
+		}
+
+		Channels = {NewChannel};
 	}
 	// Add/replace the new channel-specific mapping
 	else
 	{
-		for (int32 NewIndex = InputConnectionMap.Num(); NewIndex < InInputChannel; ++NewIndex)
+		if (!bMadeChanges && Channels.IsValidIndex(ChannelIndex) && Channels[ChannelIndex] == NewChannel)
 		{
-			InputConnectionMap[InInputIdx].Channels.Add({FDMMaterialStageConnectorChannel::NO_SOURCE, EDMMaterialPropertyType::None, 0, FDMMaterialStageConnectorChannel::WHOLE_CHANNEL});
+			return;
+		}
+
+		if (!Channels.IsValidIndex(ChannelIndex))
+		{
+			const FDMMaterialStageConnectorChannel BlankChannel = {FDMMaterialStageConnectorChannel::NO_SOURCE, EDMMaterialPropertyType::None, 0, FDMMaterialStageConnectorChannel::WHOLE_CHANNEL};
+
+			for (int32 NewIndex = Channels.Num(); NewIndex <= ChannelIndex; ++NewIndex)
+			{
+				Channels.Add(BlankChannel);
+			}
 		}
 
-		InputConnectionMap[InInputIdx].Channels[InInputChannel].SourceIndex = InSourceIndex;
-		InputConnectionMap[InInputIdx].Channels[InInputChannel].MaterialProperty = InStageProperty;
-		InputConnectionMap[InInputIdx].Channels[InInputChannel].OutputIndex = InOutputIdx;
-		InputConnectionMap[InInputIdx].Channels[InInputChannel].OutputChannel = InOutputChannel;
+		Channels[ChannelIndex] = NewChannel;
 	}
 
 	RemoveUnusedInputs();
diff --git a/Engine/Plugins/Experimental/DynamicMaterial/Source/DynamicMaterialEditor/Private/Utils/DMPrivate.cpp b/Engine/Plugins/Experimental/DynamicMaterial/Source/DynamicMaterialEditor/Private/Utils/DMPrivate.cpp
--- a/Engine/Plugins/Experimental/DynamicMaterial/Source/DynamicMaterialEditor/Private/Utils/DMPrivate.cpp
+++ b/Engine/Plugins/Experimental/DynamicMaterial/Source/DynamicMaterialEditor/Private/Utils/DMPrivate.cpp
@@ -45,24 +45,49 @@ namespace UE::DynamicMaterialEditor::Private
 	{
 		switch (InChannelIndex)
 		{
-		case 0:
-			return FDMMaterialStageConnectorChannel::WHOLE_CHANNEL;
+			case 0:
+				return FDMMaterialStageConnectorChannel::WHOLE_CHANNEL;
 
-		case 1:
-			return FDMMaterialStageConnectorChannel::FIRST_CHANNEL;
+			case 1:
+				return FDMMaterialStageConnectorChannel::FIRST_CHANNEL;
 
-		case 2:
-			return FDMMaterialStageConnectorChannel::SECOND_CHANNEL;
+			case 2:
+				return FDMMaterialStageConnectorChannel::SECOND_CHANNEL;
 
-		case 3:
-			return FDMMaterialStageConnectorChannel::THIRD_CHANNEL;
+			case 3:
+				return FDMMaterialStageConnectorChannel::THIRD_CHANNEL;
 
-		case 4:
-			return FDMMaterialStageConnectorChannel::FOURTH_CHANNEL;
+			case 4:
+				return FDMMaterialStageConnectorChannel::FOURTH_CHANNEL;
 
-		default:
-			checkNoEntry();
-			return 0;
+			default:
+				checkNoEntry();
+				return 0;
+		}
+	}
+
+	int32 ChannelBitToChannelIndex(int32 InChannelBit)
+	{
+		switch (InChannelBit)
+		{
+			case FDMMaterialStageConnectorChannel::WHOLE_CHANNEL:
+				return 0;
+
+			case FDMMaterialStageConnectorChannel::FIRST_CHANNEL:
+				return 1;
+
+			case FDMMaterialStageConnectorChannel::SECOND_CHANNEL:
+				return 2;
+
+			case FDMMaterialStageConnectorChannel::THIRD_CHANNEL:
+				return 3;
+
+			case FDMMaterialStageConnectorChannel::FOURTH_CHANNEL:
+				return 4;
+
+			default:
+				checkNoEntry();
+				return 0;
 		}
 	}
 
diff --git a/Engine/Plugins/Experimental/DynamicMaterial/Source/DynamicMaterialEditor/Private/Utils/DMPrivate.h b/Engine/Plugins/Experimental/DynamicMaterial/Source/DynamicMaterialEditor/Private/Utils/DMPrivate.h
--- a/Engine/Plugins/Experimental/DynamicMaterial/Source/DynamicMaterialEditor/Private/Utils/DMPrivate.h
+++ b/Engine/Plugins/Experimental/DynamicMaterial/Source/DynamicMaterialEditor/Private/Utils/DMPrivate.h
@@ -30,6 +30,8 @@ namespace UE::DynamicMaterialEditor::Private
 	/** Converts 0,1,2,3,4 to 0,1,2,4,8 */
 	int32 ChannelIndexToChannelBit(int32 InChannelIndex);
 
+	int32 ChannelBitToChannelIndex(int32 InChannelBit);
+
 	bool IsCustomMaterialProperty(EDMMaterialPropertyType InMaterialProperty);
 
 	void LogError(const FString& InMessage, bool bInToast = true, const UObject* InSource = nullptr);
