Material Designer
- Added transaction to adding effects.

diff --git a/Engine/Plugins/Experimental/DynamicMaterial/Source/DynamicMaterialEditor/Private/UI/Menus/DMMaterialSlotLayerAddEffectMenus.cpp b/Engine/Plugins/Experimental/DynamicMaterial/Source/DynamicMaterialEditor/Private/UI/Menus/DMMaterialSlotLayerAddEffectMenus.cpp
--- a/Engine/Plugins/Experimental/DynamicMaterial/Source/DynamicMaterialEditor/Private/UI/Menus/DMMaterialSlotLayerAddEffectMenus.cpp
+++ b/Engine/Plugins/Experimental/DynamicMaterial/Source/DynamicMaterialEditor/Private/UI/Menus/DMMaterialSlotLayerAddEffectMenus.cpp
@@ -18,6 +18,7 @@
 #include "UI/Widgets/Editor/SlotEditor/SDMMaterialSlotLayerView.h"
 #include "UI/Widgets/SDMMaterialEditor.h"
 #include "Utils/DMMaterialEffectStackPresetSubsystem.h"
+#include "Utils/DMPrivate.h"
 #include "Widgets/Input/SEditableTextBox.h"
 #include "Widgets/Notifications/SNotificationList.h"
 #include "Widgets/SNullWidget.h"
@@ -155,6 +156,9 @@ void FDMMaterialSlotLayerAddEffectMenus::AddEffect(const FToolMenuContext& InCon
 		return;
 	}
 
+	FDMScopedUITransaction Transaction(LOCTEXT("AddEffectTransaction", "Add Effect"));
+	EffectStack->Modify();
+
 	EffectStack->AddEffect(EffectFunction);
 
 	if (!EditorWidget.IsValid())
