Material Designer
- Alpha-based effects now work correctly.

diff --git a/Engine/Plugins/Experimental/DynamicMaterial/Source/DynamicMaterialEditor/Private/Components/DMMaterialEffectFunction.cpp b/Engine/Plugins/Experimental/DynamicMaterial/Source/DynamicMaterialEditor/Private/Components/DMMaterialEffectFunction.cpp
--- a/Engine/Plugins/Experimental/DynamicMaterial/Source/DynamicMaterialEditor/Private/Components/DMMaterialEffectFunction.cpp
+++ b/Engine/Plugins/Experimental/DynamicMaterial/Source/DynamicMaterialEditor/Private/Components/DMMaterialEffectFunction.cpp
@@ -300,19 +300,23 @@ void UDMMaterialEffectFunction::ApplyTo(const TSharedRef<FDMMaterialBuildState>&
 			{
 				if (InOutLastExpressionOutputChannel != FDMMaterialStageConnectorChannel::WHOLE_CHANNEL)
 				{
-					LastInputExpressions[0] = InBuildState->GetBuildUtils().CreateExpressionBitMask(
+					UMaterialExpression* Mask = InBuildState->GetBuildUtils().CreateExpressionBitMask(
 						LastInputExpressions[0], 
 						InOutLastExpressionOutputIndex,
 						InOutLastExpressionOutputChannel
 					);
 
-					InOutStageExpressions.Add(LastInputExpressions[0]);
-				}
+					FunctionCall->FunctionInputs[InputIndex].Input.Connect(0, Mask);
 
-				FunctionCall->FunctionInputs[InputIndex].Input.Connect(
-					InOutLastExpressionOutputIndex, 
-					LastInputExpressions[InputIndex]
-				);
+					InOutStageExpressions.Add(Mask);
+				}
+				else
+				{
+					FunctionCall->FunctionInputs[InputIndex].Input.Connect(
+						InOutLastExpressionOutputIndex,
+						LastInputExpressions[InputIndex]
+					);
+				}
 			}
 			else
 			{
