Material Designer
- When textures are dragged into he viewport, it will no longer cause a crash if the actor and material is deleted during compilation.

diff --git a/Engine/Plugins/Experimental/DynamicMaterial/Source/DynamicMaterialEditor/Private/Model/DynamicMaterialModelEditorOnlyData.cpp b/Engine/Plugins/Experimental/DynamicMaterial/Source/DynamicMaterialEditor/Private/Model/DynamicMaterialModelEditorOnlyData.cpp
--- a/Engine/Plugins/Experimental/DynamicMaterial/Source/DynamicMaterialEditor/Private/Model/DynamicMaterialModelEditorOnlyData.cpp
+++ b/Engine/Plugins/Experimental/DynamicMaterial/Source/DynamicMaterialEditor/Private/Model/DynamicMaterialModelEditorOnlyData.cpp
@@ -573,7 +573,14 @@ void UDynamicMaterialModelEditorOnlyData::BuildMaterial(bool bInDirtyAssets)
 	 */
 	BuildState.Reset();
 
-	MaterialStats = UMaterialEditingLibrary::GetStatistics(MaterialModel->DynamicMaterial);
+	/**
+	 * GetStatistics can call the GC which could potentially delete the material under us.
+	 * Add a reference to it while getting statistics to prevent it being destroyed.
+	 */
+	{
+		TStrongObjectPtr<UObject> ScopeReference(MaterialModel->DynamicMaterial);
+		MaterialStats = UMaterialEditingLibrary::GetStatistics(MaterialModel->DynamicMaterial);
+	}
 
 	State = EDMState::Idle;
 
