Material Designer
- Changed the node arranger to account for long material titles.

diff --git a/Engine/Plugins/Experimental/DynamicMaterial/Source/DynamicMaterialEditor/Private/Model/DMMaterialNodeArranger.cpp b/Engine/Plugins/Experimental/DynamicMaterial/Source/DynamicMaterialEditor/Private/Model/DMMaterialNodeArranger.cpp
--- a/Engine/Plugins/Experimental/DynamicMaterial/Source/DynamicMaterialEditor/Private/Model/DMMaterialNodeArranger.cpp
+++ b/Engine/Plugins/Experimental/DynamicMaterial/Source/DynamicMaterialEditor/Private/Model/DMMaterialNodeArranger.cpp
@@ -59,6 +59,54 @@ void FDMMaterialNodeArranger::ArrangeNodes()
 	}
 }
 
+int32 FDMMaterialNodeArranger::GetNodeWidth(UMaterialExpression& InNode)
+{
+	constexpr int32 LetterWidth = 8; // It's not monospaced, but it's a good estimate.
+	constexpr int32 Padding = 50;
+
+	auto EstimateNodeWidthBasedOnTitle = [](const FString& InTitle)
+		{
+			if (!InTitle.IsEmpty())
+			{
+				return (InTitle.Len() * LetterWidth) + Padding;
+			}
+
+			return -1;
+		};
+
+	const int32 NodeWidth = InNode.GetWidth() * 2; // Because it's just not enough
+
+	// We have no graph node, so let's estimate.
+	if (InNode.HasAParameterName())
+	{
+		const FString ParameterName = InNode.GetParameterName().ToString();
+		const int32 Width = EstimateNodeWidthBasedOnTitle(ParameterName);
+
+		if (Width > 0)
+		{
+			return FMath::Max(NodeWidth, Width);
+		}
+	}
+
+	TArray<FString> Captions;
+	InNode.GetCaption(Captions);
+	int32 Width = 0;
+
+	for (const FString& Caption : Captions)
+	{
+		Width = FMath::Max(Width, EstimateNodeWidthBasedOnTitle(Caption));
+	}
+
+	if (Width > 0)
+	{
+		return FMath::Max(NodeWidth, Width);
+	}
+
+	Width = EstimateNodeWidthBasedOnTitle(InNode.GetClass()->GetDescription());
+
+	return FMath::Max(NodeWidth, Width);
+}
+
 void FDMMaterialNodeArranger::ArrangeMaterialInputNodes(UMaterialExpression* MaterialInputExpression)
 {
 	if (!MaterialInputExpression)
@@ -82,7 +130,7 @@ void FDMMaterialNodeArranger::ArrangeNode(TMap<UMaterialExpression*, FIntPoint>&
 	}
 
 	InOutNodeSize = {0, 0};
-	const FIntPoint ThisNodeSize = {InNode->GetWidth() * 2, InNode->GetHeight()};
+	const FIntPoint ThisNodeSize = {GetNodeWidth(*InNode), InNode->GetHeight()};
 	FIntPoint ChildOffsetStart = InOffsetStart;
 	ChildOffsetStart.X += ThisNodeSize.X + UE::DynamicMaterialEditor::BuildState::Private::SpaceBetweenNodes;
 
diff --git a/Engine/Plugins/Experimental/DynamicMaterial/Source/DynamicMaterialEditor/Private/Model/DMMaterialNodeArranger.h b/Engine/Plugins/Experimental/DynamicMaterial/Source/DynamicMaterialEditor/Private/Model/DMMaterialNodeArranger.h
--- a/Engine/Plugins/Experimental/DynamicMaterial/Source/DynamicMaterialEditor/Private/Model/DMMaterialNodeArranger.h
+++ b/Engine/Plugins/Experimental/DynamicMaterial/Source/DynamicMaterialEditor/Private/Model/DMMaterialNodeArranger.h
@@ -19,6 +19,8 @@ public:
 	void ArrangeNodes();
 
 protected:
+	static int32 GetNodeWidth(UMaterialExpression& InNode);
+
 	UMaterial* DynamicMaterial;
 	TMap<UMaterialExpression*, FIntPoint> NodePositions;
 	FIntPoint OffsetStart = {0, 0};
