Material Designer
- Adds a notification if a texture set adds a texture to a channel which is incompatible with the current material model.

diff --git a/Engine/Plugins/Experimental/DynamicMaterial/Source/DynamicMaterialEditor/Private/Components/DMMaterialProperty.cpp b/Engine/Plugins/Experimental/DynamicMaterial/Source/DynamicMaterialEditor/Private/Components/DMMaterialProperty.cpp
--- a/Engine/Plugins/Experimental/DynamicMaterial/Source/DynamicMaterialEditor/Private/Components/DMMaterialProperty.cpp
+++ b/Engine/Plugins/Experimental/DynamicMaterial/Source/DynamicMaterialEditor/Private/Components/DMMaterialProperty.cpp
@@ -91,7 +91,7 @@ void UDMMaterialProperty::SetEnabled(bool bInEnabled)
 
 FText UDMMaterialProperty::GetDescription() const
 {
-	return UE::DynamicMaterialEditor::Private::GetMaterialPropertyShortDisplayName(MaterialProperty);
+	return UE::DynamicMaterialEditor::Private::GetMaterialPropertyLongDisplayName(MaterialProperty);
 }
 
 bool UDMMaterialProperty::IsMaterialPin() const
diff --git a/Engine/Plugins/Experimental/DynamicMaterial/Source/DynamicMaterialEditor/Private/Model/DynamicMaterialModelEditorOnlyData.cpp b/Engine/Plugins/Experimental/DynamicMaterial/Source/DynamicMaterialEditor/Private/Model/DynamicMaterialModelEditorOnlyData.cpp
--- a/Engine/Plugins/Experimental/DynamicMaterial/Source/DynamicMaterialEditor/Private/Model/DynamicMaterialModelEditorOnlyData.cpp
+++ b/Engine/Plugins/Experimental/DynamicMaterial/Source/DynamicMaterialEditor/Private/Model/DynamicMaterialModelEditorOnlyData.cpp
@@ -41,6 +41,7 @@
 #include "DynamicMaterialModule.h"
 #include "Factories/MaterialFactoryNew.h"
 #include "FileHelpers.h"
+#include "Framework/Notifications/NotificationManager.h"
 #include "IAssetTools.h"
 #include "Material/DynamicMaterialInstance.h"
 #include "Materials/Material.h"
@@ -52,6 +53,7 @@
 #include "UObject/Package.h"
 #include "Utils/DMBuildRequestSubsystem.h"
 #include "Utils/DMUtils.h"
+#include "Widgets/Notifications/SNotificationList.h"
 
 #define LOCTEXT_NAMESPACE "MaterialDesignerModel"
 
@@ -649,6 +651,8 @@ bool UDynamicMaterialModelEditorOnlyData::AddTextureSet(UDMTextureSet* InTexture
 {
 	bool bMadeChange = false;
 
+	TArray<UDMMaterialProperty*> InvalidProperties;
+
 	for (const TPair<EDMTextureSetMaterialProperty, FDMMaterialTexture>& MaterialTexture : InTextureSet->GetTextures())
 	{
 		if (MaterialTexture.Value.Texture.IsNull())
@@ -675,6 +679,11 @@ bool UDynamicMaterialModelEditorOnlyData::AddTextureSet(UDMTextureSet* InTexture
 			Property->SetEnabled(true);
 		}
 
+		if (!Property->IsValidForModel(*this))
+		{
+			InvalidProperties.Add(Property);
+		}
+
 		UDMMaterialSlot* Slot = GetSlotForMaterialProperty(PropertyType);
 
 		if (!Slot)
@@ -809,6 +818,34 @@ bool UDynamicMaterialModelEditorOnlyData::AddTextureSet(UDMTextureSet* InTexture
 		Layer->Update(Layer, EDMUpdateType::Structure);
 	}
 
+	/**
+	 * This entire function should not be here. It should be something on the texture set.
+	 * This warning should be there.
+	 */
+	if (!InvalidProperties.IsEmpty())
+	{
+		const FText WarningFormat = LOCTEXT(
+			"AddTextureSetFormat", 
+			"The following channels contain textures but are not valid for the current material settings:\n\n{0}"""
+		);
+
+		const FText IndividualWarningFormat = LOCTEXT("AddTextureSetIndividualFormat", "- {0}\n");
+
+		TArray<FText> ErrorStrings;
+		ErrorStrings.Reserve(InvalidProperties.Num());
+
+		for (UDMMaterialProperty* Property : InvalidProperties)
+		{
+			ErrorStrings.Add(FText::Format(IndividualWarningFormat, Property->GetDescription()));
+		}
+
+		const FText FullWarning = FText::Format(WarningFormat, FText::Join(FText::GetEmpty(), ErrorStrings));
+
+		FNotificationInfo Info(FullWarning);
+		Info.ExpireDuration = 5.0f;
+		FSlateNotificationManager::Get().AddNotification(Info);
+	}
+
 	return bMadeChange;
 }
 
