Material Designer
- The save button is no longer visible when editing an asset that exists inside an actor rather than as an asset in the content browser.

diff --git a/Engine/Plugins/Experimental/DynamicMaterial/Source/DynamicMaterialEditor/Private/UI/Widgets/Editor/SDMToolBar.cpp b/Engine/Plugins/Experimental/DynamicMaterial/Source/DynamicMaterialEditor/Private/UI/Widgets/Editor/SDMToolBar.cpp
--- a/Engine/Plugins/Experimental/DynamicMaterial/Source/DynamicMaterialEditor/Private/UI/Widgets/Editor/SDMToolBar.cpp
+++ b/Engine/Plugins/Experimental/DynamicMaterial/Source/DynamicMaterialEditor/Private/UI/Widgets/Editor/SDMToolBar.cpp
@@ -358,10 +358,10 @@ void SDMToolBar::SetButtonVisibilities()
 		bIsDynamic = !MaterialModelBase->IsA<UDynamicMaterialModel>();
 	}
 
-	SaveButtonWidget->SetVisibility(CanSave() ? EVisibility::Visible : EVisibility::Collapsed);
-
 	if (bIsAsset)
 	{
+		SaveButtonWidget->SetVisibility(EVisibility::Visible);
+
 		AssetNameWidget->SetText(GetAssetName());
 		AssetNameWidget->SetToolTipText(GetAssetToolTip());
 		AssetNameWidget->SetVisibility(EVisibility::Visible);
@@ -370,6 +370,8 @@ void SDMToolBar::SetButtonVisibilities()
 	}
 	else
 	{
+		SaveButtonWidget->SetVisibility(EVisibility::Collapsed);
+
 		AssetNameWidget->SetText(FText::GetEmpty());
 		AssetNameWidget->SetToolTipText(FText::GetEmpty());
 		AssetNameWidget->SetVisibility(EVisibility::Collapsed);
@@ -767,16 +769,6 @@ FText SDMToolBar::GetAssetToolTip() const
 	return FText::GetEmpty();
 }
 
-bool SDMToolBar::CanSave() const
-{
-	if (UDynamicMaterialModelBase* MaterialModelBase = GetMaterialModelBase())
-	{
-		return !MaterialModelBase->GetTypedOuter<UWorld>();
-	}
-
-	return false;
-}
-
 const FSlateBrush* SDMToolBar::GetSaveIcon() const
 {
 	if (UPackage* Package = GetSaveablePackage(GetMaterialModelBase()))
diff --git a/Engine/Plugins/Experimental/DynamicMaterial/Source/DynamicMaterialEditor/Private/UI/Widgets/Editor/SDMToolBar.h b/Engine/Plugins/Experimental/DynamicMaterial/Source/DynamicMaterialEditor/Private/UI/Widgets/Editor/SDMToolBar.h
--- a/Engine/Plugins/Experimental/DynamicMaterial/Source/DynamicMaterialEditor/Private/UI/Widgets/Editor/SDMToolBar.h
+++ b/Engine/Plugins/Experimental/DynamicMaterial/Source/DynamicMaterialEditor/Private/UI/Widgets/Editor/SDMToolBar.h
@@ -96,8 +96,6 @@ protected:
 
 	FText GetAssetToolTip() const;
 
-	bool CanSave() const;
-
 	const FSlateBrush* GetSaveIcon() const;
 
 	FReply OnSaveClicked();
