Material Designer
- When automatically setting a texture to the mask stage, the channel override is now set correctly.
- Rejigged the slot editor and its sub-widgets so that events centre around the slot editor widget, instead of its sub-widgets.
- Undo now doesn't refresh the UI when it is not necessary.

diff --git a/Engine/Plugins/Experimental/DynamicMaterial/Source/DynamicMaterialEditor/Private/Components/DMMaterialStageBlend.cpp b/Engine/Plugins/Experimental/DynamicMaterial/Source/DynamicMaterialEditor/Private/Components/DMMaterialStageBlend.cpp
--- a/Engine/Plugins/Experimental/DynamicMaterial/Source/DynamicMaterialEditor/Private/Components/DMMaterialStageBlend.cpp
+++ b/Engine/Plugins/Experimental/DynamicMaterial/Source/DynamicMaterialEditor/Private/Components/DMMaterialStageBlend.cpp
@@ -412,6 +412,13 @@ void UDMMaterialStageBlend::Update(UDMMaterialComponent* InSource, EDMUpdateType
 	PullBaseChannelOverride();
 }
 
+void UDMMaterialStageBlend::PostEditUndo()
+{
+	Super::PostEditUndo();
+
+	PullBaseChannelOverride();
+}
+
 void UDMMaterialStageBlend::NotifyPostChange(const FPropertyChangedEvent& InPropertyChangedEvent, FEditPropertyChain* InPropertyThatChanged)
 {
 	if (!IsComponentValid())
diff --git a/Engine/Plugins/Experimental/DynamicMaterial/Source/DynamicMaterialEditor/Private/Components/DMMaterialStageThroughputLayerBlend.cpp b/Engine/Plugins/Experimental/DynamicMaterial/Source/DynamicMaterialEditor/Private/Components/DMMaterialStageThroughputLayerBlend.cpp
--- a/Engine/Plugins/Experimental/DynamicMaterial/Source/DynamicMaterialEditor/Private/Components/DMMaterialStageThroughputLayerBlend.cpp
+++ b/Engine/Plugins/Experimental/DynamicMaterial/Source/DynamicMaterialEditor/Private/Components/DMMaterialStageThroughputLayerBlend.cpp
@@ -324,6 +324,13 @@ void UDMMaterialStageThroughputLayerBlend::PostEditImport()
 	InitBlendStage();
 }
 
+void UDMMaterialStageThroughputLayerBlend::PostEditUndo()
+{
+	Super::PostEditUndo();
+
+	PullMaskChannelOverride();
+}
+
 void UDMMaterialStageThroughputLayerBlend::InitBlendStage()
 {
 	UDMMaterialStage* Stage = GetStage();
diff --git a/Engine/Plugins/Experimental/DynamicMaterial/Source/DynamicMaterialEditor/Private/Components/MaterialStageExpressions/DMMSETextureSampleBase.cpp b/Engine/Plugins/Experimental/DynamicMaterial/Source/DynamicMaterialEditor/Private/Components/MaterialStageExpressions/DMMSETextureSampleBase.cpp
--- a/Engine/Plugins/Experimental/DynamicMaterial/Source/DynamicMaterialEditor/Private/Components/MaterialStageExpressions/DMMSETextureSampleBase.cpp
+++ b/Engine/Plugins/Experimental/DynamicMaterial/Source/DynamicMaterialEditor/Private/Components/MaterialStageExpressions/DMMSETextureSampleBase.cpp
@@ -411,9 +411,7 @@ void UDMMaterialStageExpressionTextureSampleBase::UpdateMask()
 								}
 							}
 
-							// Set output to alpha
-							MaskStage->UpdateInputMap(2, MaskStageInputConnections[2].Channels[0].SourceIndex,
-								FDMMaterialStageConnectorChannel::WHOLE_CHANNEL, 1, FDMMaterialStageConnectorChannel::FOURTH_CHANNEL, EDMMaterialPropertyType::None);
+							LayerBlend->SetMaskChannelOverride(EAvaColorChannel::Alpha);
 
 							return;
 						}
@@ -466,6 +464,8 @@ void UDMMaterialStageExpressionTextureSampleBase::UpdateMask()
 	check(NewInputTexture);
 
 	NewInputTexture->SetValue(BaseTexture);
+
+	LayerBlend->SetMaskChannelOverride(EAvaColorChannel::Alpha);
 }
 
 #undef LOCTEXT_NAMESPACE
diff --git a/Engine/Plugins/Experimental/DynamicMaterial/Source/DynamicMaterialEditor/Private/UI/Utils/DMEditorSelectionContext.cpp b/Engine/Plugins/Experimental/DynamicMaterial/Source/DynamicMaterialEditor/Private/UI/Utils/DMEditorSelectionContext.cpp
new file mode 100644
--- /dev/null
+++ b/Engine/Plugins/Experimental/DynamicMaterial/Source/DynamicMaterialEditor/Private/UI/Utils/DMEditorSelectionContext.cpp
@@ -0,0 +1,14 @@
+// Copyright Epic Games, Inc. All Rights Reserved.
+
+#include "UI/Utils/DMEditorSelectionContext.h"
+
+#include "Components/DMMaterialProperty.h"
+
+FDMMaterialEditorPage FDMMaterialEditorPage::Preview = {EDMMaterialEditorMode::MaterialPreview, EDMMaterialPropertyType::None};
+FDMMaterialEditorPage FDMMaterialEditorPage::GlobalSettings = {EDMMaterialEditorMode::GlobalSettings, EDMMaterialPropertyType::None};
+FDMMaterialEditorPage FDMMaterialEditorPage::Properties = {EDMMaterialEditorMode::Properties, EDMMaterialPropertyType::None};
+
+bool FDMMaterialEditorPage::operator==(const FDMMaterialEditorPage& InOther) const
+{
+	return EditorMode == InOther.EditorMode && MaterialProperty == InOther.MaterialProperty;
+}
diff --git a/Engine/Plugins/Experimental/DynamicMaterial/Source/DynamicMaterialEditor/Private/UI/Utils/DMEditorSelectionContext.h b/Engine/Plugins/Experimental/DynamicMaterial/Source/DynamicMaterialEditor/Private/UI/Utils/DMEditorSelectionContext.h
new file mode 100644
--- /dev/null
+++ b/Engine/Plugins/Experimental/DynamicMaterial/Source/DynamicMaterialEditor/Private/UI/Utils/DMEditorSelectionContext.h
@@ -0,0 +1,54 @@
+// Copyright Epic Games, Inc. All Rights Reserved.
+
+#pragma once
+
+#include "Containers/Array.h"
+#include "HAL/Platform.h"
+#include "UObject/WeakObjectPtr.h"
+
+class UDMMaterialComponent;
+class UDMMaterialLayerObject;
+class UDMMaterialSlot;
+class UDMMaterialStage;
+enum class EDMMaterialPropertyType : uint8;
+
+enum class EDMMaterialEditorMode : uint8
+{
+	GlobalSettings,
+	Properties,
+	EditSlot,
+	MaterialPreview
+};
+
+struct FDMMaterialEditorPage
+{
+	static FDMMaterialEditorPage Preview;
+	static FDMMaterialEditorPage GlobalSettings;
+	static FDMMaterialEditorPage Properties;
+
+	EDMMaterialEditorMode EditorMode;
+	EDMMaterialPropertyType MaterialProperty;
+
+	bool operator==(const FDMMaterialEditorPage& InOther) const;
+};
+
+struct FDMEditorSelectionContext
+{
+	EDMMaterialEditorMode EditorMode = EDMMaterialEditorMode::GlobalSettings;
+
+	bool bModeChanged = false;
+
+	EDMMaterialPropertyType Property;
+
+	TWeakObjectPtr<UDMMaterialSlot> Slot = nullptr;
+
+	TWeakObjectPtr<UDMMaterialLayerObject> Layer = nullptr;
+
+	TWeakObjectPtr<UDMMaterialComponent> Component = nullptr;
+
+	TArray<FDMMaterialEditorPage> PageHistory;
+
+	int32 PageHistoryActive = 0;
+
+	int32 PageHistoryCount = 0;
+};
diff --git a/Engine/Plugins/Experimental/DynamicMaterial/Source/DynamicMaterialEditor/Private/UI/Widgets/Editor/EditorLayouts/SDMMaterialEditor_LeftBase.cpp b/Engine/Plugins/Experimental/DynamicMaterial/Source/DynamicMaterialEditor/Private/UI/Widgets/Editor/EditorLayouts/SDMMaterialEditor_LeftBase.cpp
--- a/Engine/Plugins/Experimental/DynamicMaterial/Source/DynamicMaterialEditor/Private/UI/Widgets/Editor/EditorLayouts/SDMMaterialEditor_LeftBase.cpp
+++ b/Engine/Plugins/Experimental/DynamicMaterial/Source/DynamicMaterialEditor/Private/UI/Widgets/Editor/EditorLayouts/SDMMaterialEditor_LeftBase.cpp
@@ -4,9 +4,9 @@
 
 #include "DynamicMaterialEditorSettings.h"
 #include "Styling/AppStyle.h"
+#include "UI/Utils/DMEditorSelectionContext.h"
 #include "UI/Widgets/Editor/SDMMaterialComponentEditor.h"
 #include "UI/Widgets/Editor/SDMMaterialGlobalSettingsEditor.h"
-#include "UI/Widgets/Editor/SDMMaterialPreview.h"
 #include "UI/Widgets/Editor/SDMMaterialProperties.h"
 #include "UI/Widgets/Editor/SDMMaterialPropertySelector.h"
 #include "UI/Widgets/Editor/SDMMaterialSlotEditor.h"
@@ -25,61 +25,41 @@ void SDMMaterialEditor_LeftBase::Construct(const FArguments& InArgs, const TShar
 	);
 }
 
-void SDMMaterialEditor_LeftBase::EditSlot(UDMMaterialSlot* InSlot, bool bInForceRefresh)
+void SDMMaterialEditor_LeftBase::EditSlot_Impl(UDMMaterialSlot* InSlot)
 {
-	if (!bInForceRefresh && SlotEditorSlot.IsValid() && SlotEditorSlot->GetSlot() == InSlot)
-	{
-		return;
-	}
+	SDMMaterialEditor::EditSlot_Impl(InSlot);
 
 	RightSlot.Invalidate();
-
-	SDMMaterialEditor::EditSlot(InSlot, bInForceRefresh);
 }
 
-void SDMMaterialEditor_LeftBase::EditComponent(UDMMaterialComponent* InComponent, bool bInForceRefresh)
+void SDMMaterialEditor_LeftBase::EditComponent_Impl(UDMMaterialComponent* InComponent)
 {
-	if (!bInForceRefresh && ComponentEditorSlot.IsValid() && ComponentEditorSlot->GetComponent() == InComponent)
-	{
-		return;
-	}
+	SDMMaterialEditor::EditComponent_Impl(InComponent);
 
-	if (EditMode != EDMMaterialEditorMode::EditSlot)
+	if (SelectionContext.bModeChanged)
 	{
 		RightSlot.Invalidate();
 	}
-
-	SDMMaterialEditor::EditComponent(InComponent, bInForceRefresh);
 }
 
-void SDMMaterialEditor_LeftBase::EditGlobalSettings(bool bInForceRefresh)
+void SDMMaterialEditor_LeftBase::EditGlobalSettings_Impl()
 {
-	if (EditMode == EDMMaterialEditorMode::GlobalSettings && !bInForceRefresh)
-	{
-		return;
-	}
+	SDMMaterialEditor::EditGlobalSettings_Impl();
 
-	if (EditMode != EDMMaterialEditorMode::GlobalSettings)
+	if (SelectionContext.bModeChanged)
 	{
 		RightSlot.Invalidate();
 	}
-
-	SDMMaterialEditor::EditGlobalSettings(bInForceRefresh);
 }
 
-void SDMMaterialEditor_LeftBase::EditProperties(bool bInForceRefresh)
+void SDMMaterialEditor_LeftBase::EditProperties_Impl()
 {
-	if (EditMode == EDMMaterialEditorMode::Properties && !bInForceRefresh)
-	{
-		return;
-	}
+	SDMMaterialEditor::EditProperties_Impl();
 
-	if (EditMode != EDMMaterialEditorMode::Properties)
+	if (SelectionContext.bModeChanged)
 	{
 		RightSlot.Invalidate();
 	}
-
-	SDMMaterialEditor::EditProperties(bInForceRefresh);
 }
 
 void SDMMaterialEditor_LeftBase::ValidateSlots_Main()
@@ -166,20 +146,20 @@ TSharedRef<SWidget> SDMMaterialEditor_LeftBase::CreateSlot_Right()
 {
 	using namespace UE::DynamicMaterialEditor::Private;
 
-	const bool bHasSlotToEdit = SlotToEdit.IsValid();
+	const bool bHasSlotToEdit = SelectionContext.Slot.IsValid();
 
-	if (EditMode == EDMMaterialEditorMode::EditSlot && !bHasSlotToEdit)
+	if (SelectionContext.EditorMode == EDMMaterialEditorMode::EditSlot && !bHasSlotToEdit)
 	{
-		EditMode = EDMMaterialEditorMode::GlobalSettings;
+		SelectionContext.EditorMode = EDMMaterialEditorMode::GlobalSettings;
 	}
 	else if (bHasSlotToEdit)
 	{
-		EditMode = EDMMaterialEditorMode::EditSlot;
+		SelectionContext.EditorMode = EDMMaterialEditorMode::EditSlot;
 	}
 
 	TSharedPtr<SWidget> Content;
 
-	switch (EditMode)
+	switch (SelectionContext.EditorMode)
 	{
 		default:
 		case EDMMaterialEditorMode::GlobalSettings:
diff --git a/Engine/Plugins/Experimental/DynamicMaterial/Source/DynamicMaterialEditor/Private/UI/Widgets/Editor/EditorLayouts/SDMMaterialEditor_LeftBase.h b/Engine/Plugins/Experimental/DynamicMaterial/Source/DynamicMaterialEditor/Private/UI/Widgets/Editor/EditorLayouts/SDMMaterialEditor_LeftBase.h
--- a/Engine/Plugins/Experimental/DynamicMaterial/Source/DynamicMaterialEditor/Private/UI/Widgets/Editor/EditorLayouts/SDMMaterialEditor_LeftBase.h
+++ b/Engine/Plugins/Experimental/DynamicMaterial/Source/DynamicMaterialEditor/Private/UI/Widgets/Editor/EditorLayouts/SDMMaterialEditor_LeftBase.h
@@ -19,13 +19,6 @@ public:
 
 	void Construct(const FArguments& InArgs, const TSharedRef<SDMMaterialDesigner>& InDesignerWidget);
 
-	//~ Begin SDMMaterialEditor
-	virtual void EditSlot(UDMMaterialSlot* InSlot, bool bInForceRefresh = false) override;
-	virtual void EditComponent(UDMMaterialComponent* InComponent, bool bInForceRefresh = false) override;
-	virtual void EditGlobalSettings(bool bInForceRefresh = false) override;
-	virtual void EditProperties(bool bInForceRefresh = false) override;
-	//~ End SDMMaterialEditor
-
 protected:
 	TDMWidgetSlot<SWidget> LeftSlot;
 	TDMWidgetSlot<SWidget> RightSlot;
@@ -44,5 +37,9 @@ protected:
 	virtual void ValidateSlots_Main() override;
 	virtual void ClearSlots_Main() override;
 	virtual TSharedRef<SWidget> CreateSlot_Main() override;
+	virtual void EditSlot_Impl(UDMMaterialSlot* InSlot) override;
+	virtual void EditComponent_Impl(UDMMaterialComponent* InComponent) override;
+	virtual void EditGlobalSettings_Impl() override;
+	virtual void EditProperties_Impl() override;
 	//~ End SDMMaterialEditor
 };
diff --git a/Engine/Plugins/Experimental/DynamicMaterial/Source/DynamicMaterialEditor/Private/UI/Widgets/Editor/EditorLayouts/SDMMaterialEditor_TopBase.cpp b/Engine/Plugins/Experimental/DynamicMaterial/Source/DynamicMaterialEditor/Private/UI/Widgets/Editor/EditorLayouts/SDMMaterialEditor_TopBase.cpp
--- a/Engine/Plugins/Experimental/DynamicMaterial/Source/DynamicMaterialEditor/Private/UI/Widgets/Editor/EditorLayouts/SDMMaterialEditor_TopBase.cpp
+++ b/Engine/Plugins/Experimental/DynamicMaterial/Source/DynamicMaterialEditor/Private/UI/Widgets/Editor/EditorLayouts/SDMMaterialEditor_TopBase.cpp
@@ -3,9 +3,9 @@
 #include "UI/Widgets/Editor/EditorLayouts/SDMMaterialEditor_TopBase.h"
 
 #include "DynamicMaterialEditorSettings.h"
+#include "UI/Utils/DMEditorSelectionContext.h"
 #include "UI/Widgets/Editor/SDMMaterialComponentEditor.h"
 #include "UI/Widgets/Editor/SDMMaterialGlobalSettingsEditor.h"
-#include "UI/Widgets/Editor/SDMMaterialPreview.h"
 #include "UI/Widgets/Editor/SDMMaterialProperties.h"
 #include "UI/Widgets/Editor/SDMMaterialPropertySelector.h"
 #include "UI/Widgets/Editor/SDMMaterialSlotEditor.h"
@@ -23,61 +23,41 @@ void SDMMaterialEditor_TopBase::Construct(const FArguments& InArgs, const TShare
 	);
 }
 
-void SDMMaterialEditor_TopBase::EditSlot(UDMMaterialSlot* InSlot, bool bInForceRefresh)
+void SDMMaterialEditor_TopBase::EditSlot_Impl(UDMMaterialSlot* InSlot)
 {
-	if (!bInForceRefresh && SlotEditorSlot.IsValid() && SlotEditorSlot->GetSlot() == InSlot)
-	{
-		return;
-	}
+	SDMMaterialEditor::EditSlot_Impl(InSlot);
 
 	BottomSlot.Invalidate();
-
-	SDMMaterialEditor::EditSlot(InSlot, bInForceRefresh);
 }
 
-void SDMMaterialEditor_TopBase::EditComponent(UDMMaterialComponent* InComponent, bool bInForceRefresh)
+void SDMMaterialEditor_TopBase::EditComponent_Impl(UDMMaterialComponent* InComponent)
 {
-	if (!bInForceRefresh && ComponentEditorSlot.IsValid() && ComponentEditorSlot->GetComponent() == InComponent)
-	{
-		return;
-	}
+	SDMMaterialEditor::EditComponent_Impl(InComponent);
 
-	if (EditMode != EDMMaterialEditorMode::EditSlot)
+	if (SelectionContext.bModeChanged)
 	{
 		BottomSlot.Invalidate();
 	}
-
-	SDMMaterialEditor::EditComponent(InComponent, bInForceRefresh);
 }
 
-void SDMMaterialEditor_TopBase::EditGlobalSettings(bool bInForceRefresh)
+void SDMMaterialEditor_TopBase::EditGlobalSettings_Impl()
 {
-	if (EditMode == EDMMaterialEditorMode::GlobalSettings && !bInForceRefresh)
-	{
-		return;
-	}
+	SDMMaterialEditor::EditGlobalSettings_Impl();
 
-	if (EditMode != EDMMaterialEditorMode::GlobalSettings)
+	if (SelectionContext.bModeChanged)
 	{
 		BottomSlot.Invalidate();
 	}
-
-	SDMMaterialEditor::EditGlobalSettings(bInForceRefresh);
 }
 
-void SDMMaterialEditor_TopBase::EditProperties(bool bInForceRefresh)
+void SDMMaterialEditor_TopBase::EditProperties_Impl()
 {
-	if (EditMode == EDMMaterialEditorMode::Properties && !bInForceRefresh)
-	{
-		return;
-	}
+	SDMMaterialEditor::EditProperties_Impl();
 
-	if (EditMode != EDMMaterialEditorMode::Properties)
+	if (SelectionContext.bModeChanged)
 	{
 		BottomSlot.Invalidate();
 	}
-
-	SDMMaterialEditor::EditProperties(bInForceRefresh);
 }
 
 void SDMMaterialEditor_TopBase::ValidateSlots_Main()
@@ -164,20 +144,20 @@ TSharedRef<SWidget> SDMMaterialEditor_TopBase::CreateSlot_Bottom()
 {
 	using namespace UE::DynamicMaterialEditor::Private;
 
-	const bool bHasSlotToEdit = SlotToEdit.IsValid();
+	const bool bHasSlotToEdit = SelectionContext.Slot.IsValid();
 
-	if (EditMode == EDMMaterialEditorMode::EditSlot && !bHasSlotToEdit)
+	if (SelectionContext.EditorMode == EDMMaterialEditorMode::EditSlot && !bHasSlotToEdit)
 	{
-		EditMode = EDMMaterialEditorMode::GlobalSettings;
+		SelectionContext.EditorMode = EDMMaterialEditorMode::GlobalSettings;
 	}
 	else if (bHasSlotToEdit)
 	{
-		EditMode = EDMMaterialEditorMode::EditSlot;
+		SelectionContext.EditorMode = EDMMaterialEditorMode::EditSlot;
 	}
 
 	TSharedPtr<SWidget> Content;
 
-	switch (EditMode)
+	switch (SelectionContext.EditorMode)
 	{
 		default:
 		case EDMMaterialEditorMode::GlobalSettings:
diff --git a/Engine/Plugins/Experimental/DynamicMaterial/Source/DynamicMaterialEditor/Private/UI/Widgets/Editor/EditorLayouts/SDMMaterialEditor_TopBase.h b/Engine/Plugins/Experimental/DynamicMaterial/Source/DynamicMaterialEditor/Private/UI/Widgets/Editor/EditorLayouts/SDMMaterialEditor_TopBase.h
--- a/Engine/Plugins/Experimental/DynamicMaterial/Source/DynamicMaterialEditor/Private/UI/Widgets/Editor/EditorLayouts/SDMMaterialEditor_TopBase.h
+++ b/Engine/Plugins/Experimental/DynamicMaterial/Source/DynamicMaterialEditor/Private/UI/Widgets/Editor/EditorLayouts/SDMMaterialEditor_TopBase.h
@@ -19,13 +19,6 @@ public:
 
 	void Construct(const FArguments& InArgs, const TSharedRef<SDMMaterialDesigner>& InDesignerWidget);
 
-	//~ Begin SDMMaterialEditor
-	virtual void EditSlot(UDMMaterialSlot* InSlot, bool bInForceRefresh = false) override;
-	virtual void EditComponent(UDMMaterialComponent* InComponent, bool bInForceRefresh = false) override;
-	virtual void EditGlobalSettings(bool bInForceRefresh = false) override;
-	virtual void EditProperties(bool bInForceRefresh = false) override;
-	//~ End SDMMaterialEditor
-
 protected:
 	TDMWidgetSlot<SWidget> TopSlot;
 	TDMWidgetSlot<SWidget> BottomSlot;
@@ -44,5 +37,9 @@ protected:
 	virtual void ValidateSlots_Main() override;
 	virtual void ClearSlots_Main() override;
 	virtual TSharedRef<SWidget> CreateSlot_Main() override;
+	virtual void EditSlot_Impl(UDMMaterialSlot* InSlot) override;
+	virtual void EditComponent_Impl(UDMMaterialComponent* InComponent) override;
+	virtual void EditGlobalSettings_Impl() override;
+	virtual void EditProperties_Impl() override;
 	//~ End SDMMaterialEditor
 };
diff --git a/Engine/Plugins/Experimental/DynamicMaterial/Source/DynamicMaterialEditor/Private/UI/Widgets/Editor/EditorLayouts/SDMMaterialEditor_TopVertical.cpp b/Engine/Plugins/Experimental/DynamicMaterial/Source/DynamicMaterialEditor/Private/UI/Widgets/Editor/EditorLayouts/SDMMaterialEditor_TopVertical.cpp
--- a/Engine/Plugins/Experimental/DynamicMaterial/Source/DynamicMaterialEditor/Private/UI/Widgets/Editor/EditorLayouts/SDMMaterialEditor_TopVertical.cpp
+++ b/Engine/Plugins/Experimental/DynamicMaterial/Source/DynamicMaterialEditor/Private/UI/Widgets/Editor/EditorLayouts/SDMMaterialEditor_TopVertical.cpp
@@ -5,7 +5,6 @@
 #include "DynamicMaterialEditorSettings.h"
 #include "Styling/AppStyle.h"
 #include "UI/Widgets/Editor/PropertySelectorLayouts/SDMMaterialPropertySelector_Wrap.h"
-#include "UI/Widgets/Editor/SDMMaterialPreview.h"
 #include "UI/Widgets/Editor/SDMMaterialPropertySelector.h"
 #include "UI/Widgets/Editor/SDMMaterialSlotEditor.h"
 #include "Widgets/Layout/SBorder.h"
diff --git a/Engine/Plugins/Experimental/DynamicMaterial/Source/DynamicMaterialEditor/Private/UI/Widgets/Editor/SDMMaterialComponentEditor.cpp b/Engine/Plugins/Experimental/DynamicMaterial/Source/DynamicMaterialEditor/Private/UI/Widgets/Editor/SDMMaterialComponentEditor.cpp
--- a/Engine/Plugins/Experimental/DynamicMaterial/Source/DynamicMaterialEditor/Private/UI/Widgets/Editor/SDMMaterialComponentEditor.cpp
+++ b/Engine/Plugins/Experimental/DynamicMaterial/Source/DynamicMaterialEditor/Private/UI/Widgets/Editor/SDMMaterialComponentEditor.cpp
@@ -113,12 +113,4 @@ TArray<FDMPropertyHandle> SDMMaterialComponentEditor::GetPropertyRows()
 	return PropertyRows;
 }
 
-void SDMMaterialComponentEditor::OnUndo()
-{
-	if (TSharedPtr<SDMMaterialEditor> EditorWidget = GetEditorWidget())
-	{
-		EditorWidget->EditComponent(GetComponent(), /* Force refresh */ true);
-	}
-}
-
 #undef LOCTEXT_NAMESPACE
diff --git a/Engine/Plugins/Experimental/DynamicMaterial/Source/DynamicMaterialEditor/Private/UI/Widgets/Editor/SDMMaterialComponentEditor.h b/Engine/Plugins/Experimental/DynamicMaterial/Source/DynamicMaterialEditor/Private/UI/Widgets/Editor/SDMMaterialComponentEditor.h
--- a/Engine/Plugins/Experimental/DynamicMaterial/Source/DynamicMaterialEditor/Private/UI/Widgets/Editor/SDMMaterialComponentEditor.h
+++ b/Engine/Plugins/Experimental/DynamicMaterial/Source/DynamicMaterialEditor/Private/UI/Widgets/Editor/SDMMaterialComponentEditor.h
@@ -29,6 +29,5 @@ protected:
 	virtual TSharedRef<ICustomDetailsViewItem> GetDefaultCategory(const TSharedRef<ICustomDetailsView>& InDetailsView,
 		const FCustomDetailsViewItemId& InRootId) override;
 	virtual TArray<FDMPropertyHandle> GetPropertyRows() override;
-	virtual void OnUndo() override;
 	//~ End SDMObjectEditorWidgetBase
 };
diff --git a/Engine/Plugins/Experimental/DynamicMaterial/Source/DynamicMaterialEditor/Private/UI/Widgets/Editor/SDMMaterialGlobalSettingsEditor.cpp b/Engine/Plugins/Experimental/DynamicMaterial/Source/DynamicMaterialEditor/Private/UI/Widgets/Editor/SDMMaterialGlobalSettingsEditor.cpp
--- a/Engine/Plugins/Experimental/DynamicMaterial/Source/DynamicMaterialEditor/Private/UI/Widgets/Editor/SDMMaterialGlobalSettingsEditor.cpp
+++ b/Engine/Plugins/Experimental/DynamicMaterial/Source/DynamicMaterialEditor/Private/UI/Widgets/Editor/SDMMaterialGlobalSettingsEditor.cpp
@@ -39,11 +39,3 @@ TArray<FDMPropertyHandle> SDMMaterialGlobalSettingsEditor::GetPropertyRows()
 
 	return PropertyRows;
 }
-
-void SDMMaterialGlobalSettingsEditor::OnUndo()
-{
-	if (TSharedPtr<SDMMaterialEditor> EditorWidget = GetEditorWidget())
-	{
-		EditorWidget->EditGlobalSettings(/* Force refresh */ true);
-	}
-}
diff --git a/Engine/Plugins/Experimental/DynamicMaterial/Source/DynamicMaterialEditor/Private/UI/Widgets/Editor/SDMMaterialGlobalSettingsEditor.h b/Engine/Plugins/Experimental/DynamicMaterial/Source/DynamicMaterialEditor/Private/UI/Widgets/Editor/SDMMaterialGlobalSettingsEditor.h
--- a/Engine/Plugins/Experimental/DynamicMaterial/Source/DynamicMaterialEditor/Private/UI/Widgets/Editor/SDMMaterialGlobalSettingsEditor.h
+++ b/Engine/Plugins/Experimental/DynamicMaterial/Source/DynamicMaterialEditor/Private/UI/Widgets/Editor/SDMMaterialGlobalSettingsEditor.h
@@ -24,6 +24,5 @@ public:
 protected:
 	//~ Begin SDMObjectEditorWidgetBase
 	virtual TArray<FDMPropertyHandle> GetPropertyRows() override;
-	virtual void OnUndo() override;
 	//~ End SDMObjectEditorWidgetBase
 };
diff --git a/Engine/Plugins/Experimental/DynamicMaterial/Source/DynamicMaterialEditor/Private/UI/Widgets/Editor/SDMMaterialPropertySelector.cpp b/Engine/Plugins/Experimental/DynamicMaterial/Source/DynamicMaterialEditor/Private/UI/Widgets/Editor/SDMMaterialPropertySelector.cpp
--- a/Engine/Plugins/Experimental/DynamicMaterial/Source/DynamicMaterialEditor/Private/UI/Widgets/Editor/SDMMaterialPropertySelector.cpp
+++ b/Engine/Plugins/Experimental/DynamicMaterial/Source/DynamicMaterialEditor/Private/UI/Widgets/Editor/SDMMaterialPropertySelector.cpp
@@ -56,7 +56,7 @@ void SDMMaterialPropertySelector::SetSelectedProperty(const FDMMaterialEditorPag
 		return;
 	}
 
-	switch (InPage.EditMode)
+	switch (InPage.EditorMode)
 	{
 		case EDMMaterialEditorMode::MaterialPreview:
 			OpenMaterialPreviewTab();
@@ -145,7 +145,7 @@ TSharedRef<SWidget> SDMMaterialPropertySelector::CreateSlot_EnabledButton(EDMMat
 FText SDMMaterialPropertySelector::GetSelectButtonText(const FDMMaterialEditorPage& InPage, 
 	bool bInShortName)
 {
-	switch (InPage.EditMode)
+	switch (InPage.EditorMode)
 	{
 		case EDMMaterialEditorMode::MaterialPreview:
 			return bInShortName
@@ -174,7 +174,7 @@ FText SDMMaterialPropertySelector::GetSelectButtonText(const FDMMaterialEditorPa
 
 FText SDMMaterialPropertySelector::GetButtonToolTip(const FDMMaterialEditorPage& InPage)
 {
-	switch (InPage.EditMode)
+	switch (InPage.EditorMode)
 	{
 		case EDMMaterialEditorMode::MaterialPreview:
 			return LOCTEXT("MaterialPreviewToolTip", "Show a preview of the Material.");
@@ -313,7 +313,7 @@ void SDMMaterialPropertySelector::OnPropertyEnabledStateChanged(ECheckBoxState I
 
 bool SDMMaterialPropertySelector::GetPropertySelectEnabled(FDMMaterialEditorPage InPage) const
 {
-	switch (InPage.EditMode)
+	switch (InPage.EditorMode)
 	{
 		case EDMMaterialEditorMode::EditSlot:
 			return GetPropertyEnabledEnabled(InPage.MaterialProperty) && DoesPropertySlotExist(InPage.MaterialProperty);
@@ -332,10 +332,10 @@ ECheckBoxState SDMMaterialPropertySelector::GetPropertySelectState(FDMMaterialEd
 		return ECheckBoxState::Undetermined;
 	}
 
-	switch (InPage.EditMode)
+	switch (InPage.EditorMode)
 	{
 		default:
-			return EditorWidget->GetEditMode() == InPage.EditMode
+			return EditorWidget->GetEditMode() == InPage.EditorMode
 				? ECheckBoxState::Checked
 				: ECheckBoxState::Unchecked;
 
@@ -361,7 +361,7 @@ void SDMMaterialPropertySelector::OnPropertySelectStateChanged(ECheckBoxState In
 
 FSlateColor SDMMaterialPropertySelector::GetPropertySelectButtonChipColor(FDMMaterialEditorPage InPage) const
 {
-	switch (InPage.EditMode)
+	switch (InPage.EditorMode)
 	{
 		case EDMMaterialEditorMode::MaterialPreview:
 		case EDMMaterialEditorMode::GlobalSettings:
diff --git a/Engine/Plugins/Experimental/DynamicMaterial/Source/DynamicMaterialEditor/Private/UI/Widgets/Editor/SDMMaterialSlotEditor.cpp b/Engine/Plugins/Experimental/DynamicMaterial/Source/DynamicMaterialEditor/Private/UI/Widgets/Editor/SDMMaterialSlotEditor.cpp
--- a/Engine/Plugins/Experimental/DynamicMaterial/Source/DynamicMaterialEditor/Private/UI/Widgets/Editor/SDMMaterialSlotEditor.cpp
+++ b/Engine/Plugins/Experimental/DynamicMaterial/Source/DynamicMaterialEditor/Private/UI/Widgets/Editor/SDMMaterialSlotEditor.cpp
@@ -467,6 +467,37 @@ void SDMMaterialSlotEditor::InvalidateLayerSettings()
 	LayerSettingsSlot.Invalidate();
 }
 
+void SDMMaterialSlotEditor::SetSelectedLayer(UDMMaterialLayerObject* InLayer)
+{
+	if (LayerViewSlot.HasWidget())
+	{
+		LayerViewSlot->SetSelectedLayer(InLayer);
+	}
+}
+
+void SDMMaterialSlotEditor::TriggerLayerSelectionChange(const TSharedRef<SDMMaterialSlotLayerView>& InLayerView,
+	const TSharedPtr<FDMMaterialLayerReference>& InLayerReference)
+{
+	SlotSettingsSlot.Invalidate();
+
+	OnLayerSelectionChanged.Broadcast(InLayerView, InLayerReference);
+}
+
+void SDMMaterialSlotEditor::TriggerStageSelectionChange(const TSharedRef<SDMMaterialSlotLayerItem>& InLayerItem, UDMMaterialStage* InStage)
+{
+	if (UDMMaterialLayerObject* Layer = InStage->GetLayer())
+	{
+		SetSelectedLayer(Layer);
+	}
+
+	OnStageSelectionChanged.Broadcast(InLayerItem, InStage);
+}
+
+void SDMMaterialSlotEditor::TriggerEffectSelectionChange(const TSharedRef<SDMMaterialSlotLayerEffectView>& InEffectView, UDMMaterialEffect* InEffect)
+{
+	OnEffectSelectionChanged.Broadcast(InEffectView, InEffect);
+}
+
 TSharedRef<SWidget> SDMMaterialSlotEditor::CreateSlot_Container()
 {
 	SVerticalBox::FSlot* SettingsSlotPtr = nullptr;
@@ -754,7 +785,6 @@ TSharedRef<SWidget> SDMMaterialSlotEditor::CreateSlot_LayerOpacity()
 TSharedRef<SDMMaterialSlotLayerView> SDMMaterialSlotEditor::CreateSlot_LayerView()
 {
 	TSharedRef<SDMMaterialSlotLayerView> NewLayerView = SNew(SDMMaterialSlotLayerView, SharedThis(this));
-	NewLayerView->GetOnSelectionChanged().AddSP(this, &SDMMaterialSlotEditor::OnLayerSelected);
 	NewLayerView->EnsureSelectedStage();
 
 	return NewLayerView;
@@ -864,12 +894,6 @@ void SDMMaterialSlotEditor::OnSlotPropertiesUpdated(UDMMaterialSlot* InSlot)
 	}
 }
 
-void SDMMaterialSlotEditor::OnLayerSelected(const TSharedRef<SDMMaterialSlotLayerView>& InLayerView, 
-	const TSharedPtr<FDMMaterialLayerReference>& InLayerReference)
-{
-	SlotSettingsSlot.Invalidate();
-}
-
 FText SDMMaterialSlotEditor::GetLayerButtonsDescription() const
 {
 	UDMMaterialSlot* Slot = GetSlot();
diff --git a/Engine/Plugins/Experimental/DynamicMaterial/Source/DynamicMaterialEditor/Private/UI/Widgets/Editor/SDMMaterialSlotEditor.h b/Engine/Plugins/Experimental/DynamicMaterial/Source/DynamicMaterialEditor/Private/UI/Widgets/Editor/SDMMaterialSlotEditor.h
--- a/Engine/Plugins/Experimental/DynamicMaterial/Source/DynamicMaterialEditor/Private/UI/Widgets/Editor/SDMMaterialSlotEditor.h
+++ b/Engine/Plugins/Experimental/DynamicMaterial/Source/DynamicMaterialEditor/Private/UI/Widgets/Editor/SDMMaterialSlotEditor.h
@@ -9,11 +9,15 @@
 class FDetailColumnSizeData;
 class ICustomDetailsViewItem;
 class SDMMaterialEditor;
+class SDMMaterialSlotLayerEffectView;
+class SDMMaterialSlotLayerItem;
 class SDMMaterialSlotLayerView;
 class SHorizontalBox;
 class SVerticalBox;
+class UDMMaterialEffect;
 class UDMMaterialLayerObject;
 class UDMMaterialSlot;
+class UDMMaterialStage;
 class UDMMaterialValueFloat1;
 class UDMTextureSet;
 class UMaterialFunctionInterface;
@@ -28,6 +32,15 @@ class SDMMaterialSlotEditor : public SCompoundWidget
 	SLATE_END_ARGS()
 
 public:
+	DECLARE_MULTICAST_DELEGATE_TwoParams(FOnLayerSelectionChanged, const TSharedRef<SDMMaterialSlotLayerView>&,
+		const TSharedPtr<FDMMaterialLayerReference>&)
+
+	DECLARE_MULTICAST_DELEGATE_TwoParams(FOnStageSelectionChanged, const TSharedRef<SDMMaterialSlotLayerItem>&,
+		UDMMaterialStage*)
+
+	DECLARE_MULTICAST_DELEGATE_TwoParams(FOnEffectSelectionChanged, const TSharedRef<SDMMaterialSlotLayerEffectView>&,
+		UDMMaterialEffect*)
+
 	virtual ~SDMMaterialSlotEditor() override;
 
 	void Construct(const FArguments& InArgs, const TSharedRef<SDMMaterialEditor>& InEditorWidget, UDMMaterialSlot* InSlot);
@@ -77,6 +90,19 @@ public:
 
 	void InvalidateLayerSettings();
 
+	/** Actions */
+	void SetSelectedLayer(UDMMaterialLayerObject* InLayer);
+
+	/** Events */
+	FOnLayerSelectionChanged::RegistrationType& GetOnLayerSelectionChanged() { return OnLayerSelectionChanged; }
+	void TriggerLayerSelectionChange(const TSharedRef<SDMMaterialSlotLayerView>& InLayerView, const TSharedPtr<FDMMaterialLayerReference>& InLayerReference);
+
+	FOnStageSelectionChanged::RegistrationType& GetOnStageSelectionChanged() { return OnStageSelectionChanged; }
+	void TriggerStageSelectionChange(const TSharedRef<SDMMaterialSlotLayerItem>& InLayerItem, UDMMaterialStage* InStage);
+
+	FOnEffectSelectionChanged::RegistrationType& GetOnEffectSelectionChanged() { return OnEffectSelectionChanged; }
+	void TriggerEffectSelectionChange(const TSharedRef<SDMMaterialSlotLayerEffectView>& InEffectView, UDMMaterialEffect* InEffect);
+
 protected:
 	TWeakPtr<SDMMaterialEditor> EditorWidgetWeak;
 	TWeakObjectPtr<UDMMaterialSlot> MaterialSlotWeak;
@@ -90,6 +116,10 @@ protected:
 	TWeakObjectPtr<UDMMaterialValueFloat1> LayerOpacityValueWeak;
 	TSharedPtr<ICustomDetailsViewItem> LayerOpacityItem;
 
+	FOnLayerSelectionChanged OnLayerSelectionChanged;
+	FOnStageSelectionChanged OnStageSelectionChanged;
+	FOnEffectSelectionChanged OnEffectSelectionChanged;
+
 	TSharedRef<SWidget> CreateSlot_Container();
 
 	TSharedRef<SWidget> CreateSlot_SlotSettings();
@@ -104,8 +134,6 @@ protected:
 
 	void OnSlotPropertiesUpdated(UDMMaterialSlot* InSlot);
 
-	void OnLayerSelected(const TSharedRef<SDMMaterialSlotLayerView>& InLayerView, const TSharedPtr<FDMMaterialLayerReference>& InLayerReference);
-
 	FText GetLayerButtonsDescription() const;
 	TSharedRef<SWidget> GetLayerButtonsMenuContent();
 
diff --git a/Engine/Plugins/Experimental/DynamicMaterial/Source/DynamicMaterialEditor/Private/UI/Widgets/Editor/SDMObjectEditorWidgetBase.h b/Engine/Plugins/Experimental/DynamicMaterial/Source/DynamicMaterialEditor/Private/UI/Widgets/Editor/SDMObjectEditorWidgetBase.h
--- a/Engine/Plugins/Experimental/DynamicMaterial/Source/DynamicMaterialEditor/Private/UI/Widgets/Editor/SDMObjectEditorWidgetBase.h
+++ b/Engine/Plugins/Experimental/DynamicMaterial/Source/DynamicMaterialEditor/Private/UI/Widgets/Editor/SDMObjectEditorWidgetBase.h
@@ -2,7 +2,6 @@
 
 #pragma once
 
-#include "EditorUndoClient.h"
 #include "Widgets/SCompoundWidget.h"
 
 #include "Containers/Array.h"
@@ -21,7 +20,7 @@ enum class ECustomDetailsTreeInsertPosition : uint8;
 struct FDMPropertyHandle;
 
 /** Base class for the object editor. Provides the methods and layout for producing a Custom Details View. */
-class SDMObjectEditorWidgetBase : public SCompoundWidget, public FSelfRegisteringEditorUndoClient
+class SDMObjectEditorWidgetBase : public SCompoundWidget
 {
 	SLATE_DECLARE_WIDGET(SDMObjectEditorWidgetBase, SCompoundWidget)
 
@@ -39,11 +38,6 @@ public:
 
 	void Validate();
 
-	//~ Begin FUndoClient
-	virtual void PostUndo(bool bInSuccess) override { OnUndo(); }
-	virtual void PostRedo(bool bInSuccess) override { OnUndo(); }
-	//~ End FUndoClient
-
 protected:
 	static constexpr const TCHAR* DefaultCategoryName = TEXT("General");
 
@@ -74,6 +68,4 @@ protected:
 	void OnExpansionStateChanged(const TSharedRef<ICustomDetailsViewItem>& InItem, bool bInExpansionState);
 
 	virtual TArray<FDMPropertyHandle> GetPropertyRows() = 0;
-
-	virtual void OnUndo() = 0;
 };
diff --git a/Engine/Plugins/Experimental/DynamicMaterial/Source/DynamicMaterialEditor/Private/UI/Widgets/Editor/SlotEditor/SDMMaterialSlotLayerEffectView.cpp b/Engine/Plugins/Experimental/DynamicMaterial/Source/DynamicMaterialEditor/Private/UI/Widgets/Editor/SlotEditor/SDMMaterialSlotLayerEffectView.cpp
--- a/Engine/Plugins/Experimental/DynamicMaterial/Source/DynamicMaterialEditor/Private/UI/Widgets/Editor/SlotEditor/SDMMaterialSlotLayerEffectView.cpp
+++ b/Engine/Plugins/Experimental/DynamicMaterial/Source/DynamicMaterialEditor/Private/UI/Widgets/Editor/SlotEditor/SDMMaterialSlotLayerEffectView.cpp
@@ -216,14 +216,7 @@ void SDMMaterialSlotLayerEffectView::OnEffectItemSelectionChanged(UDMMaterialEff
 		return;
 	}
 
-	TSharedPtr<SDMMaterialEditor> EditorWidget = SlotEditorWidget->GetEditorWidget();
-
-	if (!EditorWidget.IsValid())
-	{
-		return;
-	}
-
-	EditorWidget->EditComponent(InSelectedItem);
+	SlotEditorWidget->TriggerEffectSelectionChange(SharedThis(this), InSelectedItem);
 }
 
 void SDMMaterialSlotLayerEffectView::OnUndo()
diff --git a/Engine/Plugins/Experimental/DynamicMaterial/Source/DynamicMaterialEditor/Private/UI/Widgets/Editor/SlotEditor/SDMMaterialSlotLayerView.cpp b/Engine/Plugins/Experimental/DynamicMaterial/Source/DynamicMaterialEditor/Private/UI/Widgets/Editor/SlotEditor/SDMMaterialSlotLayerView.cpp
--- a/Engine/Plugins/Experimental/DynamicMaterial/Source/DynamicMaterialEditor/Private/UI/Widgets/Editor/SlotEditor/SDMMaterialSlotLayerView.cpp
+++ b/Engine/Plugins/Experimental/DynamicMaterial/Source/DynamicMaterialEditor/Private/UI/Widgets/Editor/SlotEditor/SDMMaterialSlotLayerView.cpp
@@ -123,11 +123,6 @@ TSharedPtr<SDMMaterialSlotLayerItem> SDMMaterialSlotLayerView::GetWidgetForLayer
 	return nullptr;
 }
 
-SDMMaterialSlotLayerView::FOnSelectionChanged::RegistrationType& SDMMaterialSlotLayerView::GetOnSelectionChanged()
-{
-	return OnSelectionChanged;
-}
-
 void SDMMaterialSlotLayerView::EnsureSelectedStage()
 {
 	if (GetSelectedLayer())
@@ -156,7 +151,7 @@ void SDMMaterialSlotLayerView::EnsureSelectedStage()
 		return;
 	}
 
-	if (UDMMaterialComponent* ComponentToEdit = EditorWidget->GetComponentToEdit())
+	if (UDMMaterialComponent* ComponentToEdit = EditorWidget->GetSelectedComponent())
 	{
 		if (ComponentToEdit->GetTypedParent<UDMMaterialSlot>(/* Allow Subclasses */ true) == Slot)
 		{
@@ -304,7 +299,10 @@ TSharedRef<ITableRow> SDMMaterialSlotLayerView::OnGenerateLayerItemWidget(TShare
 void SDMMaterialSlotLayerView::OnLayerItemSelectionChanged(TSharedPtr<FDMMaterialLayerReference> InSelectedItem, 
 	ESelectInfo::Type InSelectInfo)
 {
-	OnSelectionChanged.Broadcast(SharedThis(this), InSelectedItem);
+	if (TSharedPtr<SDMMaterialSlotEditor> SlotEditorWidget = GetSlotEditorWidget())
+	{
+		SlotEditorWidget->TriggerLayerSelectionChange(SharedThis(this), InSelectedItem);
+	}
 }
 
 TSharedPtr<SWidget> SDMMaterialSlotLayerView::CreateLayerItemContextMenu()
diff --git a/Engine/Plugins/Experimental/DynamicMaterial/Source/DynamicMaterialEditor/Private/UI/Widgets/Editor/SlotEditor/SDMMaterialSlotLayerView.h b/Engine/Plugins/Experimental/DynamicMaterial/Source/DynamicMaterialEditor/Private/UI/Widgets/Editor/SlotEditor/SDMMaterialSlotLayerView.h
--- a/Engine/Plugins/Experimental/DynamicMaterial/Source/DynamicMaterialEditor/Private/UI/Widgets/Editor/SlotEditor/SDMMaterialSlotLayerView.h
+++ b/Engine/Plugins/Experimental/DynamicMaterial/Source/DynamicMaterialEditor/Private/UI/Widgets/Editor/SlotEditor/SDMMaterialSlotLayerView.h
@@ -21,9 +21,6 @@ class SDMMaterialSlotLayerView : public SListView<TSharedPtr<FDMMaterialLayerRef
 	SLATE_END_ARGS()
 
 public:
-	DECLARE_MULTICAST_DELEGATE_TwoParams(FOnSelectionChanged, const TSharedRef<SDMMaterialSlotLayerView>&, 
-		const TSharedPtr<FDMMaterialLayerReference>&)
-
 	virtual ~SDMMaterialSlotLayerView() override;
 
 	void Construct(const FArguments& InArgs, const TSharedRef<SDMMaterialSlotEditor> InSlotEditorWidget);
@@ -36,8 +33,6 @@ public:
 
 	TSharedPtr<SDMMaterialSlotLayerItem> GetWidgetForLayer(UDMMaterialLayerObject* InLayer) const;
 
-	FOnSelectionChanged::RegistrationType& GetOnSelectionChanged();
-
 	void EnsureSelectedStage();
 
 	//~ Begin FUndoClient
@@ -49,7 +44,6 @@ protected:
 	TWeakPtr<SDMMaterialSlotEditor> SlotEditorWidgetWeak;
 
 	TArray<TSharedPtr<FDMMaterialLayerReference>> LayerItems;
-	FOnSelectionChanged OnSelectionChanged;
 
 	void BindCommands();
 
diff --git a/Engine/Plugins/Experimental/DynamicMaterial/Source/DynamicMaterialEditor/Private/UI/Widgets/Editor/SlotEditor/SDMMaterialStage.cpp b/Engine/Plugins/Experimental/DynamicMaterial/Source/DynamicMaterialEditor/Private/UI/Widgets/Editor/SlotEditor/SDMMaterialStage.cpp
--- a/Engine/Plugins/Experimental/DynamicMaterial/Source/DynamicMaterialEditor/Private/UI/Widgets/Editor/SlotEditor/SDMMaterialStage.cpp
+++ b/Engine/Plugins/Experimental/DynamicMaterial/Source/DynamicMaterialEditor/Private/UI/Widgets/Editor/SlotEditor/SDMMaterialStage.cpp
@@ -495,19 +495,7 @@ void SDMMaterialStage::OnMouseButtonDown_Left()
 		return;
 	}
 
-	TSharedPtr<SDMMaterialEditor> EditorWidget = SlotEditorWidget->GetEditorWidget();
-
-	if (!EditorWidget.IsValid())
-	{
-		return;
-	}
-
-	EditorWidget->EditComponent(Stage);
-
-	if (UDMMaterialLayerObject* Layer = Stage->GetLayer())
-	{
-		SlotLayerView->SetSelectedLayer(Layer);
-	}
+	SlotEditorWidget->TriggerStageSelectionChange(SlotLayerItem.ToSharedRef(), Stage);
 }
 
 void SDMMaterialStage::OnMouseButtonUp_Right()
diff --git a/Engine/Plugins/Experimental/DynamicMaterial/Source/DynamicMaterialEditor/Private/UI/Widgets/SDMMaterialEditor.cpp b/Engine/Plugins/Experimental/DynamicMaterial/Source/DynamicMaterialEditor/Private/UI/Widgets/SDMMaterialEditor.cpp
--- a/Engine/Plugins/Experimental/DynamicMaterial/Source/DynamicMaterialEditor/Private/UI/Widgets/SDMMaterialEditor.cpp
+++ b/Engine/Plugins/Experimental/DynamicMaterial/Source/DynamicMaterialEditor/Private/UI/Widgets/SDMMaterialEditor.cpp
@@ -2,6 +2,7 @@
 
 #include "UI/Widgets/SDMMaterialEditor.h"
 
+#include "Components/DMMaterialEffect.h"
 #include "Components/DMMaterialLayer.h"
 #include "Components/DMMaterialProperty.h"
 #include "Components/DMMaterialSlot.h"
@@ -10,6 +11,7 @@
 #include "DynamicMaterialEditorCommands.h"
 #include "DynamicMaterialEditorSettings.h"
 #include "DynamicMaterialModule.h"
+#include "Editor/SlotEditor/SDMMaterialSlotLayerView.h"
 #include "Engine/Texture.h"
 #include "Framework/Application/IInputProcessor.h"
 #include "Framework/Application/SlateApplication.h"
@@ -22,6 +24,7 @@
 #include "Model/DynamicMaterialModelDynamic.h"
 #include "Model/DynamicMaterialModelEditorOnlyData.h"
 #include "Styling/SlateIconFinder.h"
+#include "UI/Utils/DMEditorSelectionContext.h"
 #include "UI/Utils/DMPreviewMaterialManager.h"
 #include "UI/Widgets/Editor/SDMMaterialComponentEditor.h"
 #include "UI/Widgets/Editor/SDMMaterialGlobalSettingsEditor.h"
@@ -106,15 +109,6 @@ private:
 	bool bKeyDown;
 };
 
-FDMMaterialEditorPage FDMMaterialEditorPage::Preview = {EDMMaterialEditorMode::MaterialPreview, EDMMaterialPropertyType::None};
-FDMMaterialEditorPage FDMMaterialEditorPage::GlobalSettings = {EDMMaterialEditorMode::GlobalSettings, EDMMaterialPropertyType::None};
-FDMMaterialEditorPage FDMMaterialEditorPage::Properties = {EDMMaterialEditorMode::Properties, EDMMaterialPropertyType::None};
-
-bool FDMMaterialEditorPage::operator==(const FDMMaterialEditorPage& InOther) const
-{
-	return EditMode == InOther.EditMode && MaterialProperty == InOther.MaterialProperty;
-}
-
 void SDMMaterialEditor::PrivateRegisterAttributes(FSlateAttributeDescriptor::FInitializer&)
 {
 }
@@ -123,13 +117,7 @@ SDMMaterialEditor::SDMMaterialEditor()
 	: SplitterSlot(nullptr)
 	, CommandList(MakeShared<FUICommandList>())
 	, PreviewMaterialManager(MakeShared<FDMPreviewMaterialManager>())
-	, EditMode(EDMMaterialEditorMode::GlobalSettings)
-	, PageHistoryActive(0)
-	, PageHistoryCount(1)
 {
-	// Some small number to get us going
-	PageHistory.Reserve(20);
-	PageHistory.Add(FDMMaterialEditorPage::GlobalSettings);
 }
 
 SDMMaterialEditor::~SDMMaterialEditor()
@@ -164,8 +152,15 @@ SDMMaterialEditor::~SDMMaterialEditor()
 void SDMMaterialEditor::Construct(const FArguments& InArgs, const TSharedRef<SDMMaterialDesigner>& InDesignerWidget)
 {
 	DesignerWidgetWeak = InDesignerWidget;
-	EditMode = EDMMaterialEditorMode::GlobalSettings;
-	SelectedMaterialProperty = EDMMaterialPropertyType::None;
+	
+	SelectionContext.EditorMode = EDMMaterialEditorMode::GlobalSettings;
+	SelectionContext.Property = EDMMaterialPropertyType::None;
+
+	SelectionContext.PageHistory.Add(FDMMaterialEditorPage::GlobalSettings);
+	SelectionContext.PageHistoryCount = 1;
+
+	// Some small number to get us going
+	SelectionContext.PageHistory.Reserve(20);
 
 	SetCanTick(false);
 
@@ -274,7 +269,7 @@ AActor* SDMMaterialEditor::GetMaterialActor() const
 
 EDMMaterialEditorMode SDMMaterialEditor::GetEditMode() const
 {
-	return EditMode;
+	return SelectionContext.EditorMode;
 }
 
 void SDMMaterialEditor::SetMaterialActor(AActor* InActor)
@@ -299,48 +294,46 @@ TSharedPtr<SDMMaterialComponentEditor> SDMMaterialEditor::GetComponentEditorWidg
 	return &ComponentEditorSlot;
 }
 
-UDMMaterialSlot* SDMMaterialEditor::GetSlotToEdit() const
+UDMMaterialSlot* SDMMaterialEditor::SlotSelectedSlot() const
 {
-	return SlotToEdit.Get();
+	return SelectionContext.Slot.Get();
 }
 
-UDMMaterialComponent* SDMMaterialEditor::GetComponentToEdit() const
+UDMMaterialComponent* SDMMaterialEditor::GetSelectedComponent() const
 {
-	return ComponentToEdit.Get();
+	return SelectionContext.Component.Get();
 }
 
 EDMMaterialPropertyType SDMMaterialEditor::GetSelectedPropertyType() const
 {
-	return SelectedMaterialProperty;
+	return SelectionContext.Property;
 }
 
 void SDMMaterialEditor::SelectProperty(EDMMaterialPropertyType InProperty, bool bInForceRefresh)
 {
-	if (EditMode == EDMMaterialEditorMode::EditSlot && SelectedMaterialProperty == InProperty && !bInForceRefresh)
+	if (SelectionContext.EditorMode == EDMMaterialEditorMode::EditSlot && SelectionContext.Property == InProperty && !bInForceRefresh)
 	{
 		return;
 	}
 
-	EditMode = EDMMaterialEditorMode::EditSlot;
-	SelectedMaterialProperty = InProperty;
+	SelectionContext.bModeChanged = bInForceRefresh || (SelectionContext.EditorMode != EDMMaterialEditorMode::EditSlot);
 
-	UDynamicMaterialModelEditorOnlyData* EditorOnlyData = UDynamicMaterialModelEditorOnlyData::Get(MaterialModelBaseWeak);
+	SelectProperty_Impl(InProperty);
 
-	if (!EditorOnlyData)
-	{
-		return;
-	}
-
-	UDMMaterialSlot* Slot = EditorOnlyData->GetSlotForMaterialProperty(InProperty);
+	PageHistoryAdd({EDMMaterialEditorMode::EditSlot, InProperty});
 
-	if (!Slot)
+	if (UDynamicMaterialModelEditorOnlyData* EditorOnlyData = UDynamicMaterialModelEditorOnlyData::Get(MaterialModelBaseWeak))
 	{
-		return;
-	}
-
-	EditSlot(Slot);
+		UDMMaterialSlot* CurrentSlot = SelectionContext.Slot.Get();
 
-	PageHistoryAdd({EDMMaterialEditorMode::EditSlot, InProperty});
+		if (!CurrentSlot || !EditorOnlyData->GetMaterialPropertiesForSlot(CurrentSlot).Contains(InProperty))
+		{
+			if (UDMMaterialSlot* PropertySlot = EditorOnlyData->GetSlotForMaterialProperty(InProperty))
+			{
+				EditSlot(PropertySlot);
+			}
+		}
+	}
 }
 
 const TSharedRef<FUICommandList>& SDMMaterialEditor::GetCommandList() const
@@ -360,24 +353,30 @@ void SDMMaterialEditor::EditSlot(UDMMaterialSlot* InSlot, bool bInForceRefresh)
 		return;
 	}
 
-	SlotEditorSlot.Invalidate();
-	SplitterSlot = nullptr;
-	SlotToEdit = InSlot;
+	SelectionContext.bModeChanged = bInForceRefresh || (SelectionContext.EditorMode != EDMMaterialEditorMode::EditSlot);
 
-	ComponentEditorSlot.Invalidate();
-	ComponentToEdit.Reset();
+	EditSlot_Impl(InSlot);
+
+	SelectionContext.Component.Reset();
+
+	if (!InSlot)
+	{
+		return;
+	}
+
+	UDMMaterialLayerObject* SelectedLayer = SelectionContext.Layer.Get();
 
-	EditMode = EDMMaterialEditorMode::EditSlot;
+	if (SelectedLayer && SelectedLayer->GetSlot() == InSlot)
+	{
+		return;
+	}
 
-	if (InSlot)
+	for (const TObjectPtr<UDMMaterialLayerObject>& Layer : InSlot->GetLayers())
 	{
-		for (const TObjectPtr<UDMMaterialLayerObject>& Layer : InSlot->GetLayers())
+		if (UDMMaterialStage* Stage = Layer->GetFirstValidStage(EDMMaterialLayerStage::All))
 		{
-			if (UDMMaterialStage* Stage = Layer->GetFirstValidStage(EDMMaterialLayerStage::All))
-			{
-				ComponentToEdit = Stage;
-				break;
-			}
+			EditComponent(Stage);
+			break;
 		}
 	}
 }
@@ -389,64 +388,51 @@ void SDMMaterialEditor::EditComponent(UDMMaterialComponent* InComponent, bool bI
 		return;
 	}
 
-	if (EditMode != EDMMaterialEditorMode::EditSlot)
-	{
-		SlotEditorSlot.Invalidate();
-		SplitterSlot = nullptr;
-		GlobalSettingsEditorSlot.Invalidate();
-		MaterialPropertiesSlot.Invalidate();
-	}
-
-	EditMode = EDMMaterialEditorMode::EditSlot;
-
-	ComponentEditorSlot.Invalidate();
-	ComponentToEdit = InComponent;
+	EditComponent_Impl(InComponent);
 }
 
 void SDMMaterialEditor::EditGlobalSettings(bool bInForceRefresh)
 {
-	if (EditMode == EDMMaterialEditorMode::GlobalSettings && !bInForceRefresh)
+	if (SelectionContext.EditorMode == EDMMaterialEditorMode::GlobalSettings && !bInForceRefresh)
 	{
 		return;
 	}
 
-	if (EditMode != EDMMaterialEditorMode::GlobalSettings)
-	{
-		SlotEditorSlot.Invalidate();
-		SplitterSlot = nullptr;
-		ComponentEditorSlot.Invalidate();
-		MaterialPropertiesSlot.Invalidate();
-	}
-
-	EditMode = EDMMaterialEditorMode::GlobalSettings;
-	SelectedMaterialProperty = EDMMaterialPropertyType::None;
-
-	GlobalSettingsEditorSlot.Invalidate();
+	SelectionContext.bModeChanged = bInForceRefresh || (SelectionContext.EditorMode != EDMMaterialEditorMode::GlobalSettings);
 
 	PageHistoryAdd(FDMMaterialEditorPage::GlobalSettings);
+
+	EditGlobalSettings_Impl();
 }
 
 void SDMMaterialEditor::EditProperties(bool bInForceRefresh)
 {
-	if (EditMode == EDMMaterialEditorMode::Properties && !bInForceRefresh)
+	if (SelectionContext.EditorMode == EDMMaterialEditorMode::Properties && !bInForceRefresh)
 	{
 		return;
 	}
 
-	if (EditMode != EDMMaterialEditorMode::Properties)
-	{
-		SlotEditorSlot.Invalidate();
-		SplitterSlot = nullptr;
-		ComponentEditorSlot.Invalidate();
-		GlobalSettingsEditorSlot.Invalidate();
-	}
+	SelectionContext.bModeChanged = bInForceRefresh || (SelectionContext.EditorMode != EDMMaterialEditorMode::Properties);
 
-	EditMode = EDMMaterialEditorMode::Properties;
-	SelectedMaterialProperty = EDMMaterialPropertyType::None;
+	PageHistoryAdd(FDMMaterialEditorPage::Properties);
 
-	MaterialPropertiesSlot.Invalidate();
+	EditProperties_Impl();
+}
 
-	PageHistoryAdd(FDMMaterialEditorPage::Properties);
+void SDMMaterialEditor::OnLayerSelected(const TSharedRef<SDMMaterialSlotLayerView>& InSlotView, 
+	const TSharedPtr<FDMMaterialLayerReference>& InLayerView)
+{
+	SelectionContext.Layer = InLayerView.IsValid() ? InLayerView->GetLayer() : nullptr;
+}
+
+void SDMMaterialEditor::OnStageSelected(const TSharedRef<SDMMaterialSlotLayerItem>& InLayerItem, UDMMaterialStage* InStage)
+{
+	EditComponent(InStage);
+}
+
+void SDMMaterialEditor::OnEffectSelected(const TSharedRef<SDMMaterialSlotLayerEffectView>& InEffectView, UDMMaterialEffect* InEffect)
+{
+	EditComponent(InEffect);
 }
 
 void SDMMaterialEditor::OpenMaterialPreviewTab()
@@ -834,47 +820,49 @@ void SDMMaterialEditor::ValidateSlots()
 			PropertySelectorSlot << CreateSlot_PropertySelector();
 		}
 
-		if (EditMode == EDMMaterialEditorMode::GlobalSettings)
-		{
-			if (GlobalSettingsEditorSlot.HasBeenInvalidated())
-			{
-				GlobalSettingsEditorSlot << CreateSlot_GlobalSettingsEditor();
-			}
-			else
-			{
-				GlobalSettingsEditorSlot->Validate();
-			}
-		}
-		else if (EditMode == EDMMaterialEditorMode::Properties)
-		{
-			if (MaterialPropertiesSlot.HasBeenInvalidated())
-			{
-				MaterialPropertiesSlot << CreateSlot_MaterialProperties();
-			}
-			else
-			{
-				MaterialPropertiesSlot->Validate();
-			}
-		}
-		else
+		switch (SelectionContext.EditorMode)
 		{
-			if (SlotEditorSlot.HasBeenInvalidated())
-			{
-				SlotEditorSlot << CreateSlot_SlotEditor();
-			}
-			else
-			{
-				SlotEditorSlot->ValidateSlots();
-			}
+			case EDMMaterialEditorMode::GlobalSettings:
+				if (GlobalSettingsEditorSlot.HasBeenInvalidated())
+				{
+					GlobalSettingsEditorSlot << CreateSlot_GlobalSettingsEditor();
+				}
+				else
+				{
+					GlobalSettingsEditorSlot->Validate();
+				}
+				break;
 
-			if (ComponentEditorSlot.HasBeenInvalidated())
-			{
-				ComponentEditorSlot << CreateSlot_ComponentEditor();
-			}
-			else
-			{
-				ComponentEditorSlot->Validate();
-			}
+			case EDMMaterialEditorMode::Properties:
+				if (MaterialPropertiesSlot.HasBeenInvalidated())
+				{
+					MaterialPropertiesSlot << CreateSlot_MaterialProperties();
+				}
+				else
+				{
+					MaterialPropertiesSlot->Validate();
+				}
+				break;
+
+			default:
+				if (SlotEditorSlot.HasBeenInvalidated())
+				{
+					SlotEditorSlot << CreateSlot_SlotEditor();
+				}
+				else
+				{
+					SlotEditorSlot->ValidateSlots();
+				}
+
+				if (ComponentEditorSlot.HasBeenInvalidated())
+				{
+					ComponentEditorSlot << CreateSlot_ComponentEditor();
+				}
+				else
+				{
+					ComponentEditorSlot->Validate();
+				}
+				break;
 		}
 	}
 
@@ -882,6 +870,8 @@ void SDMMaterialEditor::ValidateSlots()
 	{
 		StatusBarSlot << CreateSlot_StatusBar();
 	}
+
+	SelectionContext.bModeChanged = false;
 }
 
 void SDMMaterialEditor::ClearSlots()
@@ -902,29 +892,29 @@ void SDMMaterialEditor::ClearSlots()
 
 void SDMMaterialEditor::PageHistoryAdd(const FDMMaterialEditorPage& InPage)
 {
-	if (PageHistory.IsValidIndex(PageHistoryActive) && PageHistory[PageHistoryActive] == InPage)
+	if (SelectionContext.PageHistory.IsValidIndex(SelectionContext.PageHistoryActive) && SelectionContext.PageHistory[SelectionContext.PageHistoryActive] == InPage)
 	{
 		return;
 	}
 
-	const int32 NewPageIndex = PageHistoryActive + 1;
+	const int32 NewPageIndex = SelectionContext.PageHistoryActive + 1;
 
-	if (!PageHistory.IsValidIndex(NewPageIndex))
+	if (!SelectionContext.PageHistory.IsValidIndex(NewPageIndex))
 	{
-		PageHistory.Add(InPage);
+		SelectionContext.PageHistory.Add(InPage);
 	}
 	else
 	{
-		PageHistory[NewPageIndex] = InPage;
+		SelectionContext.PageHistory[NewPageIndex] = InPage;
 	}	
 
-	PageHistoryActive = NewPageIndex;
-	PageHistoryCount = NewPageIndex + 1;
+	SelectionContext.PageHistoryActive = NewPageIndex;
+	SelectionContext.PageHistoryCount = NewPageIndex + 1;
 }
 
 bool SDMMaterialEditor::SetActivePage(const FDMMaterialEditorPage& InPage)
 {
-	switch (InPage.EditMode)
+	switch (InPage.EditorMode)
 	{
 		// This is not a valid page
 		case EDMMaterialEditorMode::MaterialPreview:
@@ -1020,19 +1010,19 @@ void SDMMaterialEditor::HandleDrop_TextureSet(UDMTextureSet* InTextureSet)
 
 bool SDMMaterialEditor::PageHistoryBack()
 {
-	const int32 NewPageIndex = PageHistoryActive - 1;
+	const int32 NewPageIndex = SelectionContext.PageHistoryActive - 1;
 
-	if (!PageHistory.IsValidIndex(NewPageIndex))
+	if (!SelectionContext.PageHistory.IsValidIndex(NewPageIndex))
 	{
 		return false;
 	}
 
-	const int32 OldPageIndex = PageHistoryActive;
-	PageHistoryActive = NewPageIndex;
+	const int32 OldPageIndex = SelectionContext.PageHistoryActive;
+	SelectionContext.PageHistoryActive = NewPageIndex;
 
-	if (!SetActivePage(PageHistory[NewPageIndex]))
+	if (!SetActivePage(SelectionContext.PageHistory[NewPageIndex]))
 	{
-		PageHistoryActive = OldPageIndex;
+		SelectionContext.PageHistoryActive = OldPageIndex;
 		return false;
 	}
 
@@ -1041,19 +1031,19 @@ bool SDMMaterialEditor::PageHistoryBack()
 
 bool SDMMaterialEditor::PageHistoryForward()
 {
-	const int32 NewPageIndex = PageHistoryActive + 1;
+	const int32 NewPageIndex = SelectionContext.PageHistoryActive + 1;
 
-	if (NewPageIndex >= PageHistoryCount || !PageHistory.IsValidIndex(NewPageIndex))
+	if (NewPageIndex >= SelectionContext.PageHistoryCount || !SelectionContext.PageHistory.IsValidIndex(NewPageIndex))
 	{
 		return false;
 	}
 
-	const int32 OldPageIndex = PageHistoryActive;
-	PageHistoryActive = NewPageIndex;
+	const int32 OldPageIndex = SelectionContext.PageHistoryActive;
+	SelectionContext.PageHistoryActive = NewPageIndex;
 
-	if (!SetActivePage(PageHistory[NewPageIndex]))
+	if (!SetActivePage(SelectionContext.PageHistory[NewPageIndex]))
 	{
-		PageHistoryActive = OldPageIndex;
+		SelectionContext.PageHistoryActive = OldPageIndex;
 		return false;
 	}
 
@@ -1147,7 +1137,7 @@ TSharedRef<SDMMaterialPropertySelector> SDMMaterialEditor::CreateSlot_PropertySe
 {
 	TSharedRef<SDMMaterialPropertySelector> NewPropertySelector = CreateSlot_PropertySelector_Impl();
 
-	if (EditMode == EDMMaterialEditorMode::EditSlot && SelectedMaterialProperty == EDMMaterialPropertyType::None)
+	if (SelectionContext.EditorMode == EDMMaterialEditorMode::EditSlot && SelectionContext.Property == EDMMaterialPropertyType::None)
 	{
 		if (UDynamicMaterialModel* MaterialModel = GetMaterialModel())
 		{
@@ -1157,7 +1147,7 @@ TSharedRef<SDMMaterialPropertySelector> SDMMaterialEditor::CreateSlot_PropertySe
 				{
 					if (PropertyPair.Value->IsEnabled() && PropertyPair.Value->IsValidForModel(*EditorOnlyData))
 					{
-						SelectedMaterialProperty = PropertyPair.Key;
+						SelectionContext.Property = PropertyPair.Key;
 						break;
 					}
 				}
@@ -1170,11 +1160,19 @@ TSharedRef<SDMMaterialPropertySelector> SDMMaterialEditor::CreateSlot_PropertySe
 
 TSharedRef<SDMMaterialSlotEditor> SDMMaterialEditor::CreateSlot_SlotEditor()
 {
-	UDMMaterialSlot* Slot = SlotToEdit.Get();
-	SlotToEdit.Reset();
+	UDMMaterialSlot* Slot = SelectionContext.Slot.Get();
 
 	TSharedRef<SDMMaterialSlotEditor> NewSlotEditor = SNew(SDMMaterialSlotEditor, SharedThis(this), Slot);
 
+	if (UDMMaterialLayerObject* Layer = SelectionContext.Layer.Get())
+	{
+		NewSlotEditor->SetSelectedLayer(Layer);
+	}
+
+	NewSlotEditor->GetOnLayerSelectionChanged().AddSP(this, &SDMMaterialEditor::OnLayerSelected);
+	NewSlotEditor->GetOnStageSelectionChanged().AddSP(this, &SDMMaterialEditor::OnStageSelected);
+	NewSlotEditor->GetOnEffectSelectionChanged().AddSP(this, &SDMMaterialEditor::OnEffectSelected);
+
 	BindCommands(&*NewSlotEditor);
 
 	OnEditedSlotChanged.Broadcast(NewSlotEditor, Slot);
@@ -1184,8 +1182,7 @@ TSharedRef<SDMMaterialSlotEditor> SDMMaterialEditor::CreateSlot_SlotEditor()
 
 TSharedRef<SDMMaterialComponentEditor> SDMMaterialEditor::CreateSlot_ComponentEditor()
 {
-	UDMMaterialComponent* Component = ComponentToEdit.Get();
-	ComponentToEdit.Reset();
+	UDMMaterialComponent* Component = SelectionContext.Component.Get();
 
 	TSharedRef<SDMMaterialComponentEditor> NewComponentEditor = SNew(SDMMaterialComponentEditor, SharedThis(this), Component);
 
@@ -1209,7 +1206,31 @@ void SDMMaterialEditor::OnUndo()
 		return;
 	}
 
-	if (EditMode == EDMMaterialEditorMode::EditSlot)
+	ApplySelectionContext();
+}
+
+void SDMMaterialEditor::ApplySelectionContext()
+{
+	switch (SelectionContext.EditorMode)
+	{
+		case EDMMaterialEditorMode::MaterialPreview:
+			// do nothing
+			break;
+
+		case EDMMaterialEditorMode::EditSlot:
+			EditSlot(SelectionContext.Slot.Get());
+			break;
+
+		case EDMMaterialEditorMode::GlobalSettings:
+			EditGlobalSettings();
+			break;
+
+		case EDMMaterialEditorMode::Properties:
+			EditProperties();
+			break;
+	}
+
+	if (SelectionContext.EditorMode == EDMMaterialEditorMode::EditSlot)
 	{
 		if (UDynamicMaterialModelEditorOnlyData* EditorOnlyData = UDynamicMaterialModelEditorOnlyData::Get(MaterialModelBaseWeak))
 		{
@@ -1296,7 +1317,7 @@ void SDMMaterialEditor::NavigateForward_Execute()
 
 bool SDMMaterialEditor::NavigateForward_CanExecute()
 {
-	return (PageHistoryActive + 1) < PageHistoryCount;
+	return (SelectionContext.PageHistoryActive + 1) < SelectionContext.PageHistoryCount;
 }
 
 void SDMMaterialEditor::NavigateBack_Execute()
@@ -1306,7 +1327,7 @@ void SDMMaterialEditor::NavigateBack_Execute()
 
 bool SDMMaterialEditor::NavigateBack_CanExecute()
 {
-	return PageHistoryActive > 0;
+	return SelectionContext.PageHistoryActive > 0;
 }
 
 bool SDMMaterialEditor::CheckOpacityInput(const FKeyEvent& InKeyEvent)
@@ -1327,4 +1348,82 @@ bool SDMMaterialEditor::CheckOpacityInput(const FKeyEvent& InKeyEvent)
 	return false;
 }
 
+void SDMMaterialEditor::SelectProperty_Impl(EDMMaterialPropertyType InProperty)
+{
+	SlotEditorSlot.Invalidate();
+	SplitterSlot = nullptr;
+
+	if (SelectionContext.EditorMode != EDMMaterialEditorMode::EditSlot)
+	{
+		SelectionContext.bModeChanged = true;
+	}
+
+	SelectionContext.EditorMode = EDMMaterialEditorMode::EditSlot;
+	SelectionContext.Property = InProperty;
+}
+
+void SDMMaterialEditor::EditSlot_Impl(UDMMaterialSlot* InSlot)
+{
+	SlotEditorSlot.Invalidate();
+	SplitterSlot = nullptr;
+
+	ComponentEditorSlot.Invalidate();
+
+	SelectionContext.EditorMode = EDMMaterialEditorMode::EditSlot;
+	SelectionContext.Slot = InSlot;
+}
+
+void SDMMaterialEditor::EditComponent_Impl(UDMMaterialComponent* InComponent)
+{
+	if (SelectionContext.EditorMode != EDMMaterialEditorMode::EditSlot)
+	{
+		SlotEditorSlot.Invalidate();
+		SplitterSlot = nullptr;
+		GlobalSettingsEditorSlot.Invalidate();
+		MaterialPropertiesSlot.Invalidate();
+	}
+
+	ComponentEditorSlot.Invalidate();
+
+	SelectionContext.EditorMode = EDMMaterialEditorMode::EditSlot;
+	SelectionContext.Component = InComponent;
+}
+
+void SDMMaterialEditor::EditGlobalSettings_Impl()
+{
+	if (SelectionContext.EditorMode != EDMMaterialEditorMode::GlobalSettings)
+	{
+		SlotEditorSlot.Invalidate();
+		SplitterSlot = nullptr;
+		ComponentEditorSlot.Invalidate();
+		MaterialPropertiesSlot.Invalidate();
+	}
+
+	SelectionContext.EditorMode = EDMMaterialEditorMode::GlobalSettings;
+	SelectionContext.Property = EDMMaterialPropertyType::None;
+	SelectionContext.Slot.Reset();
+	SelectionContext.Layer.Reset();
+
+	GlobalSettingsEditorSlot.Invalidate();
+
+}
+
+void SDMMaterialEditor::EditProperties_Impl()
+{
+	if (SelectionContext.EditorMode != EDMMaterialEditorMode::Properties)
+	{
+		SlotEditorSlot.Invalidate();
+		SplitterSlot = nullptr;
+		ComponentEditorSlot.Invalidate();
+		GlobalSettingsEditorSlot.Invalidate();
+	}
+
+	SelectionContext.EditorMode = EDMMaterialEditorMode::Properties;
+	SelectionContext.Property = EDMMaterialPropertyType::None;
+	SelectionContext.Slot.Reset();
+	SelectionContext.Layer.Reset();
+
+	MaterialPropertiesSlot.Invalidate();
+}
+
 #undef LOCTEXT_NAMESPACE
diff --git a/Engine/Plugins/Experimental/DynamicMaterial/Source/DynamicMaterialEditor/Private/UI/Widgets/SDMMaterialEditor.h b/Engine/Plugins/Experimental/DynamicMaterial/Source/DynamicMaterialEditor/Private/UI/Widgets/SDMMaterialEditor.h
--- a/Engine/Plugins/Experimental/DynamicMaterial/Source/DynamicMaterialEditor/Private/UI/Widgets/SDMMaterialEditor.h
+++ b/Engine/Plugins/Experimental/DynamicMaterial/Source/DynamicMaterialEditor/Private/UI/Widgets/SDMMaterialEditor.h
@@ -9,7 +9,9 @@
 #include "Delegates/DelegateCombinations.h"
 #include "DMObjectMaterialProperty.h"
 #include "Misc/Optional.h"
+#include "UI/Utils/DMEditorSelectionContext.h"
 #include "UI/Utils/DMWidgetSlot.h"
+#include "UObject/StrongObjectPtr.h"
 
 class FDMKeyTracker;
 class FDMPreviewMaterialManager;
@@ -23,17 +25,24 @@ class SDMMaterialPreview;
 class SDMMaterialProperties;
 class SDMMaterialPropertySelector;
 class SDMMaterialSlotEditor;
+class SDMMaterialSlotLayerEffectView;
+class SDMMaterialSlotLayerItem;
+class SDMMaterialSlotLayerView;
 class SDMStatusBar;
 class SDMToolBar;
 class SDockTab;
 class SSplitter;
 class UDMMaterialComponent;
+class UDMMaterialEffect;
+class UDMMaterialLayerObject;
 class UDMMaterialSlot;
+class UDMMaterialStage;
 class UDMTextureSet;
 class UDynamicMaterialModelBase;
 class UDynamicMaterialModelEditorOnlyData;
 enum class EDMMaterialPropertyType : uint8;
 enum class EDMUpdateType : uint8;
+struct FDMMaterialLayerReference;
 struct InPropertyChangedEvent;
 
 namespace UE::DynamicMaterialEditor::Private
@@ -42,26 +51,6 @@ namespace UE::DynamicMaterialEditor::Private
 	inline const TCHAR* EditorLightBackground = TEXT("Brushes.Header");
 }
 
-enum class EDMMaterialEditorMode : uint8
-{
-	GlobalSettings,
-	Properties,
-	EditSlot,
-	MaterialPreview
-};
-
-struct FDMMaterialEditorPage
-{
-	static FDMMaterialEditorPage Preview;
-	static FDMMaterialEditorPage GlobalSettings;
-	static FDMMaterialEditorPage Properties;
-
-	EDMMaterialEditorMode EditMode;
-	EDMMaterialPropertyType MaterialProperty;
-
-	bool operator==(const FDMMaterialEditorPage& InOther) const;
-};
-
 class SDMMaterialEditor : public SCompoundWidget, public FSelfRegisteringEditorUndoClient
 {
 	SLATE_DECLARE_WIDGET(SDMMaterialEditor, SCompoundWidget)
@@ -107,20 +96,26 @@ public:
 
 	TSharedPtr<SDMMaterialComponentEditor> GetComponentEditorWidget() const;
 
-	UDMMaterialSlot* GetSlotToEdit() const;
+	UDMMaterialSlot* SlotSelectedSlot() const;
 
-	UDMMaterialComponent* GetComponentToEdit() const;
+	UDMMaterialComponent* GetSelectedComponent() const;
 
 	/** Actions */
 	void SelectProperty(EDMMaterialPropertyType InProperty, bool bInForceRefresh = false);
 
-	virtual void EditSlot(UDMMaterialSlot* InSlot, bool bInForceRefresh = false);
+	void EditSlot(UDMMaterialSlot* InSlot, bool bInForceRefresh = false);
+
+	void EditComponent(UDMMaterialComponent* InComponent, bool bInForceRefresh = false);
+
+	void EditGlobalSettings(bool bInForceRefresh = false);
 
-	virtual void EditComponent(UDMMaterialComponent* InComponent, bool bInForceRefresh = false);
+	void EditProperties(bool bInForceRefresh = false);
 
-	virtual void EditGlobalSettings(bool bInForceRefresh = false);
+	void OnLayerSelected(const TSharedRef<SDMMaterialSlotLayerView>& InSlotView, const TSharedPtr<FDMMaterialLayerReference>& InLayerView);
 
-	virtual void EditProperties(bool bInForceRefresh = false);
+	void OnStageSelected(const TSharedRef<SDMMaterialSlotLayerItem>& InLayerItem, UDMMaterialStage* InStage);
+
+	void OnEffectSelected(const TSharedRef<SDMMaterialSlotLayerEffectView>& InEffectView, UDMMaterialEffect* InEffect);
 
 	void OpenMaterialPreviewTab();
 
@@ -185,14 +180,7 @@ protected:
 	TSharedPtr<IToolTip> MaterialPreviewToolTip;
 	TDMWidgetSlot<SWidget> MaterialPreviewToolTipSlot;
 
-	EDMMaterialEditorMode EditMode;
-	EDMMaterialPropertyType SelectedMaterialProperty;
-	TWeakObjectPtr<UDMMaterialSlot> SlotToEdit;
-	TWeakObjectPtr<UDMMaterialComponent> ComponentToEdit;
-
-	TArray<FDMMaterialEditorPage> PageHistory;
-	int32 PageHistoryActive;
-	int32 PageHistoryCount;
+	FDMEditorSelectionContext SelectionContext;
 
 	FOnEditedSlotChanged OnEditedSlotChanged;
 	FOnEditedComponentChanged OnEditedComponentChanged;
@@ -250,6 +238,8 @@ protected:
 	/** Events */
 	void OnUndo();
 
+	void ApplySelectionContext();
+
 	/** The material preview window is not cleaned up properly on uobject shutdown, so do it here. */
 	void OnEnginePreExit();
 
@@ -272,4 +262,15 @@ protected:
 	bool NavigateBack_CanExecute();
 
 	bool CheckOpacityInput(const FKeyEvent& InKeyEvent);
+
+	/** Actions */
+	void SelectProperty_Impl(EDMMaterialPropertyType InProperty);
+
+	virtual void EditSlot_Impl(UDMMaterialSlot* InSlot);
+
+	virtual void EditComponent_Impl(UDMMaterialComponent* InComponent);
+
+	virtual void EditGlobalSettings_Impl();
+
+	virtual void EditProperties_Impl();
 };
diff --git a/Engine/Plugins/Experimental/DynamicMaterial/Source/DynamicMaterialEditor/Public/Components/DMMaterialStageBlend.h b/Engine/Plugins/Experimental/DynamicMaterial/Source/DynamicMaterialEditor/Public/Components/DMMaterialStageBlend.h
--- a/Engine/Plugins/Experimental/DynamicMaterial/Source/DynamicMaterialEditor/Public/Components/DMMaterialStageBlend.h
+++ b/Engine/Plugins/Experimental/DynamicMaterial/Source/DynamicMaterialEditor/Public/Components/DMMaterialStageBlend.h
@@ -68,6 +68,10 @@ public:
 	DYNAMICMATERIALEDITOR_API virtual void Update(UDMMaterialComponent* InSource, EDMUpdateType InUpdateType) override;
 	//~ End UDMMaterialComponent
 
+	//~ Begin UObject
+	DYNAMICMATERIALEDITOR_API virtual void PostEditUndo() override;
+	//~ End UObject
+
 	//~ Begin FNotifyHook
 	DYNAMICMATERIALEDITOR_API virtual void NotifyPostChange(const FPropertyChangedEvent& InPropertyChangedEvent, class FEditPropertyChain* InPropertyThatChanged) override;
 	//~ End FNotifyHook
diff --git a/Engine/Plugins/Experimental/DynamicMaterial/Source/DynamicMaterialEditor/Public/Components/DMMaterialStageThroughputLayerBlend.h b/Engine/Plugins/Experimental/DynamicMaterial/Source/DynamicMaterialEditor/Public/Components/DMMaterialStageThroughputLayerBlend.h
--- a/Engine/Plugins/Experimental/DynamicMaterial/Source/DynamicMaterialEditor/Public/Components/DMMaterialStageThroughputLayerBlend.h
+++ b/Engine/Plugins/Experimental/DynamicMaterial/Source/DynamicMaterialEditor/Public/Components/DMMaterialStageThroughputLayerBlend.h
@@ -81,6 +81,7 @@ public:
 	//~ Begin UObject
 	DYNAMICMATERIALEDITOR_API virtual void PostLoad() override;
 	DYNAMICMATERIALEDITOR_API virtual void PostEditImport() override;
+	DYNAMICMATERIALEDITOR_API virtual void PostEditUndo() override;
 	//~ End UObject
 
 protected:
