Material Designer
- When undoing, it will not reset back to the first property.

diff --git a/Engine/Plugins/Experimental/DynamicMaterial/Source/DynamicMaterialEditor/Private/UI/Widgets/SDMMaterialEditor.cpp b/Engine/Plugins/Experimental/DynamicMaterial/Source/DynamicMaterialEditor/Private/UI/Widgets/SDMMaterialEditor.cpp
--- a/Engine/Plugins/Experimental/DynamicMaterial/Source/DynamicMaterialEditor/Private/UI/Widgets/SDMMaterialEditor.cpp
+++ b/Engine/Plugins/Experimental/DynamicMaterial/Source/DynamicMaterialEditor/Private/UI/Widgets/SDMMaterialEditor.cpp
@@ -1234,6 +1234,17 @@ void SDMMaterialEditor::ApplySelectionContext()
 	{
 		if (UDynamicMaterialModelEditorOnlyData* EditorOnlyData = UDynamicMaterialModelEditorOnlyData::Get(MaterialModelBaseWeak))
 		{
+			if (UDMMaterialSlot* Slot = SelectionContext.Slot.Get())
+			{
+				const TArray<EDMMaterialPropertyType> SlotProperties = EditorOnlyData->GetMaterialPropertiesForSlot(Slot);
+
+				if (!SlotProperties.IsEmpty())
+				{
+					SelectProperty(SlotProperties[0]);
+					return;
+				}
+			}
+
 			for (const TPair<EDMMaterialPropertyType, UDMMaterialProperty*>& PropertyPair : EditorOnlyData->GetMaterialProperties())
 			{
 				if (PropertyPair.Value->IsEnabled() && PropertyPair.Value->IsValidForModel(*EditorOnlyData))
