Material Designer
- Moved the "Add texture set" code into a library.

diff --git a/Engine/Plugins/Experimental/DynamicMaterial/Source/DynamicMaterialEditor/Private/DMContentBrowserIntegration.cpp b/Engine/Plugins/Experimental/DynamicMaterial/Source/DynamicMaterialEditor/Private/DMContentBrowserIntegration.cpp
--- a/Engine/Plugins/Experimental/DynamicMaterial/Source/DynamicMaterialEditor/Private/DMContentBrowserIntegration.cpp
+++ b/Engine/Plugins/Experimental/DynamicMaterial/Source/DynamicMaterialEditor/Private/DMContentBrowserIntegration.cpp
@@ -24,6 +24,7 @@
 #include "ScopedTransaction.h"
 #include "Styling/SlateIconFinder.h"
 #include "Utils/DMMaterialModelFunctionLibrary.h"
+#include "Utils/DMTextureSetFunctionLibrary.h"
 
 #define LOCTEXT_NAMESPACE "FDMContentBrowserIntegration"
 
@@ -144,7 +145,7 @@ void FDMContentBrowserIntegration::OnCreateMaterialDesignerMaterialFromTextureSe
 
 	EditorOnlyData->SetChannelListPreset("All");
 
-	if (!EditorOnlyData->AddTextureSet(InTextureSet, /* Replace */ true))
+	if (!UDMTextureSetFunctionLibrary::AddTextureSetToModel(EditorOnlyData, InTextureSet, /* Replace */ true))
 	{
 		return;
 	}
@@ -213,7 +214,7 @@ void FDMContentBrowserIntegration::OnUpdateMaterialDesignerMaterialFromTextureSe
 	FScopedTransaction Transaction(LOCTEXT("AddTextureSet", "Add Texture Set"));
 	EditorOnlyData->Modify();
 
-	const bool bSuccess = EditorOnlyData->AddTextureSet(InTextureSet, bInReplace);
+	const bool bSuccess = UDMTextureSetFunctionLibrary::AddTextureSetToModel(EditorOnlyData, InTextureSet, bInReplace);
 
 	if (!bSuccess)
 	{
diff --git a/Engine/Plugins/Experimental/DynamicMaterial/Source/DynamicMaterialEditor/Private/Model/DynamicMaterialModelEditorOnlyData.cpp b/Engine/Plugins/Experimental/DynamicMaterial/Source/DynamicMaterialEditor/Private/Model/DynamicMaterialModelEditorOnlyData.cpp
--- a/Engine/Plugins/Experimental/DynamicMaterial/Source/DynamicMaterialEditor/Private/Model/DynamicMaterialModelEditorOnlyData.cpp
+++ b/Engine/Plugins/Experimental/DynamicMaterial/Source/DynamicMaterialEditor/Private/Model/DynamicMaterialModelEditorOnlyData.cpp
@@ -7,8 +7,6 @@
 #include "Components/DMMaterialLayer.h"
 #include "Components/DMMaterialProperty.h"
 #include "Components/DMMaterialSlot.h"
-#include "Components/DMMaterialStageBlend.h"
-#include "Components/DMMaterialSubStage.h"
 #include "Components/DMMaterialValue.h"
 #include "Components/DMTextureUV.h"
 #include "Components/MaterialProperties/DMMPAmbientOcclusion.h"
@@ -28,20 +26,14 @@
 #include "Components/MaterialProperties/DMMPSurfaceThickness.h"
 #include "Components/MaterialProperties/DMMPTangent.h"
 #include "Components/MaterialProperties/DMMPWorldPositionOffset.h"
-#include "Components/MaterialStageExpressions/DMMSETextureSample.h"
-#include "Components/MaterialStageInputs/DMMSIExpression.h"
-#include "Components/MaterialStageInputs/DMMSIValue.h"
-#include "Components/MaterialValues/DMMaterialValueTexture.h"
 #include "CoreGlobals.h"
 #include "DMComponentPath.h"
 #include "DMDefs.h"
-#include "DMTextureSet.h"
 #include "DynamicMaterialEditorModule.h"
 #include "DynamicMaterialEditorSettings.h"
 #include "DynamicMaterialModule.h"
 #include "Factories/MaterialFactoryNew.h"
 #include "FileHelpers.h"
-#include "Framework/Notifications/NotificationManager.h"
 #include "IAssetTools.h"
 #include "Material/DynamicMaterialInstance.h"
 #include "Materials/Material.h"
@@ -52,8 +44,6 @@
 #include "Model/DynamicMaterialModel.h"
 #include "UObject/Package.h"
 #include "Utils/DMBuildRequestSubsystem.h"
-#include "Utils/DMUtils.h"
-#include "Widgets/Notifications/SNotificationList.h"
 
 #define LOCTEXT_NAMESPACE "MaterialDesignerModel"
 
@@ -647,208 +637,6 @@ TSharedRef<FDMMaterialBuildState> UDynamicMaterialModelEditorOnlyData::CreateBui
 	return BuildState;
 }
 
-bool UDynamicMaterialModelEditorOnlyData::AddTextureSet(UDMTextureSet* InTextureSet, bool bInReplaceSlots)
-{
-	bool bMadeChange = false;
-
-	TArray<UDMMaterialProperty*> InvalidProperties;
-
-	for (const TPair<EDMTextureSetMaterialProperty, FDMMaterialTexture>& MaterialTexture : InTextureSet->GetTextures())
-	{
-		if (MaterialTexture.Value.Texture.IsNull())
-		{
-			continue;
-		}
-
-		const EDMMaterialPropertyType PropertyType = FDMUtils::TextureSetMaterialPropertyToMaterialPropertyType(MaterialTexture.Key);
-
-		if (PropertyType == EDMMaterialPropertyType::None)
-		{
-			continue;
-		}
-
-		UDMMaterialProperty* Property = GetMaterialProperty(PropertyType);
-
-		if (!Property)
-		{
-			continue;
-		}
-
-		if (!Property->IsEnabled())
-		{
-			Property->SetEnabled(true);
-		}
-
-		if (!Property->IsValidForModel(*this))
-		{
-			InvalidProperties.Add(Property);
-		}
-
-		UDMMaterialSlot* Slot = GetSlotForMaterialProperty(PropertyType);
-
-		if (!Slot)
-		{
-			Slot = AddSlotForMaterialProperty(PropertyType);
-
-			if (!Slot)
-			{
-				continue;
-			}
-		}
-
-		UTexture* Texture = MaterialTexture.Value.Texture.LoadSynchronous();
-
-		if (!Texture)
-		{
-			continue;
-		}
-
-		if (GUndo)
-		{
-			Slot->Modify();
-		}
-
-		UDMMaterialLayerObject* Layer = nullptr;
-
-		{
-			const FDMUpdateGuard Guard;
-
-			Layer = Slot->AddDefaultLayer(PropertyType);
-
-			if (!ensure(Layer))
-			{
-				continue;
-			}
-
-			bMadeChange = true;
-
-			UDMMaterialStage* Stage = Layer->GetStage(EDMMaterialLayerStage::Base);
-
-			if (!ensure(Stage))
-			{
-				continue;
-			}
-
-			UDMMaterialStageInputExpression* NewExpression = UDMMaterialStageInputExpression::ChangeStageInput_Expression(
-				Stage,
-				UDMMaterialStageExpressionTextureSample::StaticClass(),
-				UDMMaterialStageBlend::InputB,
-				FDMMaterialStageConnectorChannel::WHOLE_CHANNEL,
-				0,
-				FDMMaterialStageConnectorChannel::WHOLE_CHANNEL
-			);
-
-			if (!ensure(NewExpression))
-			{
-				continue;
-			}
-
-			UDMMaterialSubStage* SubStage = NewExpression->GetSubStage();
-
-			if (!ensure(SubStage))
-			{
-				continue;
-			}
-
-			UDMMaterialStageInputValue* InputValue = UDMMaterialStageInputValue::ChangeStageInput_NewLocalValue(
-				SubStage,
-				0,
-				FDMMaterialStageConnectorChannel::WHOLE_CHANNEL,
-				EDMValueType::VT_Texture,
-				FDMMaterialStageConnectorChannel::WHOLE_CHANNEL
-			);
-
-			if (ensure(InputValue))
-			{
-				UDMMaterialValueTexture* InputTexture = Cast<UDMMaterialValueTexture>(InputValue->GetValue());
-
-				if (ensure(InputTexture))
-				{
-
-					InputTexture->SetValue(Texture);
-				}
-			}
-
-			if (UDMMaterialStageBlend* Blend = Cast<UDMMaterialStageBlend>(Stage->GetSource()))
-			{
-				EAvaColorChannel AvaColorChannel = EAvaColorChannel::None;
-
-				if (EnumHasAnyFlags(MaterialTexture.Value.TextureChannel, EDMTextureChannelMask::Red))
-				{
-					AvaColorChannel |= EAvaColorChannel::Red;
-				}
-
-				if (EnumHasAnyFlags(MaterialTexture.Value.TextureChannel, EDMTextureChannelMask::Green))
-				{
-					AvaColorChannel |= EAvaColorChannel::Green;
-				}
-
-				if (EnumHasAnyFlags(MaterialTexture.Value.TextureChannel, EDMTextureChannelMask::Blue))
-				{
-					AvaColorChannel |= EAvaColorChannel::Blue;
-				}
-
-				if (EnumHasAnyFlags(MaterialTexture.Value.TextureChannel, EDMTextureChannelMask::Alpha))
-				{
-					AvaColorChannel |= EAvaColorChannel::Alpha;
-				}
-
-				if (AvaColorChannel != EAvaColorChannel::RGBA)
-				{
-					Blend->SetBaseChannelOverride(AvaColorChannel);
-				}				
-			}
-
-			if (bInReplaceSlots)
-			{
-				for (int32 Index = Slot->GetLayers().Num() - 1; Index >= 0; --Index)
-				{
-					UDMMaterialLayerObject* LayerIter = Slot->GetLayer(Index);
-
-					if (!LayerIter || LayerIter->GetStage(EDMMaterialLayerStage::Base) == Stage)
-					{
-						continue;
-					}
-
-					Slot->RemoveLayer(LayerIter);
-				}
-			}
-		}
-
-		Layer->Update(Layer, EDMUpdateType::Structure);
-	}
-
-	/**
-	 * This entire function should not be here. It should be something on the texture set.
-	 * This warning should be there.
-	 */
-	if (!InvalidProperties.IsEmpty())
-	{
-		const FText WarningFormat = LOCTEXT(
-			"AddTextureSetFormat", 
-			"The following channels contain textures but are not valid for the current material settings:\n\n{0}"""
-		);
-
-		const FText IndividualWarningFormat = LOCTEXT("AddTextureSetIndividualFormat", "- {0}\n");
-
-		TArray<FText> ErrorStrings;
-		ErrorStrings.Reserve(InvalidProperties.Num());
-
-		for (UDMMaterialProperty* Property : InvalidProperties)
-		{
-			ErrorStrings.Add(FText::Format(IndividualWarningFormat, Property->GetDescription()));
-		}
-
-		const FText FullWarning = FText::Format(WarningFormat, FText::Join(FText::GetEmpty(), ErrorStrings));
-
-		FNotificationInfo Info(FullWarning);
-		Info.ExpireDuration = 5.0f;
-		FSlateNotificationManager::Get().AddNotification(Info);
-	}
-
-	return bMadeChange;
-}
-
 bool UDynamicMaterialModelEditorOnlyData::NeedsWizard() const
 {
 	return PropertySlotMap.IsEmpty();
diff --git a/Engine/Plugins/Experimental/DynamicMaterial/Source/DynamicMaterialEditor/Private/UI/Widgets/SDMMaterialEditor.cpp b/Engine/Plugins/Experimental/DynamicMaterial/Source/DynamicMaterialEditor/Private/UI/Widgets/SDMMaterialEditor.cpp
--- a/Engine/Plugins/Experimental/DynamicMaterial/Source/DynamicMaterialEditor/Private/UI/Widgets/SDMMaterialEditor.cpp
+++ b/Engine/Plugins/Experimental/DynamicMaterial/Source/DynamicMaterialEditor/Private/UI/Widgets/SDMMaterialEditor.cpp
@@ -37,6 +37,7 @@
 #include "UI/Widgets/SDMMaterialDesigner.h"
 #include "Utils/DMMaterialModelFunctionLibrary.h"
 #include "Utils/DMPrivate.h"
+#include "Utils/DMTextureSetFunctionLibrary.h"
 #include "Widgets/Docking/SDockTab.h"
 #include "Widgets/Layout/SBox.h"
 #include "Widgets/Layout/SScrollBox.h"
@@ -994,12 +995,12 @@ void SDMMaterialEditor::HandleDrop_TextureSet(UDMTextureSet* InTextureSet)
 	{
 		case EAppReturnType::No:
 			EditorOnlyData->Modify();
-			EditorOnlyData->AddTextureSet(InTextureSet, /* Replace */ false);
+			UDMTextureSetFunctionLibrary::AddTextureSetToModel(EditorOnlyData, InTextureSet, /* Replace */ false);
 			break;
 
 		case EAppReturnType::Yes:
 			EditorOnlyData->Modify();
-			EditorOnlyData->AddTextureSet(InTextureSet, /* Replace */ true);
+			UDMTextureSetFunctionLibrary::AddTextureSetToModel(EditorOnlyData, InTextureSet, /* Replace */ true);
 			break;
 
 		default:
diff --git a/Engine/Plugins/Experimental/DynamicMaterial/Source/DynamicMaterialEditor/Private/UI/Widgets/SDMMaterialWizard.cpp b/Engine/Plugins/Experimental/DynamicMaterial/Source/DynamicMaterialEditor/Private/UI/Widgets/SDMMaterialWizard.cpp
--- a/Engine/Plugins/Experimental/DynamicMaterial/Source/DynamicMaterialEditor/Private/UI/Widgets/SDMMaterialWizard.cpp
+++ b/Engine/Plugins/Experimental/DynamicMaterial/Source/DynamicMaterialEditor/Private/UI/Widgets/SDMMaterialWizard.cpp
@@ -24,6 +24,7 @@
 #include "Utils/DMMaterialInstanceFunctionLibrary.h"
 #include "Utils/DMMaterialModelFunctionLibrary.h"
 #include "Utils/DMPrivate.h"
+#include "Utils/DMTextureSetFunctionLibrary.h"
 #include "Widgets/Input/SButton.h"
 #include "Widgets/Input/SCheckBox.h"
 #include "Widgets/Layout/SBorder.h"
@@ -205,7 +206,7 @@ void SDMMaterialWizard::HandleDrop_TextureSet(UDMTextureSet* InTextureSet)
 	FDMScopedUITransaction Transaction(LOCTEXT("DropTextureSet", "Drop Texture Set"));
 
 	EditorOnlyData->Modify();
-	EditorOnlyData->AddTextureSet(InTextureSet, /* Replace */ true);
+	UDMTextureSetFunctionLibrary::AddTextureSetToModel(EditorOnlyData, InTextureSet, /* Replace */ true);
 
 	OpenMaterialInEditor();
 }
diff --git a/Engine/Plugins/Experimental/DynamicMaterial/Source/DynamicMaterialEditor/Private/Utils/DMTextureSetFunctionLibrary.cpp b/Engine/Plugins/Experimental/DynamicMaterial/Source/DynamicMaterialEditor/Private/Utils/DMTextureSetFunctionLibrary.cpp
new file mode 100644
--- /dev/null
+++ b/Engine/Plugins/Experimental/DynamicMaterial/Source/DynamicMaterialEditor/Private/Utils/DMTextureSetFunctionLibrary.cpp
@@ -0,0 +1,234 @@
+// Copyright Epic Games, Inc. All Rights Reserved.
+
+#include "Utils/DMTextureSetFunctionLibrary.h"
+
+#include "DMMaterialTexture.h"
+#include "DMTextureSet.h"
+#include "DMTextureSetMaterialProperty.h"
+#include "Components/DMMaterialLayer.h"
+#include "Components/DMMaterialProperty.h"
+#include "Components/DMMaterialSlot.h"
+#include "Components/DMMaterialStage.h"
+#include "Components/DMMaterialStageBlend.h"
+#include "Components/DMMaterialSubStage.h"
+#include "Components/MaterialStageExpressions/DMMSETextureSample.h"
+#include "Components/MaterialStageInputs/DMMSIExpression.h"
+#include "Components/MaterialStageInputs/DMMSIValue.h"
+#include "Components/MaterialValues/DMMaterialValueTexture.h"
+#include "Engine/Texture.h"
+#include "Framework/Notifications/NotificationManager.h"
+#include "Model/DynamicMaterialModelEditorOnlyData.h"
+#include "Utils/DMUtils.h"
+#include "Widgets/Notifications/SNotificationList.h"
+
+#define LOCTEXT_NAMESPACE "DMTextureSetFunctionLibrary"
+
+bool UDMTextureSetFunctionLibrary::AddTextureSetToModel(UDynamicMaterialModelEditorOnlyData* InEditorOnlyData, UDMTextureSet* InTextureSet,
+	bool bInReplaceSlots)
+{
+	if (!InEditorOnlyData || !InTextureSet)
+	{
+		return false;
+	}
+
+	bool bMadeChange = false;
+
+	TArray<UDMMaterialProperty*> InvalidProperties;
+
+	for (const TPair<EDMTextureSetMaterialProperty, FDMMaterialTexture>& MaterialTexture : InTextureSet->GetTextures())
+	{
+		if (MaterialTexture.Value.Texture.IsNull())
+		{
+			continue;
+		}
+
+		const EDMMaterialPropertyType PropertyType = FDMUtils::TextureSetMaterialPropertyToMaterialPropertyType(MaterialTexture.Key);
+
+		if (PropertyType == EDMMaterialPropertyType::None)
+		{
+			continue;
+		}
+
+		UDMMaterialProperty* Property = InEditorOnlyData->GetMaterialProperty(PropertyType);
+
+		if (!Property)
+		{
+			continue;
+		}
+
+		if (!Property->IsEnabled())
+		{
+			Property->SetEnabled(true);
+		}
+
+		if (!Property->IsValidForModel(*InEditorOnlyData))
+		{
+			InvalidProperties.Add(Property);
+		}
+
+		UDMMaterialSlot* Slot = InEditorOnlyData->GetSlotForMaterialProperty(PropertyType);
+
+		if (!Slot)
+		{
+			Slot = InEditorOnlyData->AddSlotForMaterialProperty(PropertyType);
+
+			if (!Slot)
+			{
+				continue;
+			}
+		}
+
+		UTexture* Texture = MaterialTexture.Value.Texture.LoadSynchronous();
+
+		if (!Texture)
+		{
+			continue;
+		}
+
+		if (GUndo)
+		{
+			Slot->Modify();
+		}
+
+		UDMMaterialLayerObject* Layer = nullptr;
+
+		{
+			const FDMUpdateGuard Guard;
+
+			Layer = Slot->AddDefaultLayer(PropertyType);
+
+			if (!ensure(Layer))
+			{
+				continue;
+			}
+
+			bMadeChange = true;
+
+			UDMMaterialStage* Stage = Layer->GetStage(EDMMaterialLayerStage::Base);
+
+			if (!ensure(Stage))
+			{
+				continue;
+			}
+
+			UDMMaterialStageInputExpression* NewExpression = UDMMaterialStageInputExpression::ChangeStageInput_Expression(
+				Stage,
+				UDMMaterialStageExpressionTextureSample::StaticClass(),
+				UDMMaterialStageBlend::InputB,
+				FDMMaterialStageConnectorChannel::WHOLE_CHANNEL,
+				0,
+				FDMMaterialStageConnectorChannel::WHOLE_CHANNEL
+			);
+
+			if (!ensure(NewExpression))
+			{
+				continue;
+			}
+
+			UDMMaterialSubStage* SubStage = NewExpression->GetSubStage();
+
+			if (!ensure(SubStage))
+			{
+				continue;
+			}
+
+			UDMMaterialStageInputValue* InputValue = UDMMaterialStageInputValue::ChangeStageInput_NewLocalValue(
+				SubStage,
+				0,
+				FDMMaterialStageConnectorChannel::WHOLE_CHANNEL,
+				EDMValueType::VT_Texture,
+				FDMMaterialStageConnectorChannel::WHOLE_CHANNEL
+			);
+
+			if (ensure(InputValue))
+			{
+				UDMMaterialValueTexture* InputTexture = Cast<UDMMaterialValueTexture>(InputValue->GetValue());
+
+				if (ensure(InputTexture))
+				{
+
+					InputTexture->SetValue(Texture);
+				}
+			}
+
+			if (UDMMaterialStageBlend* Blend = Cast<UDMMaterialStageBlend>(Stage->GetSource()))
+			{
+				EAvaColorChannel AvaColorChannel = EAvaColorChannel::None;
+
+				if (EnumHasAnyFlags(MaterialTexture.Value.TextureChannel, EDMTextureChannelMask::Red))
+				{
+					AvaColorChannel |= EAvaColorChannel::Red;
+				}
+
+				if (EnumHasAnyFlags(MaterialTexture.Value.TextureChannel, EDMTextureChannelMask::Green))
+				{
+					AvaColorChannel |= EAvaColorChannel::Green;
+				}
+
+				if (EnumHasAnyFlags(MaterialTexture.Value.TextureChannel, EDMTextureChannelMask::Blue))
+				{
+					AvaColorChannel |= EAvaColorChannel::Blue;
+				}
+
+				if (EnumHasAnyFlags(MaterialTexture.Value.TextureChannel, EDMTextureChannelMask::Alpha))
+				{
+					AvaColorChannel |= EAvaColorChannel::Alpha;
+				}
+
+				if (AvaColorChannel != EAvaColorChannel::RGBA)
+				{
+					Blend->SetBaseChannelOverride(AvaColorChannel);
+				}
+			}
+
+			if (bInReplaceSlots)
+			{
+				for (int32 Index = Slot->GetLayers().Num() - 1; Index >= 0; --Index)
+				{
+					UDMMaterialLayerObject* LayerIter = Slot->GetLayer(Index);
+
+					if (!LayerIter || LayerIter->GetStage(EDMMaterialLayerStage::Base) == Stage)
+					{
+						continue;
+					}
+
+					Slot->RemoveLayer(LayerIter);
+				}
+			}
+		}
+
+		Layer->Update(Layer, EDMUpdateType::Structure);
+	}
+
+	/**
+	 * This entire function should not be here. It should be something on the texture set.
+	 * This warning should be there.
+	 */
+	if (!InvalidProperties.IsEmpty())
+	{
+		const FText WarningFormat = LOCTEXT(
+			"AddTextureSetFormat",
+			"The following channels contain textures but are not valid for the current material settings:\n\n{0}"""
+		);
+
+		const FText IndividualWarningFormat = LOCTEXT("AddTextureSetIndividualFormat", "- {0}\n");
+
+		TArray<FText> ErrorStrings;
+		ErrorStrings.Reserve(InvalidProperties.Num());
+
+		for (UDMMaterialProperty* Property : InvalidProperties)
+		{
+			ErrorStrings.Add(FText::Format(IndividualWarningFormat, Property->GetDescription()));
+		}
+
+		const FText FullWarning = FText::Format(WarningFormat, FText::Join(FText::GetEmpty(), ErrorStrings));
+
+		FNotificationInfo Info(FullWarning);
+		Info.ExpireDuration = 5.0f;
+		FSlateNotificationManager::Get().AddNotification(Info);
+	}
+
+	return bMadeChange;
+}
+
+#undef LOCTEXT_NAMESPACE
diff --git a/Engine/Plugins/Experimental/DynamicMaterial/Source/DynamicMaterialEditor/Public/Model/DynamicMaterialModelEditorOnlyData.h b/Engine/Plugins/Experimental/DynamicMaterial/Source/DynamicMaterialEditor/Public/Model/DynamicMaterialModelEditorOnlyData.h
--- a/Engine/Plugins/Experimental/DynamicMaterial/Source/DynamicMaterialEditor/Public/Model/DynamicMaterialModelEditorOnlyData.h
+++ b/Engine/Plugins/Experimental/DynamicMaterial/Source/DynamicMaterialEditor/Public/Model/DynamicMaterialModelEditorOnlyData.h
@@ -210,15 +210,6 @@ public:
 
 	DYNAMICMATERIALEDITOR_API TSharedRef<FDMMaterialBuildState> CreateBuildState(UMaterial* InMaterialToBuild, bool bInDirtyAssets = true) const;
 
-	/**
-	 * Integrates a Texture Set with the model.
-	 * @param InTextureSet The set to integrate.
-	 * @param bInReplaceSlots Whether to add to or completely replace slots.
-	 * @return True on success.
-	 */
-	UFUNCTION(BlueprintCallable, Category = "Material Designer")
-	DYNAMICMATERIALEDITOR_API bool AddTextureSet(UDMTextureSet* InTextureSet, bool bInReplaceSlots);
-
 	bool NeedsWizard() const;
 
 	DYNAMICMATERIALEDITOR_API void OnWizardComplete();
diff --git a/Engine/Plugins/Experimental/DynamicMaterial/Source/DynamicMaterialEditor/Public/Utils/DMTextureSetFunctionLibrary.h b/Engine/Plugins/Experimental/DynamicMaterial/Source/DynamicMaterialEditor/Public/Utils/DMTextureSetFunctionLibrary.h
new file mode 100644
--- /dev/null
+++ b/Engine/Plugins/Experimental/DynamicMaterial/Source/DynamicMaterialEditor/Public/Utils/DMTextureSetFunctionLibrary.h
@@ -0,0 +1,31 @@
+// Copyright Epic Games, Inc. All Rights Reserved.
+
+#pragma once
+
+#include "Kismet/BlueprintFunctionLibrary.h"
+
+#include "DMTextureSetFunctionLibrary.generated.h"
+
+class UDynamicMaterialModelEditorOnlyData;
+class UDMTextureSet;
+
+/**
+ * Material Stage Function Library
+ */
+UCLASS()
+class UDMTextureSetFunctionLibrary : public UBlueprintFunctionLibrary
+{
+	GENERATED_BODY()
+
+public:
+	/**
+	 * Integrates a Texture Set with the given model's editor only data.
+	 * @param InEditorOnlyData The editor only data for the model.
+	 * @param InTextureSet The set to integrate.
+	 * @param bInReplaceSlots Whether to add to or completely replace slots.
+	 * @return True on success.
+	 */
+	UFUNCTION(BlueprintCallable, Category = "Material Designer")
+	static DYNAMICMATERIALEDITOR_API bool AddTextureSetToModel(UDynamicMaterialModelEditorOnlyData* InEditorOnlyData, UDMTextureSet* InTextureSet,
+		bool bInReplaceSlots);
+};
