Material Designer
- Displacement center and magnitude are now available as options.dis

diff --git a/Engine/Plugins/Experimental/DynamicMaterial/Source/DynamicMaterial/Private/Model/DynamicMaterialModel.cpp b/Engine/Plugins/Experimental/DynamicMaterial/Source/DynamicMaterial/Private/Model/DynamicMaterialModel.cpp
--- a/Engine/Plugins/Experimental/DynamicMaterial/Source/DynamicMaterial/Private/Model/DynamicMaterialModel.cpp
+++ b/Engine/Plugins/Experimental/DynamicMaterial/Source/DynamicMaterial/Private/Model/DynamicMaterialModel.cpp
@@ -134,7 +134,7 @@ UDynamicMaterialModel::UDynamicMaterialModel()
 
 	AddFloatParameter(GlobalBaseColorParameterValue, GlobalBaseColorValueName, GlobalBaseColorParameterName);
 	AddFloatParameter(GlobalEmissiveColorParameterValue, GlobalEmissiveColorValueName, GlobalEmissiveColorParameterName);
-	AddFloatParameter(GlobalOpacityParameterValue, GlobalOpacityValueName, GlobalOpacityParameterName, /* Set Range */ true);
+	AddFloatParameter(GlobalOpacityParameterValue, GlobalOpacityValueName, GlobalOpacityParameterName, 1.f, /* Set Range */ true);
 	AddFloatParameter(GlobalMetallicParameterValue, GlobalMetallicValueName, GlobalMetallicParameterName);
 	AddFloatParameter(GlobalSpecularParameterValue, GlobalSpecularValueName, GlobalSpecularParameterName);
 	AddFloatParameter(GlobalRoughnessParameterValue, GlobalRoughnessValueName, GlobalRoughnessParameterName);
diff --git a/Engine/Plugins/Experimental/DynamicMaterial/Source/DynamicMaterialEditor/Private/Model/DynamicMaterialModelEditorOnlyData.cpp b/Engine/Plugins/Experimental/DynamicMaterial/Source/DynamicMaterialEditor/Private/Model/DynamicMaterialModelEditorOnlyData.cpp
--- a/Engine/Plugins/Experimental/DynamicMaterial/Source/DynamicMaterialEditor/Private/Model/DynamicMaterialModelEditorOnlyData.cpp
+++ b/Engine/Plugins/Experimental/DynamicMaterial/Source/DynamicMaterialEditor/Private/Model/DynamicMaterialModelEditorOnlyData.cpp
@@ -174,6 +174,8 @@ UDynamicMaterialModelEditorOnlyData::UDynamicMaterialModelEditorOnlyData()
 	, bTwoSided(true)
 	, bOutputTranslucentVelocityEnabled(true)
 	, bNaniteTessellationEnabled(true)
+	, DisplacementCenter(0.5f)
+	, DisplacementMagnitude(1.f)
 	, bResponsiveAAEnabled(true)
 {
 	BaseColor =           CreateDefaultSubobject<UDMMaterialPropertyBaseColor>(          "MaterialProperty_BaseColor");
@@ -307,6 +309,11 @@ void UDynamicMaterialModelEditorOnlyData::OnMaterialFlagChanged()
 	RequestMaterialBuild();
 }
 
+void UDynamicMaterialModelEditorOnlyData::OnDisplacementSettingsChanged()
+{
+	RequestMaterialBuild();
+}
+
 void UDynamicMaterialModelEditorOnlyData::Initialize()
 {
 	if (!Slots.IsEmpty())
@@ -394,6 +401,8 @@ void UDynamicMaterialModelEditorOnlyData::BuildMaterial(bool bInDirtyAssets)
 	MaterialModel->DynamicMaterial->bOutputTranslucentVelocity = bOutputTranslucentVelocityEnabled;
 	MaterialModel->DynamicMaterial->bEnableTessellation = bNaniteTessellationEnabled;
 	MaterialModel->DynamicMaterial->bEnableResponsiveAA = bResponsiveAAEnabled;
+	MaterialModel->DynamicMaterial->DisplacementScaling.Magnitude = DisplacementMagnitude;
+	MaterialModel->DynamicMaterial->DisplacementScaling.Center = DisplacementCenter;
 
 	// Not setting this to true can cause the level associated with this material to dirty itself when it
 	// is used with Niagara. It doesn't negatively affect the material in any meaningful way.
@@ -945,6 +954,40 @@ void UDynamicMaterialModelEditorOnlyData::SetResponsiveAAEnabled(bool bInEnabled
 	OnMaterialFlagChanged();
 }
 
+float UDynamicMaterialModelEditorOnlyData::GetDisplacementCenter() const
+{
+	return DisplacementCenter;
+}
+
+void UDynamicMaterialModelEditorOnlyData::SetDisplacementCenter(float InCenter)
+{
+	if (FMath::IsNearlyEqual(DisplacementCenter, InCenter))
+	{
+		return;
+	}
+
+	DisplacementCenter = InCenter;
+
+	OnDisplacementSettingsChanged();
+}
+
+float UDynamicMaterialModelEditorOnlyData::GetDisplacementMagnitude() const
+{
+	return DisplacementMagnitude;
+}
+
+void UDynamicMaterialModelEditorOnlyData::SetDisplacementMagnitude(float InMagnitude)
+{
+	if (FMath::IsNearlyEqual(DisplacementMagnitude, InMagnitude))
+	{
+		return;
+	}
+
+	DisplacementMagnitude = InMagnitude;
+
+	OnDisplacementSettingsChanged();
+}
+
 void UDynamicMaterialModelEditorOnlyData::OpenMaterialEditor() const
 {
 	if (!IsValid(MaterialModel) || !IsValid(MaterialModel->DynamicMaterial))
@@ -1372,6 +1415,11 @@ void UDynamicMaterialModelEditorOnlyData::PostEditChangeChainProperty(FPropertyC
 	{
 		OnMaterialFlagChanged();
 	}
+	else if (Property == GET_MEMBER_NAME_CHECKED(ThisClass, DisplacementCenter)
+		|| Property == GET_MEMBER_NAME_CHECKED(ThisClass, DisplacementMagnitude))
+	{
+		OnDisplacementSettingsChanged();
+	}
 }
 
 void UDynamicMaterialModelEditorOnlyData::OnValueUpdated(UDMMaterialValue* InValue, EDMUpdateType InUpdateType)
diff --git a/Engine/Plugins/Experimental/DynamicMaterial/Source/DynamicMaterialEditor/Private/UI/PropertyGenerators/DMMaterialModelPropertyRowGenerator.cpp b/Engine/Plugins/Experimental/DynamicMaterial/Source/DynamicMaterialEditor/Private/UI/PropertyGenerators/DMMaterialModelPropertyRowGenerator.cpp
--- a/Engine/Plugins/Experimental/DynamicMaterial/Source/DynamicMaterialEditor/Private/UI/PropertyGenerators/DMMaterialModelPropertyRowGenerator.cpp
+++ b/Engine/Plugins/Experimental/DynamicMaterial/Source/DynamicMaterialEditor/Private/UI/PropertyGenerators/DMMaterialModelPropertyRowGenerator.cpp
@@ -60,6 +60,12 @@ void FDMMaterialModelPropertyRowGenerator::AddMaterialModelProperties(const TSha
 
 		AddVariable(InGlobalSettingEditorWidget, InMaterialModelBase, InOutPropertyRows, EditorOnlyData,
 			GET_MEMBER_NAME_CHECKED(UDynamicMaterialModelEditorOnlyData, bNaniteTessellationEnabled));
+
+		AddVariable(InGlobalSettingEditorWidget, InMaterialModelBase, InOutPropertyRows, EditorOnlyData,
+			GET_MEMBER_NAME_CHECKED(UDynamicMaterialModelEditorOnlyData, DisplacementMagnitude));
+
+		AddVariable(InGlobalSettingEditorWidget, InMaterialModelBase, InOutPropertyRows, EditorOnlyData,
+			GET_MEMBER_NAME_CHECKED(UDynamicMaterialModelEditorOnlyData, DisplacementCenter));
 	}
 }
 
diff --git a/Engine/Plugins/Experimental/DynamicMaterial/Source/DynamicMaterialEditor/Public/Model/DynamicMaterialModelEditorOnlyData.h b/Engine/Plugins/Experimental/DynamicMaterial/Source/DynamicMaterialEditor/Public/Model/DynamicMaterialModelEditorOnlyData.h
--- a/Engine/Plugins/Experimental/DynamicMaterial/Source/DynamicMaterialEditor/Public/Model/DynamicMaterialModelEditorOnlyData.h
+++ b/Engine/Plugins/Experimental/DynamicMaterial/Source/DynamicMaterialEditor/Public/Model/DynamicMaterialModelEditorOnlyData.h
@@ -149,6 +149,18 @@ public:
 	UFUNCTION(BlueprintCallable, Category = "Material Designer")
 	DYNAMICMATERIALEDITOR_API void SetResponsiveAAEnabled(bool bInEnabled);
 
+	UFUNCTION(BlueprintPure, Category = "Material Designer")
+	DYNAMICMATERIALEDITOR_API float GetDisplacementCenter() const;
+
+	UFUNCTION(BlueprintCallable, Category = "Material Designer")
+	DYNAMICMATERIALEDITOR_API void SetDisplacementCenter(float InCenter);
+
+	UFUNCTION(BlueprintPure, Category = "Material Designer")
+	DYNAMICMATERIALEDITOR_API float GetDisplacementMagnitude() const;
+
+	UFUNCTION(BlueprintCallable, Category = "Material Designer")
+	DYNAMICMATERIALEDITOR_API void SetDisplacementMagnitude(float InMagnitude);
+
 	UFUNCTION(BlueprintCallable, Category = "Material Designer")
 	DYNAMICMATERIALEDITOR_API void SetChannelListPreset(FName InPresetName);
 
@@ -281,6 +293,14 @@ protected:
 	UPROPERTY(EditInstanceOnly, BlueprintReadOnly, Category = "Material Designer", meta = (DisplayName = "Nanite Tessellation"))
 	bool bNaniteTessellationEnabled;
 
+	/** Mid point for displacement in the range 0-1. */
+	UPROPERTY(EditInstanceOnly, BlueprintReadOnly, Category = "Material Designer", meta = ( UIMin = 0, UIMax = 1, ClampMin = 0, ClampMax = 1))
+	float DisplacementCenter;
+
+	/** Multipler for displacement values. */
+	UPROPERTY(EditInstanceOnly, BlueprintReadOnly, Category = "Material Designer")
+	float DisplacementMagnitude;
+
 	/**
 	 * Indicates that the material should be rendered using responsive anti-aliasing. Improves sharpness of small moving particles such as sparks.
 	 * Only use for small moving features because it will cause aliasing of the background.
@@ -397,6 +417,7 @@ protected:
 	void OnBlendModeChanged();
 	void OnShadingModelChanged();
 	void OnMaterialFlagChanged();
+	void OnDisplacementSettingsChanged();
 
 	//~ Begin IDynamicMaterialModelEditorOnlyDataInterface
 	virtual void ReinitComponents() override;
