Material Designer
- Actors assigned material assets now have a button to craete and replace the asset with a local material instance.

diff --git a/Engine/Plugins/Experimental/DynamicMaterial/Source/DynamicMaterialEditor/Private/UI/Widgets/Editor/SDMToolBar.cpp b/Engine/Plugins/Experimental/DynamicMaterial/Source/DynamicMaterialEditor/Private/UI/Widgets/Editor/SDMToolBar.cpp
--- a/Engine/Plugins/Experimental/DynamicMaterial/Source/DynamicMaterialEditor/Private/UI/Widgets/Editor/SDMToolBar.cpp
+++ b/Engine/Plugins/Experimental/DynamicMaterial/Source/DynamicMaterialEditor/Private/UI/Widgets/Editor/SDMToolBar.cpp
@@ -11,6 +11,7 @@
 #include "GameFramework/Actor.h"
 #include "Layout/Margin.h"
 #include "Material/DynamicMaterialInstance.h"
+#include "Material/DynamicMaterialInstanceFactory.h"
 #include "Math/Vector2D.h"
 #include "Model/DynamicMaterialModel.h"
 #include "Model/DynamicMaterialModelBase.h"
@@ -95,7 +96,7 @@ TSharedRef<SWidget> SDMToolBar::CreateToolBarEntries()
 				SNew(SHorizontalBox)
 				+ SHorizontalBox::Slot()
 				.AutoWidth()
-				.Padding(5.0f, 0.0f, 0.0f, 0.0f)
+				.Padding(0.f)
 				.VAlign(EVerticalAlignment::VAlign_Center)
 				[
 					SAssignNew(SaveButtonWidget, SButton)
@@ -112,7 +113,7 @@ TSharedRef<SWidget> SDMToolBar::CreateToolBarEntries()
 				]		
 				+ SHorizontalBox::Slot()
 				.AutoWidth()
-				.Padding(5.0f, 0.0f, 0.0f, 0.0f)
+				.Padding(0.f)
 				.VAlign(EVerticalAlignment::VAlign_Center)
 				[
 					SNew(SButton)
@@ -129,7 +130,7 @@ TSharedRef<SWidget> SDMToolBar::CreateToolBarEntries()
 				+ SHorizontalBox::Slot()
 				.AutoWidth()
 				.VAlign(VAlign_Center)
-				.Padding(0.f, 0.f, 0.f, 0.f)
+				.Padding(0.f)
 				[
 					SAssignNew(OpenParentButton, SButton)
 					.Visibility(EVisibility::Collapsed)
@@ -146,7 +147,7 @@ TSharedRef<SWidget> SDMToolBar::CreateToolBarEntries()
 				+ SHorizontalBox::Slot()
 				.AutoWidth()
 				.VAlign(VAlign_Center)
-				.Padding(0.f, 0.f, 0.f, 0.f)
+				.Padding(0.f)
 				[
 					SAssignNew(ConvertToEditableButton, SButton)
 					.Visibility(EVisibility::Collapsed)
@@ -160,6 +161,23 @@ TSharedRef<SWidget> SDMToolBar::CreateToolBarEntries()
 						.DesiredSizeOverride(LargeIconToolBarButtonSize)
 					]
 				]
+				+ SHorizontalBox::Slot()
+				.AutoWidth()
+				.VAlign(VAlign_Center)
+				.Padding(0.f)
+				[
+					SAssignNew(ConvertToInstanceButton, SButton)
+					.Visibility(EVisibility::Collapsed)
+					.ContentPadding(LargeIconToolBarButtonContentPadding)
+					.ButtonStyle(FDynamicMaterialEditorStyle::Get(), "HoverHintOnly")
+					.ToolTipText(LOCTEXT("MaterialDesignerConvertToInstanceTooltip", "Convert this Material Designer Material to an Instance, creating a local Instance inside the Actor."))
+					.OnClicked(this, &SDMToolBar::OnConvertToInstanceClicked)
+					[
+						SNew(SImage)
+						.Image(FAppStyle::GetBrush(TEXT("Icons.Toolbar.Export")))
+						.DesiredSizeOverride(LargeIconToolBarButtonSize)
+					]
+				]
 			]
 
 			+ SWrapBox::Slot()
@@ -391,6 +409,18 @@ void SDMToolBar::SetButtonVisibilities()
 		ConvertToEditableButton->SetVisibility(EVisibility::Collapsed);
 		InstanceWidget->SetVisibility(EVisibility::Collapsed);
 	}
+
+	if (bIsAsset && !bIsDynamic && 
+		!ActorMaterialProperties.IsEmpty() 
+		&& ActorMaterialProperties.IsValidIndex(SelectedMaterialElementIndex)
+		&& ActorMaterialProperties[SelectedMaterialElementIndex].IsValid() && GetMaterialActor())
+	{
+		ConvertToInstanceButton->SetVisibility(EVisibility::Visible);
+	}
+	else
+	{
+		ConvertToInstanceButton->SetVisibility(EVisibility::Collapsed);
+	}
 }
 
 TSharedRef<SWidget> SDMToolBar::CreateToolBarButton(TAttribute<const FSlateBrush*> InImageBrush, const TAttribute<FText>& InTooltipText, FOnClicked InOnClicked)
@@ -916,6 +946,89 @@ FReply SDMToolBar::OnConvertToEditableClicked()
 	return FReply::Handled();
 }
 
+FReply SDMToolBar::OnConvertToInstanceClicked()
+{
+	UDynamicMaterialModel* CurrentModel = Cast<UDynamicMaterialModel>(GetMaterialModelBase());
+
+	if (!CurrentModel)
+	{
+		UE_LOG(LogDynamicMaterialEditor, Error, TEXT("Tried to convert a null or dynamic model to an instance."));
+		return FReply::Handled();
+	}
+
+	UDynamicMaterialInstance* OldInstance = CurrentModel->GetDynamicMaterialInstance();
+
+	if (!OldInstance)
+	{
+		UE_LOG(LogDynamicMaterialEditor, Error, TEXT("Unable to find material to replace."));
+		return FReply::Handled();
+	}
+
+	bool bIsAsset = false;
+
+	if (CurrentModel->IsAsset())
+	{
+		bIsAsset = true;
+	}
+	else if (OldInstance)
+	{
+		if (OldInstance->IsAsset())
+		{
+			bIsAsset = true;
+		}
+	}
+
+	if (!bIsAsset)
+	{
+		UE_LOG(LogDynamicMaterialEditor, Error, TEXT("Tried to convert a non-asset material to an instance."));
+		return FReply::Handled();
+	}
+
+	if (ActorMaterialProperties.IsEmpty() 
+		|| !ActorMaterialProperties.IsValidIndex(SelectedMaterialElementIndex)
+		|| !ActorMaterialProperties[SelectedMaterialElementIndex].IsValid()
+		|| !GetMaterialActor())
+	{
+		UE_LOG(LogDynamicMaterialEditor, Error, TEXT("Unable to find actor or actor property to insert the new instance."));
+		return FReply::Handled();
+	}
+
+	FDMObjectMaterialProperty& Property = *ActorMaterialProperties[SelectedMaterialElementIndex];
+
+	UDynamicMaterialInstanceFactory* InstanceFactory = NewObject<UDynamicMaterialInstanceFactory>();
+	check(InstanceFactory);
+
+	UDynamicMaterialInstance* NewInstance = Cast<UDynamicMaterialInstance>(InstanceFactory->FactoryCreateNew(
+		UDynamicMaterialInstance::StaticClass(),
+		OldInstance->GetOuter(),
+		NAME_None,
+		RF_Transactional,
+		nullptr,
+		GWarn	
+	));
+
+	if (!NewInstance)
+	{
+		UE_LOG(LogDynamicMaterialEditor, Error, TEXT("Failed to create new material."));
+		return FReply::Handled();
+	}
+
+	if (!UDMMaterialModelFunctionLibrary::CreateModelInstanceInMaterial(CurrentModel, NewInstance))
+	{
+		UE_LOG(LogDynamicMaterialEditor, Error, TEXT("Failed to create new material model instance."));
+		return FReply::Handled();
+	}
+
+	// Setting it on the actor will automatically open it if the actor property is currently active.
+	if (!UDMMaterialInstanceFunctionLibrary::SetMaterialInObject(Property, NewInstance))
+	{
+		UE_LOG(LogDynamicMaterialEditor, Error, TEXT("Failed to assign new material to actor."));
+		return FReply::Handled();
+	}
+
+	return FReply::Handled();
+}
+
 TSharedRef<SWidget> SDMToolBar::GenerateSettingsMenu()
 {
 	return FDMToolBarMenus::MakeEditorLayoutMenu(GetEditorWidget());
diff --git a/Engine/Plugins/Experimental/DynamicMaterial/Source/DynamicMaterialEditor/Private/UI/Widgets/Editor/SDMToolBar.h b/Engine/Plugins/Experimental/DynamicMaterial/Source/DynamicMaterialEditor/Private/UI/Widgets/Editor/SDMToolBar.h
--- a/Engine/Plugins/Experimental/DynamicMaterial/Source/DynamicMaterialEditor/Private/UI/Widgets/Editor/SDMToolBar.h
+++ b/Engine/Plugins/Experimental/DynamicMaterial/Source/DynamicMaterialEditor/Private/UI/Widgets/Editor/SDMToolBar.h
@@ -64,6 +64,7 @@ protected:
 	TSharedPtr<STextBlock> InstanceWidget;
 	TSharedPtr<SWidget> OpenParentButton;
 	TSharedPtr<SWidget> ConvertToEditableButton;
+	TSharedPtr<SWidget> ConvertToInstanceButton;
 
 	UDynamicMaterialModelBase* GetMaterialModelBase() const;
 
@@ -104,5 +105,7 @@ protected:
 
 	FReply OnConvertToEditableClicked();
 
+	FReply OnConvertToInstanceClicked();
+
 	TSharedRef<SWidget> GenerateSettingsMenu();
 };
