Material Designer: Undo
- The wizard and actor slot selector now have undo.

diff --git a/Engine/Plugins/Experimental/DynamicMaterial/Source/DynamicMaterialEditor/Private/DMObjectMaterialProperty.cpp b/Engine/Plugins/Experimental/DynamicMaterial/Source/DynamicMaterialEditor/Private/DMObjectMaterialProperty.cpp
--- a/Engine/Plugins/Experimental/DynamicMaterial/Source/DynamicMaterialEditor/Private/DMObjectMaterialProperty.cpp
+++ b/Engine/Plugins/Experimental/DynamicMaterial/Source/DynamicMaterialEditor/Private/DMObjectMaterialProperty.cpp
@@ -122,6 +122,11 @@ void FDMObjectMaterialProperty::SetMaterial(UDynamicMaterialInstance* InDynamicM
 		return;
 	}
 
+	if (GUndo)
+	{
+		Outer->Modify();
+	}
+
 	if (Property != nullptr)
 	{
 		Outer->PreEditChange(Property);
diff --git a/Engine/Plugins/Experimental/DynamicMaterial/Source/DynamicMaterialEditor/Private/UI/Widgets/SDMActorMaterialSelector.cpp b/Engine/Plugins/Experimental/DynamicMaterial/Source/DynamicMaterialEditor/Private/UI/Widgets/SDMActorMaterialSelector.cpp
--- a/Engine/Plugins/Experimental/DynamicMaterial/Source/DynamicMaterialEditor/Private/UI/Widgets/SDMActorMaterialSelector.cpp
+++ b/Engine/Plugins/Experimental/DynamicMaterial/Source/DynamicMaterialEditor/Private/UI/Widgets/SDMActorMaterialSelector.cpp
@@ -14,6 +14,7 @@
 #include "ThumbnailRendering/ThumbnailManager.h"
 #include "UI/Widgets/SDMMaterialDesigner.h"
 #include "Utils/DMMaterialInstanceFunctionLibrary.h"
+#include "Utils/DMPrivate.h"
 #include "Widgets/Input/SButton.h"
 #include "Widgets/Layout/SBox.h"
 #include "Widgets/Layout/SScrollBox.h"
@@ -213,6 +214,13 @@ FReply SDMActorMaterialSelector::OnCreateMaterialButtonClicked(int32 InActorProp
 		return FReply::Handled();
 	}
 
+	const FDMScopedUITransaction Transaction(LOCTEXT("CreateMaterialDesignerMaterialInActor", "Create Material Designer Material in Actor"));
+
+	if (UObject* Outer = ActorProperties[InActorPropertyIndex].GetOuter())
+	{
+		Outer->Modify();
+	}
+
 	UDynamicMaterialModel* NewMaterialModel = UDMMaterialInstanceFunctionLibrary::CreateMaterialInObject(ActorProperties[InActorPropertyIndex]);
 
 	if (NewMaterialModel)
diff --git a/Engine/Plugins/Experimental/DynamicMaterial/Source/DynamicMaterialEditor/Private/UI/Widgets/SDMMaterialWizard.cpp b/Engine/Plugins/Experimental/DynamicMaterial/Source/DynamicMaterialEditor/Private/UI/Widgets/SDMMaterialWizard.cpp
--- a/Engine/Plugins/Experimental/DynamicMaterial/Source/DynamicMaterialEditor/Private/UI/Widgets/SDMMaterialWizard.cpp
+++ b/Engine/Plugins/Experimental/DynamicMaterial/Source/DynamicMaterialEditor/Private/UI/Widgets/SDMMaterialWizard.cpp
@@ -608,6 +608,13 @@ FReply SDMMaterialWizard::NewTemplateAccept_OnClick()
 	{
 		if (MaterialObjectProperty->IsValid())
 		{
+			const FDMScopedUITransaction Transaction(LOCTEXT("CreateMaterialDesignerMaterialInActor", "Create Material Designer Material in Actor"));
+
+			if (UObject* Outer = MaterialObjectProperty->GetOuter())
+			{
+				Outer->Modify();
+			}
+
 			CreateTemplateMaterialInActor(CurrentPreset, MaterialObjectProperty.GetValue());
 		}
 		else
@@ -850,6 +857,13 @@ void SDMMaterialWizard::SelectTemplate(UDynamicMaterialModel* InTemplateModel)
 	{
 		if (MaterialObjectProperty->IsValid())
 		{
+			const FDMScopedUITransaction Transaction(LOCTEXT("CreateMaterialDesignerInstanceInActor", "Create Material Designer Instance in Actor"));
+
+			if (UObject* Outer = MaterialObjectProperty->GetOuter())
+			{
+				Outer->Modify();
+			}
+
 			CreateNewDynamicInstanceInActor(InTemplateModel, MaterialObjectProperty.GetValue());
 		}
 		else
@@ -863,6 +877,9 @@ void SDMMaterialWizard::SelectTemplate(UDynamicMaterialModel* InTemplateModel)
 	{
 		if (UDynamicMaterialInstance* Instance = MaterialModel->GetDynamicMaterialInstance())
 		{
+			const FDMScopedUITransaction Transaction(LOCTEXT("ReplaceMaterialDesignerModelInMaterial", "Replace Material Designer Model in Material"));
+			Instance->Modify();
+
 			CreateDynamicMaterialInInstance(InTemplateModel, Instance);
 		}
 		else
