Material Designer: Texture sample optimisation
- You can now select "Use Base Texture" on mask stages to re-use the texture sample from the base stage, if it's using a texture.
- The correct channel is used for mask stage previews

diff --git a/Engine/Plugins/Experimental/DynamicMaterial/Source/DynamicMaterialEditor/Private/Components/DMMaterialStageThroughputLayerBlend.cpp b/Engine/Plugins/Experimental/DynamicMaterial/Source/DynamicMaterialEditor/Private/Components/DMMaterialStageThroughputLayerBlend.cpp
--- a/Engine/Plugins/Experimental/DynamicMaterial/Source/DynamicMaterialEditor/Private/Components/DMMaterialStageThroughputLayerBlend.cpp
+++ b/Engine/Plugins/Experimental/DynamicMaterial/Source/DynamicMaterialEditor/Private/Components/DMMaterialStageThroughputLayerBlend.cpp
@@ -6,6 +6,7 @@
 #include "Components/DMMaterialProperty.h"
 #include "Components/DMMaterialSlot.h"
 #include "Components/DMMaterialStage.h"
+#include "Components/DMMaterialStageBlend.h"
 #include "Components/DMMaterialSubStage.h"
 #include "Components/MaterialStageExpressions/DMMSETextureSample.h"
 #include "Components/MaterialStageInputs/DMMSIExpression.h"
@@ -575,7 +576,44 @@ bool UDMMaterialStageThroughputLayerBlend::GenerateStagePreviewMaterial(UDMMater
 	UDMMaterialLayerObject* Layer = InStage->GetLayer();
 	check(Layer);
 
-	UDMMaterialSlot* Slot = Layer->GetSlot();
+	UDMMaterialStage* StageActual = InStage;
+	UDMMaterialStageSource* PreviewSource = nullptr;
+
+	if (UDMMaterialStageInputExpression* InputExpression = Cast<UDMMaterialStageInputExpression>(GetInputMask()))
+	{
+		if (UDMMaterialStageExpressionTextureSample* TextureSample = Cast<UDMMaterialStageExpressionTextureSample>(InputExpression->GetMaterialStageExpression()))
+		{
+			if (TextureSample->GetUseBaseTexture())
+			{
+				if (const UDMMaterialStageExpressionTextureSample* BaseTextureSample = TextureSample->GetBaseTextureSample())
+				{
+					UDMMaterialStage* BaseStage = BaseTextureSample->GetStage();
+
+					if (UDMMaterialSubStage* SubStage = Cast<UDMMaterialSubStage>(BaseStage))
+					{
+						BaseStage = SubStage->GetParentMostStage();
+					}
+
+					if (BaseStage)
+					{
+						if (UDMMaterialStageBlend* Blend = Cast<UDMMaterialStageBlend>(BaseStage->GetSource()))
+						{
+							if (UDMMaterialStageInput* BaseInputExpression = Blend->GetInputB())
+							{
+								StageActual = BaseStage;
+								PreviewSource = BaseInputExpression;
+							}
+						}
+					}
+				}
+			}
+		}
+	}
+
+	UDMMaterialLayerObject* LayerActual = StageActual->GetLayer();
+	check(LayerActual);
+
+	UDMMaterialSlot* Slot = LayerActual->GetSlot();
 	check(Slot);
 
 	UDynamicMaterialModelEditorOnlyData* ModelEditorOnlyData = Slot->GetMaterialModelEditorOnlyData();
@@ -590,9 +628,12 @@ bool UDMMaterialStageThroughputLayerBlend::GenerateStagePreviewMaterial(UDMMater
 	}
 
 	TSharedRef<FDMMaterialBuildState> BuildState = ModelEditorOnlyData->CreateBuildState(InPreviewMaterial);
-	BuildState->SetPreviewObject(InStage);
+	BuildState->SetPreviewObject(StageActual);
 
-	UDMMaterialStageSource* PreviewSource = GetInputMask();
+	if (!PreviewSource)
+	{
+		PreviewSource = GetInputMask();
+	}
 
 	PreviewSource->GenerateExpressions(BuildState);
 	TArray<UMaterialExpression*> SourceExpressions = BuildState->GetStageSourceExpressions(PreviewSource);
@@ -829,8 +870,50 @@ void UDMMaterialStageThroughputLayerBlend::SetPremultiplyAlpha(bool bInValue)
 int32 UDMMaterialStageThroughputLayerBlend::ResolveMaskInput(const TSharedRef<FDMMaterialBuildState>& InBuildState, int32 InputIndex, 
 	FDMMaterialStageConnectorChannel& OutChannel, TArray<UMaterialExpression*>& OutExpressions) const
 {
-	// Our baseline input. Could be a texture node or color value. Anything inputy.
-	OutChannel.OutputIndex = Super::ResolveInput(InBuildState, InputIndex, OutChannel, OutExpressions);
+	bool bIsSharedBaseTextureSample = false;
+
+	if (UDMMaterialStageInputExpression* InputExpression = Cast<UDMMaterialStageInputExpression>(GetInputMask()))
+	{
+		if (UDMMaterialStageExpressionTextureSample* TextureSample = Cast<UDMMaterialStageExpressionTextureSample>(InputExpression->GetMaterialStageExpression()))
+		{
+			if (TextureSample->GetUseBaseTexture())
+			{
+				if (const UDMMaterialStageExpressionTextureSample* BaseTextureSample = TextureSample->GetBaseTextureSample())
+				{
+					BaseTextureSample->GenerateExpressions(InBuildState);
+					OutExpressions = InBuildState->GetStageSourceExpressions(BaseTextureSample);
+
+					switch (MaskChannelOverride)
+					{
+						default:
+						case EAvaColorChannel::Red:
+							OutChannel.OutputIndex = 1;
+							break;
+
+						case EAvaColorChannel::Green:
+							OutChannel.OutputIndex = 2;
+							break;
+
+						case EAvaColorChannel::Blue:
+							OutChannel.OutputIndex = 3;
+							break;
+
+						case EAvaColorChannel::Alpha:
+							OutChannel.OutputIndex = 4;
+							break;
+					}
+
+					bIsSharedBaseTextureSample = true;
+				}
+			}
+		}
+	}
+
+	if (!bIsSharedBaseTextureSample)
+	{
+		// Our baseline input. Could be a texture node or color value. Anything inputy.
+		OutChannel.OutputIndex = Super::ResolveInput(InBuildState, InputIndex, OutChannel, OutExpressions);
+	}
 
 	UDMMaterialStage* ThisStage = GetStage();
 	check(ThisStage);
diff --git a/Engine/Plugins/Experimental/DynamicMaterial/Source/DynamicMaterialEditor/Private/Components/MaterialStageExpressions/DMMSETextureSample.cpp b/Engine/Plugins/Experimental/DynamicMaterial/Source/DynamicMaterialEditor/Private/Components/MaterialStageExpressions/DMMSETextureSample.cpp
--- a/Engine/Plugins/Experimental/DynamicMaterial/Source/DynamicMaterialEditor/Private/Components/MaterialStageExpressions/DMMSETextureSample.cpp
+++ b/Engine/Plugins/Experimental/DynamicMaterial/Source/DynamicMaterialEditor/Private/Components/MaterialStageExpressions/DMMSETextureSample.cpp
@@ -1,7 +1,14 @@
 // Copyright Epic Games, Inc. All Rights Reserved.
 
 #include "Components/MaterialStageExpressions/DMMSETextureSample.h"
+
+#include "Components/DMMaterialLayer.h"
+#include "Components/DMMaterialStageBlend.h"
+#include "Components/DMMaterialStageThroughputLayerBlend.h"
+#include "Components/DMMaterialSubStage.h"
+#include "Components/MaterialStageInputs/DMMSIThroughput.h"
 #include "Materials/MaterialExpressionTextureSample.h"
+#include "Model/DMMaterialBuildState.h"
 
 #define LOCTEXT_NAMESPACE "DMMaterialStageExpressionTextureSample"
 
@@ -12,6 +19,153 @@ UDMMaterialStageExpressionTextureSample::UDMMaterialStageExpressionTextureSample
 	)
 {
 	Menus.Add(EDMExpressionMenu::Texture);
+	EditableProperties.Add(GET_MEMBER_NAME_CHECKED(UDMMaterialStageExpressionTextureSample, bUseBaseTexture));
+}
+
+void UDMMaterialStageExpressionTextureSample::PostEditChangeProperty(FPropertyChangedEvent& PropertyChangedEvent)
+{
+	Super::PostEditChangeProperty(PropertyChangedEvent);
+
+	if (PropertyChangedEvent.GetMemberPropertyName() == GET_MEMBER_NAME_CHECKED(UDMMaterialStageExpressionTextureSample, bUseBaseTexture))
+	{
+		OnUseBaseTextureChanged();
+	}
+}
+
+bool UDMMaterialStageExpressionTextureSample::IsPropertyVisible(FName InProperty) const
+{
+	const bool bIsUseBaseTexture = (InProperty == GET_MEMBER_NAME_CHECKED(UDMMaterialStageExpressionTextureSample, bUseBaseTexture));
+	const bool bIsClampTexture = (InProperty == GET_MEMBER_NAME_CHECKED(UDMMaterialStageExpressionTextureSample, bClampTexture));
+
+	if (bIsUseBaseTexture || bIsClampTexture)
+	{
+		const bool bCanUseBaseTexture = CanUseBaseTexture();
+
+		if (bIsUseBaseTexture)
+		{
+			return bCanUseBaseTexture;
+		}
+
+		if (bIsClampTexture)
+		{
+			return !bUseBaseTexture || !bCanUseBaseTexture;
+		}
+	}
+
+	return Super::IsPropertyVisible(InProperty);
+}
+
+bool UDMMaterialStageExpressionTextureSample::GetUseBaseTexture() const
+{
+	return bUseBaseTexture;
+}
+
+void UDMMaterialStageExpressionTextureSample::SetUseBaseTexture(bool bInUseBaseTexture)
+{
+	if (bUseBaseTexture == bInUseBaseTexture)
+	{
+		return;
+	}
+
+	bUseBaseTexture = bInUseBaseTexture;
+
+	OnUseBaseTextureChanged();
+}
+
+bool UDMMaterialStageExpressionTextureSample::CanUseBaseTexture() const
+{
+	return !!GetBaseTextureSample();
+}
+
+const UDMMaterialStageExpressionTextureSample* UDMMaterialStageExpressionTextureSample::GetBaseTextureSample() const
+{
+	UDMMaterialSubStage* SubStage = Cast<UDMMaterialSubStage>(GetStage());
+
+	if (!SubStage)
+	{
+		return nullptr;
+	}
+
+	UDMMaterialStage* Stage = SubStage->GetParentStage();
+
+	if (!Stage || Stage->IsA<UDMMaterialSubStage>())
+	{
+		return nullptr;
+	}
+
+	UDMMaterialLayerObject* Layer = Stage->GetLayer();
+
+	if (!Layer || !Layer->IsTextureUVLinkEnabled())
+	{
+		return nullptr;
+	}
+
+	UDMMaterialStageThroughputLayerBlend* LayerBlend = Cast<UDMMaterialStageThroughputLayerBlend>(Stage->GetSource());
+
+	if (!LayerBlend)
+	{
+		return nullptr;
+	}
+
+	UDMMaterialStage* BaseStage = Layer->GetFirstEnabledStage(EDMMaterialLayerStage::Base);
+
+	if (!BaseStage)
+	{
+		return nullptr;
+	}
+
+	UDMMaterialStageBlend* Blend = Cast<UDMMaterialStageBlend>(BaseStage->GetSource());
+
+	if (!Blend)
+	{
+		return nullptr;
+	}
+
+	UDMMaterialStageInputThroughput* InputThroughput = Cast<UDMMaterialStageInputThroughput>(Blend->GetInputB());
+
+	if (!InputThroughput)
+	{
+		return nullptr;
+	}
+
+	UDMMaterialSubStage* BaseSubStage = InputThroughput->GetSubStage();
+
+	if (!BaseSubStage)
+	{
+		return nullptr;
+	}
+
+	return Cast<UDMMaterialStageExpressionTextureSample>(BaseSubStage->GetSource());
+}
+
+void UDMMaterialStageExpressionTextureSample::OnUseBaseTextureChanged()
+{
+	Update(this, EDMUpdateType::Structure | EDMUpdateType::AllowParentUpdate);
+}
+
+void UDMMaterialStageExpressionTextureSample::GenerateExpressions(const TSharedRef<FDMMaterialBuildState>& InBuildState) const
+{
+	if (!IsComponentValid() || !IsComponentAdded())
+	{
+		return;
+	}
+
+	if (InBuildState->HasStageSource(this))
+	{
+		return;
+	}
+
+	if (bUseBaseTexture)
+	{
+		if (const UDMMaterialStageExpressionTextureSample* BaseTextureSample = GetBaseTextureSample())
+		{
+			BaseTextureSample->GenerateExpressions(InBuildState);
+			InBuildState->AddStageSourceExpressions(this, InBuildState->GetStageSourceExpressions(BaseTextureSample));
+			return;
+		}
+	}
+
+	Super::GenerateExpressions(InBuildState);
 }
 
 #undef LOCTEXT_NAMESPACE
diff --git a/Engine/Plugins/Experimental/DynamicMaterial/Source/DynamicMaterialEditor/Private/Components/MaterialStageInputs/DMMSIValue.cpp b/Engine/Plugins/Experimental/DynamicMaterial/Source/DynamicMaterialEditor/Private/Components/MaterialStageInputs/DMMSIValue.cpp
--- a/Engine/Plugins/Experimental/DynamicMaterial/Source/DynamicMaterialEditor/Private/Components/MaterialStageInputs/DMMSIValue.cpp
+++ b/Engine/Plugins/Experimental/DynamicMaterial/Source/DynamicMaterialEditor/Private/Components/MaterialStageInputs/DMMSIValue.cpp
@@ -9,6 +9,8 @@
 #include "Components/DMMaterialStageThroughputLayerBlend.h"
 #include "Components/DMMaterialSubStage.h"
 #include "Components/DMMaterialValue.h"
+#include "Components/MaterialStageExpressions/DMMSETextureSample.h"
+#include "Components/MaterialStageInputs/DMMSIExpression.h"
 #include "Components/MaterialValues/DMMaterialValueFloat3RGB.h"
 #include "Components/MaterialValues/DMMaterialValueFloat4.h"
 #include "Components/MaterialValues/DMMaterialValueTexture.h"
@@ -469,12 +471,11 @@ void UDMMaterialStageInputValue::GenerateExpressions(const TSharedRef<FDMMateria
 		return;
 	}
 
-	check(Value);
-
 	if (InBuildState->HasStageSource(this))
 	{
 		return;
 	}
+	check(Value);
 
 	Value->GenerateExpression(InBuildState);
 	InBuildState->AddStageSourceExpressions(this, InBuildState->GetValueExpressions(Value));
@@ -733,8 +734,7 @@ void UDMMaterialStageInputValue::PostEditImport()
 	InitInputValue();
 }
 
-void UDMMaterialStageInputValue::PostEditorDuplicate(UDynamicMaterialModel* InMaterialModel,
-	UDMMaterialComponent* InParent)
+void UDMMaterialStageInputValue::PostEditorDuplicate(UDynamicMaterialModel* InMaterialModel, UDMMaterialComponent* InParent)
 {
 	Super::PostEditorDuplicate(InMaterialModel, InParent);
 
@@ -751,6 +751,16 @@ void UDMMaterialStageInputValue::PostEditorDuplicate(UDynamicMaterialModel* InMa
 	InitInputValue();
 }
 
+bool UDMMaterialStageInputValue::IsPropertyVisible(FName InProperty) const
+{
+	if (InProperty == GET_MEMBER_NAME_CHECKED(UDMMaterialStageInputValue, Value))
+	{
+		return IsTextureValueVisible();
+	}
+
+	return Super::IsPropertyVisible(InProperty);
+}
+
 UDMMaterialStageInputValue::UDMMaterialStageInputValue()
 	: Value(nullptr)
 {
@@ -827,4 +837,68 @@ void UDMMaterialStageInputValue::ApplyWholeLayerValue()
 	);
 }
 
+bool UDMMaterialStageInputValue::IsTextureValueVisible() const
+{
+	UDMMaterialValueTexture* TextureValue = Cast<UDMMaterialValueTexture>(Value);
+
+	if (!TextureValue)
+	{
+		return true;
+	}
+
+	UDMMaterialStage* Stage = GetStage();
+
+	if (!Stage)
+	{
+		return true;
+	}
+
+	if (UDMMaterialSubStage* SubStage = Cast<UDMMaterialSubStage>(Stage))
+	{
+		Stage = SubStage->GetParentMostStage();
+
+		if (!Stage)
+		{
+			return true;
+		}
+	}
+
+	UDMMaterialLayerObject* Layer = Stage->GetLayer();
+
+	if (!Layer)
+	{
+		return true;
+	}
+
+	UDMMaterialStage* MaskStage = Layer->GetFirstEnabledStage(EDMMaterialLayerStage::Mask);
+
+	if (MaskStage != Stage)
+	{
+		return true;
+	}
+
+	UDMMaterialStageThroughputLayerBlend* LayerBlend = Cast<UDMMaterialStageThroughputLayerBlend>(Stage->GetSource());
+
+	if (!LayerBlend)
+	{
+		return true;
+	}
+
+	UDMMaterialStageInputExpression* InputExpression = Cast<UDMMaterialStageInputExpression>(LayerBlend->GetInputMask());
+
+	if (!InputExpression)
+	{
+		return true;
+	}
+
+	UDMMaterialStageExpressionTextureSample* TextureSample = Cast<UDMMaterialStageExpressionTextureSample>(InputExpression->GetMaterialStageExpression());
+
+	if (!TextureSample)
+	{
+		return true;
+	}
+
+	return !TextureSample->GetUseBaseTexture() || !TextureSample->CanUseBaseTexture();
+}
+
 #undef LOCTEXT_NAMESPACE
diff --git a/Engine/Plugins/Experimental/DynamicMaterial/Source/DynamicMaterialEditor/Private/Model/DMMaterialBuildUtils.cpp b/Engine/Plugins/Experimental/DynamicMaterial/Source/DynamicMaterialEditor/Private/Model/DMMaterialBuildUtils.cpp
--- a/Engine/Plugins/Experimental/DynamicMaterial/Source/DynamicMaterialEditor/Private/Model/DMMaterialBuildUtils.cpp
+++ b/Engine/Plugins/Experimental/DynamicMaterial/Source/DynamicMaterialEditor/Private/Model/DMMaterialBuildUtils.cpp
@@ -311,7 +311,7 @@ UMaterialExpressionComponentMask* FDMMaterialBuildUtils::CreateExpressionBitMask
 	{
 		const uint32 OutputType = InExpression->GetOutputType(InOutputIndex);
 
-		Mask->R = (OutputType == MCT_Float);
+		Mask->R = (OutputType == MCT_Float1);
 		Mask->G = (OutputType == MCT_Float2);
 		Mask->B = (OutputType == MCT_Float3);
 		Mask->A = (OutputType == MCT_Float4);
diff --git a/Engine/Plugins/Experimental/DynamicMaterial/Source/DynamicMaterialEditor/Public/Components/DMMaterialStageThroughputLayerBlend.h b/Engine/Plugins/Experimental/DynamicMaterial/Source/DynamicMaterialEditor/Public/Components/DMMaterialStageThroughputLayerBlend.h
--- a/Engine/Plugins/Experimental/DynamicMaterial/Source/DynamicMaterialEditor/Public/Components/DMMaterialStageThroughputLayerBlend.h
+++ b/Engine/Plugins/Experimental/DynamicMaterial/Source/DynamicMaterialEditor/Public/Components/DMMaterialStageThroughputLayerBlend.h
@@ -85,13 +85,13 @@ public:
 	//~ End UObject
 
 protected:
+	/** Changes the output channel of the mask input. */
 	UPROPERTY(EditInstanceOnly, BlueprintReadOnly, Getter=GetMaskChannelOverride, Setter=SetMaskChannelOverride, BlueprintGetter=GetMaskChannelOverride,
-		Category = "Material Designer", DisplayName = "Channel Mask",
-		meta=(NotKeyframeable, ToolTip="Changes the output channel of the mask input."))
+		Category = "Material Designer", DisplayName = "Channel Mask", meta=(NotKeyframeable))
 	mutable EAvaColorChannel MaskChannelOverride;
 
-	UPROPERTY(EditInstanceOnly, BlueprintReadOnly, Category = "Material Designer", 
-		meta=(NotKeyframeable, ToolTip="Additionally multiplies the output of the RGB Stage by the output from this Stage."))
+	/** Additionally multiplies the output of the RGB Stage by the output from this Stage. */
+	UPROPERTY(EditInstanceOnly, BlueprintReadOnly, Category = "Material Designer", meta=(NotKeyframeable))
 	bool bPremultiplyAlpha;
 
 	UPROPERTY(EditInstanceOnly, BlueprintReadOnly, Category = "Material Designer",
diff --git a/Engine/Plugins/Experimental/DynamicMaterial/Source/DynamicMaterialEditor/Public/Components/MaterialStageExpressions/DMMSETextureSample.h b/Engine/Plugins/Experimental/DynamicMaterial/Source/DynamicMaterialEditor/Public/Components/MaterialStageExpressions/DMMSETextureSample.h
--- a/Engine/Plugins/Experimental/DynamicMaterial/Source/DynamicMaterialEditor/Public/Components/MaterialStageExpressions/DMMSETextureSample.h
+++ b/Engine/Plugins/Experimental/DynamicMaterial/Source/DynamicMaterialEditor/Public/Components/MaterialStageExpressions/DMMSETextureSample.h
@@ -15,4 +15,35 @@ class UDMMaterialStageExpressionTextureSample : public UDMMaterialStageExpressio
 
 public:
 	UDMMaterialStageExpressionTextureSample();
+
+	UFUNCTION(BlueprintPure, Category = "Material Designer")
+	bool GetUseBaseTexture() const;
+
+	UFUNCTION(BlueprintCallable, Category = "Material Designer")
+	void SetUseBaseTexture(bool bInUseBaseTexture);
+
+	UFUNCTION(BlueprintCallable, Category = "Material Designer")
+	bool CanUseBaseTexture() const;
+
+	const UDMMaterialStageExpressionTextureSample* GetBaseTextureSample() const;
+
+	//~ Begin UDMMaterialComponent
+	DYNAMICMATERIALEDITOR_API virtual void GenerateExpressions(const TSharedRef<FDMMaterialBuildState>& InBuildState) const override;
+	DYNAMICMATERIALEDITOR_API virtual bool IsPropertyVisible(FName InProperty) const override;
+	//~ End UDMMaterialComponent
+
+	//~ Begin UObject
+	DYNAMICMATERIALEDITOR_API virtual void PostEditChangeProperty(FPropertyChangedEvent& PropertyChangedEvent) override;
+	//~ End UObject
+
+protected:
+	/*
+	 * When both the base and mask stages are set to a texture and toggle layer UV link is on, use the texture sampler from
+	 * base stage directly in the mask stage, with no mask stage texture sampler.
+	 */
+	UPROPERTY(EditInstanceOnly, BlueprintReadOnly, Getter = GetUseBaseTexture, Setter = SetUseBaseTexture, BlueprintGetter = GetUseBaseTexture,
+		Category = "Material Designer", DisplayName = "Use Base Texture", meta = (NotKeyframeable))
+	bool bUseBaseTexture;
+
+	void OnUseBaseTextureChanged();
 };
diff --git a/Engine/Plugins/Experimental/DynamicMaterial/Source/DynamicMaterialEditor/Public/Components/MaterialStageInputs/DMMSIValue.h b/Engine/Plugins/Experimental/DynamicMaterial/Source/DynamicMaterialEditor/Public/Components/MaterialStageInputs/DMMSIValue.h
--- a/Engine/Plugins/Experimental/DynamicMaterial/Source/DynamicMaterialEditor/Public/Components/MaterialStageInputs/DMMSIValue.h
+++ b/Engine/Plugins/Experimental/DynamicMaterial/Source/DynamicMaterialEditor/Public/Components/MaterialStageInputs/DMMSIValue.h
@@ -81,6 +81,7 @@ public:
 	DYNAMICMATERIALEDITOR_API virtual FSlateIcon GetComponentIcon() const override;
 	DYNAMICMATERIALEDITOR_API virtual void GenerateExpressions(const TSharedRef<FDMMaterialBuildState>& InBuildState) const override;
 	DYNAMICMATERIALEDITOR_API virtual void PostEditorDuplicate(UDynamicMaterialModel* InMaterialModel, UDMMaterialComponent* InParent) override;
+	DYNAMICMATERIALEDITOR_API virtual bool IsPropertyVisible(FName InProperty) const override;
 	//~ End UDMMaterialComponent
 
 	//~ Begin UObject
@@ -102,6 +103,8 @@ protected:
 
 	void ApplyWholeLayerValue();
 
+	bool IsTextureValueVisible() const;
+
 	//~ Begin UDMMaterialStageInput
 	DYNAMICMATERIALEDITOR_API virtual void UpdateOutputConnectors() override;
 	//~ End UDMMaterialStageInput
