[PSDImporter] Initial commit of the PSD importer

#virtualized

diff --git a/Engine/Plugins/Experimental/DynamicMaterial/Source/DynamicMaterial/Private/Model/DynamicMaterialModelDynamic.cpp b/Engine/Plugins/Experimental/DynamicMaterial/Source/DynamicMaterial/Private/Model/DynamicMaterialModelDynamic.cpp
--- a/Engine/Plugins/Experimental/DynamicMaterial/Source/DynamicMaterial/Private/Model/DynamicMaterialModelDynamic.cpp
+++ b/Engine/Plugins/Experimental/DynamicMaterial/Source/DynamicMaterial/Private/Model/DynamicMaterialModelDynamic.cpp
@@ -398,6 +398,11 @@ UDMMaterialComponent* UDynamicMaterialModelDynamic::GetComponentByPath(FDMCompon
 	return nullptr;
 }
 
+const TMap<FName, TObjectPtr<UDMMaterialComponentDynamic>>& UDynamicMaterialModelDynamic::GetComponentMap() const
+{
+	return DynamicComponents;
+}
+
 #if WITH_EDITOR
 UDynamicMaterialModel* UDynamicMaterialModelDynamic::ToEditable(UObject* InOuter)
 {
diff --git a/Engine/Plugins/Experimental/DynamicMaterial/Source/DynamicMaterial/Public/Model/DynamicMaterialModelDynamic.h b/Engine/Plugins/Experimental/DynamicMaterial/Source/DynamicMaterial/Public/Model/DynamicMaterialModelDynamic.h
--- a/Engine/Plugins/Experimental/DynamicMaterial/Source/DynamicMaterial/Public/Model/DynamicMaterialModelDynamic.h
+++ b/Engine/Plugins/Experimental/DynamicMaterial/Source/DynamicMaterial/Public/Model/DynamicMaterialModelDynamic.h
@@ -90,6 +90,9 @@ public:
 	/** Finds the component with the given path. */
 	DYNAMICMATERIAL_API UDMMaterialComponent* GetComponentByPath(FDMComponentPath& InPath) const;
 
+	/** Returns the component map. */
+	DYNAMICMATERIAL_API const TMap<FName, TObjectPtr<UDMMaterialComponentDynamic>>& GetComponentMap() const;
+
 #if WITH_EDITOR
 	/** Converts this model dynamic to a new model and returns it. */
 	DYNAMICMATERIAL_API UDynamicMaterialModel* ToEditable(UObject* InOuter);
diff --git a/Engine/Plugins/Experimental/DynamicMaterial/Source/DynamicMaterialEditor/Private/Utils/DMMaterialModelFunctionLibrary.cpp b/Engine/Plugins/Experimental/DynamicMaterial/Source/DynamicMaterialEditor/Private/Utils/DMMaterialModelFunctionLibrary.cpp
--- a/Engine/Plugins/Experimental/DynamicMaterial/Source/DynamicMaterialEditor/Private/Utils/DMMaterialModelFunctionLibrary.cpp
+++ b/Engine/Plugins/Experimental/DynamicMaterial/Source/DynamicMaterialEditor/Private/Utils/DMMaterialModelFunctionLibrary.cpp
@@ -448,7 +448,8 @@ bool UDMMaterialModelFunctionLibrary::DuplicateModelBetweenMaterials(UDynamicMat
 		CurrentName = CurrentModel->GetName();
 
 		const FString NewName = CurrentName + TEXT("_OLD");
-		CurrentModel->Rename(*NewName, GetTransientPackage(), UE::DynamicMaterial::RenameFlags);
+		const FName NewNameUnique = MakeUniqueObjectName(GetTransientPackage(), CurrentModel->GetClass(), *NewName);
+		CurrentModel->Rename(*NewNameUnique.ToString(), GetTransientPackage(), UE::DynamicMaterial::RenameFlags);
 
 		InToInstance->SetMaterialModel(nullptr);
 	}
diff --git a/Engine/Plugins/Experimental/DynamicMaterial/Source/DynamicMaterialEditor/Public/Components/MaterialStageExpressions/DMMSETextureSample.h b/Engine/Plugins/Experimental/DynamicMaterial/Source/DynamicMaterialEditor/Public/Components/MaterialStageExpressions/DMMSETextureSample.h
--- a/Engine/Plugins/Experimental/DynamicMaterial/Source/DynamicMaterialEditor/Public/Components/MaterialStageExpressions/DMMSETextureSample.h
+++ b/Engine/Plugins/Experimental/DynamicMaterial/Source/DynamicMaterialEditor/Public/Components/MaterialStageExpressions/DMMSETextureSample.h
@@ -17,15 +17,15 @@ public:
 	UDMMaterialStageExpressionTextureSample();
 
 	UFUNCTION(BlueprintPure, Category = "Material Designer")
-	bool GetUseBaseTexture() const;
+	DYNAMICMATERIALEDITOR_API bool GetUseBaseTexture() const;
 
 	UFUNCTION(BlueprintCallable, Category = "Material Designer")
-	void SetUseBaseTexture(bool bInUseBaseTexture);
+	DYNAMICMATERIALEDITOR_API void SetUseBaseTexture(bool bInUseBaseTexture);
 
 	UFUNCTION(BlueprintCallable, Category = "Material Designer")
-	bool CanUseBaseTexture() const;
+	DYNAMICMATERIALEDITOR_API bool CanUseBaseTexture() const;
 
-	const UDMMaterialStageExpressionTextureSample* GetBaseTextureSample() const;
+	DYNAMICMATERIALEDITOR_API const UDMMaterialStageExpressionTextureSample* GetBaseTextureSample() const;
 
 	//~ Begin UDMMaterialComponent
 	DYNAMICMATERIALEDITOR_API virtual void GenerateExpressions(const TSharedRef<FDMMaterialBuildState>& InBuildState) const override;
