Material Designer / Media Stream Integration

diff --git a/Engine/Plugins/Experimental/DynamicMaterial/Source/DynamicMaterial/Private/Components/DMMaterialComponent.cpp b/Engine/Plugins/Experimental/DynamicMaterial/Source/DynamicMaterial/Private/Components/DMMaterialComponent.cpp
--- a/Engine/Plugins/Experimental/DynamicMaterial/Source/DynamicMaterial/Private/Components/DMMaterialComponent.cpp
+++ b/Engine/Plugins/Experimental/DynamicMaterial/Source/DynamicMaterial/Private/Components/DMMaterialComponent.cpp
@@ -86,6 +86,14 @@ void UDMMaterialComponent::Update(UDMMaterialComponent* InSource, EDMUpdateType
 	}
 
 #if WITH_EDITOR
+	if (EnumHasAnyFlags(InUpdateType, EDMUpdateType::RefreshDetailView))
+	{
+		if (UDMMaterialComponent* ParentComponent = GetParentComponent())
+		{
+			ParentComponent->Update(InSource, InUpdateType);
+		}
+	}
+
 	OnUpdate.Broadcast(this, InSource, InUpdateType);
 #endif
 }
diff --git a/Engine/Plugins/Experimental/DynamicMaterial/Source/DynamicMaterial/Public/Components/MaterialValues/DMMaterialValueBool.h b/Engine/Plugins/Experimental/DynamicMaterial/Source/DynamicMaterial/Public/Components/MaterialValues/DMMaterialValueBool.h
--- a/Engine/Plugins/Experimental/DynamicMaterial/Source/DynamicMaterial/Public/Components/MaterialValues/DMMaterialValueBool.h
+++ b/Engine/Plugins/Experimental/DynamicMaterial/Source/DynamicMaterial/Public/Components/MaterialValues/DMMaterialValueBool.h
@@ -14,7 +14,7 @@ class UDMMaterialValueBool : public UDMMaterialValue
 	GENERATED_BODY()
  
 public:
-	UDMMaterialValueBool();
+	DYNAMICMATERIAL_API UDMMaterialValueBool();
 
 	UFUNCTION(BlueprintPure, Category = "Material Designer")
 	bool GetValue() const { return Value; }
diff --git a/Engine/Plugins/Experimental/DynamicMaterial/Source/DynamicMaterial/Public/Components/MaterialValues/DMMaterialValueFloat.h b/Engine/Plugins/Experimental/DynamicMaterial/Source/DynamicMaterial/Public/Components/MaterialValues/DMMaterialValueFloat.h
--- a/Engine/Plugins/Experimental/DynamicMaterial/Source/DynamicMaterial/Public/Components/MaterialValues/DMMaterialValueFloat.h
+++ b/Engine/Plugins/Experimental/DynamicMaterial/Source/DynamicMaterial/Public/Components/MaterialValues/DMMaterialValueFloat.h
@@ -14,7 +14,7 @@ class UDMMaterialValueFloat : public UDMMaterialValue
 	GENERATED_BODY()
  
 public:
-	UDMMaterialValueFloat();
+	DYNAMICMATERIAL_API UDMMaterialValueFloat();
 
 	/** The value range for any float component. If both values are the same, it is not set. */
 	UFUNCTION(BlueprintPure, Category = "Material Designer")
diff --git a/Engine/Plugins/Experimental/DynamicMaterial/Source/DynamicMaterial/Public/Components/MaterialValues/DMMaterialValueFloat1.h b/Engine/Plugins/Experimental/DynamicMaterial/Source/DynamicMaterial/Public/Components/MaterialValues/DMMaterialValueFloat1.h
--- a/Engine/Plugins/Experimental/DynamicMaterial/Source/DynamicMaterial/Public/Components/MaterialValues/DMMaterialValueFloat1.h
+++ b/Engine/Plugins/Experimental/DynamicMaterial/Source/DynamicMaterial/Public/Components/MaterialValues/DMMaterialValueFloat1.h
@@ -14,7 +14,7 @@ class UDMMaterialValueFloat1 : public UDMMaterialValueFloat
 	GENERATED_BODY()
  
 public:
-	UDMMaterialValueFloat1();
+	DYNAMICMATERIAL_API UDMMaterialValueFloat1();
 
 	UFUNCTION(BlueprintPure, Category = "Material Designer")
 	float GetValue() const { return Value; }
diff --git a/Engine/Plugins/Experimental/DynamicMaterial/Source/DynamicMaterial/Public/Components/MaterialValues/DMMaterialValueFloat2.h b/Engine/Plugins/Experimental/DynamicMaterial/Source/DynamicMaterial/Public/Components/MaterialValues/DMMaterialValueFloat2.h
--- a/Engine/Plugins/Experimental/DynamicMaterial/Source/DynamicMaterial/Public/Components/MaterialValues/DMMaterialValueFloat2.h
+++ b/Engine/Plugins/Experimental/DynamicMaterial/Source/DynamicMaterial/Public/Components/MaterialValues/DMMaterialValueFloat2.h
@@ -14,7 +14,7 @@ class UDMMaterialValueFloat2 : public UDMMaterialValueFloat
 	GENERATED_BODY()
  
 public: 
-	UDMMaterialValueFloat2();
+	DYNAMICMATERIAL_API UDMMaterialValueFloat2();
 
 	UFUNCTION(BlueprintPure, Category = "Material Designer")
 	const FVector2D& GetValue() const { return Value; }
diff --git a/Engine/Plugins/Experimental/DynamicMaterial/Source/DynamicMaterial/Public/Components/MaterialValues/DMMaterialValueFloat3RGB.h b/Engine/Plugins/Experimental/DynamicMaterial/Source/DynamicMaterial/Public/Components/MaterialValues/DMMaterialValueFloat3RGB.h
--- a/Engine/Plugins/Experimental/DynamicMaterial/Source/DynamicMaterial/Public/Components/MaterialValues/DMMaterialValueFloat3RGB.h
+++ b/Engine/Plugins/Experimental/DynamicMaterial/Source/DynamicMaterial/Public/Components/MaterialValues/DMMaterialValueFloat3RGB.h
@@ -14,7 +14,7 @@ class UDMMaterialValueFloat3RGB : public UDMMaterialValueFloat
 	GENERATED_BODY()
  
 public: 
-	UDMMaterialValueFloat3RGB();
+	DYNAMICMATERIAL_API UDMMaterialValueFloat3RGB();
 
 	UFUNCTION(BlueprintPure, Category = "Material Designer")
 	const FLinearColor& GetValue() const { return Value; }
diff --git a/Engine/Plugins/Experimental/DynamicMaterial/Source/DynamicMaterial/Public/Components/MaterialValues/DMMaterialValueFloat3RPY.h b/Engine/Plugins/Experimental/DynamicMaterial/Source/DynamicMaterial/Public/Components/MaterialValues/DMMaterialValueFloat3RPY.h
--- a/Engine/Plugins/Experimental/DynamicMaterial/Source/DynamicMaterial/Public/Components/MaterialValues/DMMaterialValueFloat3RPY.h
+++ b/Engine/Plugins/Experimental/DynamicMaterial/Source/DynamicMaterial/Public/Components/MaterialValues/DMMaterialValueFloat3RPY.h
@@ -14,7 +14,7 @@ class UDMMaterialValueFloat3RPY : public UDMMaterialValueFloat
 	GENERATED_BODY()
  
 public: 
-	UDMMaterialValueFloat3RPY();
+	DYNAMICMATERIAL_API UDMMaterialValueFloat3RPY();
 
 	UFUNCTION(BlueprintPure, Category = "Material Designer")
 	const FRotator& GetValue() const { return Value; }
diff --git a/Engine/Plugins/Experimental/DynamicMaterial/Source/DynamicMaterial/Public/Components/MaterialValues/DMMaterialValueFloat3XYZ.h b/Engine/Plugins/Experimental/DynamicMaterial/Source/DynamicMaterial/Public/Components/MaterialValues/DMMaterialValueFloat3XYZ.h
--- a/Engine/Plugins/Experimental/DynamicMaterial/Source/DynamicMaterial/Public/Components/MaterialValues/DMMaterialValueFloat3XYZ.h
+++ b/Engine/Plugins/Experimental/DynamicMaterial/Source/DynamicMaterial/Public/Components/MaterialValues/DMMaterialValueFloat3XYZ.h
@@ -14,7 +14,7 @@ class UDMMaterialValueFloat3XYZ : public UDMMaterialValueFloat
 	GENERATED_BODY()
  
 public: 
-	UDMMaterialValueFloat3XYZ();
+	DYNAMICMATERIAL_API UDMMaterialValueFloat3XYZ();
 
 	UFUNCTION(BlueprintPure, Category = "Material Designer")
 	const FVector& GetValue() const { return Value; }
diff --git a/Engine/Plugins/Experimental/DynamicMaterial/Source/DynamicMaterial/Public/Components/MaterialValues/DMMaterialValueFloat4.h b/Engine/Plugins/Experimental/DynamicMaterial/Source/DynamicMaterial/Public/Components/MaterialValues/DMMaterialValueFloat4.h
--- a/Engine/Plugins/Experimental/DynamicMaterial/Source/DynamicMaterial/Public/Components/MaterialValues/DMMaterialValueFloat4.h
+++ b/Engine/Plugins/Experimental/DynamicMaterial/Source/DynamicMaterial/Public/Components/MaterialValues/DMMaterialValueFloat4.h
@@ -14,7 +14,7 @@ class UDMMaterialValueFloat4 : public UDMMaterialValueFloat
 	GENERATED_BODY()
  
 public: 
-	UDMMaterialValueFloat4();
+	DYNAMICMATERIAL_API UDMMaterialValueFloat4();
 
 	UFUNCTION(BlueprintPure, Category = "Material Designer")
 	const FLinearColor& GetValue() const { return Value; }
diff --git a/Engine/Plugins/Experimental/DynamicMaterial/Source/DynamicMaterial/Public/Components/MaterialValues/DMMaterialValueRenderTarget.h b/Engine/Plugins/Experimental/DynamicMaterial/Source/DynamicMaterial/Public/Components/MaterialValues/DMMaterialValueRenderTarget.h
--- a/Engine/Plugins/Experimental/DynamicMaterial/Source/DynamicMaterial/Public/Components/MaterialValues/DMMaterialValueRenderTarget.h
+++ b/Engine/Plugins/Experimental/DynamicMaterial/Source/DynamicMaterial/Public/Components/MaterialValues/DMMaterialValueRenderTarget.h
@@ -22,7 +22,7 @@ class UDMMaterialValueRenderTarget : public UDMMaterialValueTexture
 public:
 	DYNAMICMATERIAL_API static const FString RendererPathToken;
 
-	UDMMaterialValueRenderTarget();
+	DYNAMICMATERIAL_API UDMMaterialValueRenderTarget();
 
 	virtual ~UDMMaterialValueRenderTarget() override;
 
diff --git a/Engine/Plugins/Experimental/DynamicMaterial/Source/DynamicMaterial/Public/Components/MaterialValues/DMMaterialValueTexture.h b/Engine/Plugins/Experimental/DynamicMaterial/Source/DynamicMaterial/Public/Components/MaterialValues/DMMaterialValueTexture.h
--- a/Engine/Plugins/Experimental/DynamicMaterial/Source/DynamicMaterial/Public/Components/MaterialValues/DMMaterialValueTexture.h
+++ b/Engine/Plugins/Experimental/DynamicMaterial/Source/DynamicMaterial/Public/Components/MaterialValues/DMMaterialValueTexture.h
@@ -26,7 +26,7 @@ class UDMMaterialValueTexture : public UDMMaterialValue
 	GENERATED_BODY()
  
 public:
-	UDMMaterialValueTexture();
+	DYNAMICMATERIAL_API UDMMaterialValueTexture();
 
 #if WITH_EDITOR
 	DYNAMICMATERIAL_API static FDMGetDefaultRGBTexture GetDefaultRGBTexture;
diff --git a/Engine/Plugins/Experimental/DynamicMaterial/Source/DynamicMaterial/Public/Components/MaterialValuesDynamic/DMMaterialValueBoolDynamic.h b/Engine/Plugins/Experimental/DynamicMaterial/Source/DynamicMaterial/Public/Components/MaterialValuesDynamic/DMMaterialValueBoolDynamic.h
--- a/Engine/Plugins/Experimental/DynamicMaterial/Source/DynamicMaterial/Public/Components/MaterialValuesDynamic/DMMaterialValueBoolDynamic.h
+++ b/Engine/Plugins/Experimental/DynamicMaterial/Source/DynamicMaterial/Public/Components/MaterialValuesDynamic/DMMaterialValueBoolDynamic.h
@@ -14,7 +14,7 @@ class UDMMaterialValueBoolDynamic : public UDMMaterialValueDynamic
 	GENERATED_BODY()
 
 public:
-	UDMMaterialValueBoolDynamic();
+	DYNAMICMATERIAL_API UDMMaterialValueBoolDynamic();
 
 	UFUNCTION(BlueprintPure, Category = "Material Designer")
 	bool GetValue() const { return Value; }
diff --git a/Engine/Plugins/Experimental/DynamicMaterial/Source/DynamicMaterial/Public/Components/MaterialValuesDynamic/DMMaterialValueColorAtlasDynamic.h b/Engine/Plugins/Experimental/DynamicMaterial/Source/DynamicMaterial/Public/Components/MaterialValuesDynamic/DMMaterialValueColorAtlasDynamic.h
--- a/Engine/Plugins/Experimental/DynamicMaterial/Source/DynamicMaterial/Public/Components/MaterialValuesDynamic/DMMaterialValueColorAtlasDynamic.h
+++ b/Engine/Plugins/Experimental/DynamicMaterial/Source/DynamicMaterial/Public/Components/MaterialValuesDynamic/DMMaterialValueColorAtlasDynamic.h
@@ -14,7 +14,7 @@ class UDMMaterialValueColorAtlasDynamic : public UDMMaterialValueDynamic
 	GENERATED_BODY()
 
 public:
-	UDMMaterialValueColorAtlasDynamic();
+	DYNAMICMATERIAL_API UDMMaterialValueColorAtlasDynamic();
 
 	UFUNCTION(BlueprintPure, Category = "Material Designer")
 	float GetValue() const { return Value; }
diff --git a/Engine/Plugins/Experimental/DynamicMaterial/Source/DynamicMaterial/Public/Components/MaterialValuesDynamic/DMMaterialValueFloat1Dynamic.h b/Engine/Plugins/Experimental/DynamicMaterial/Source/DynamicMaterial/Public/Components/MaterialValuesDynamic/DMMaterialValueFloat1Dynamic.h
--- a/Engine/Plugins/Experimental/DynamicMaterial/Source/DynamicMaterial/Public/Components/MaterialValuesDynamic/DMMaterialValueFloat1Dynamic.h
+++ b/Engine/Plugins/Experimental/DynamicMaterial/Source/DynamicMaterial/Public/Components/MaterialValuesDynamic/DMMaterialValueFloat1Dynamic.h
@@ -14,7 +14,7 @@ class UDMMaterialValueFloat1Dynamic : public UDMMaterialValueDynamic
 	GENERATED_BODY()
  
 public:
-	UDMMaterialValueFloat1Dynamic();
+	DYNAMICMATERIAL_API UDMMaterialValueFloat1Dynamic();
 
 	UFUNCTION(BlueprintPure, Category = "Material Designer")
 	float GetValue() const { return Value; }
diff --git a/Engine/Plugins/Experimental/DynamicMaterial/Source/DynamicMaterial/Public/Components/MaterialValuesDynamic/DMMaterialValueFloat2Dynamic.h b/Engine/Plugins/Experimental/DynamicMaterial/Source/DynamicMaterial/Public/Components/MaterialValuesDynamic/DMMaterialValueFloat2Dynamic.h
--- a/Engine/Plugins/Experimental/DynamicMaterial/Source/DynamicMaterial/Public/Components/MaterialValuesDynamic/DMMaterialValueFloat2Dynamic.h
+++ b/Engine/Plugins/Experimental/DynamicMaterial/Source/DynamicMaterial/Public/Components/MaterialValuesDynamic/DMMaterialValueFloat2Dynamic.h
@@ -15,7 +15,7 @@ class UDMMaterialValueFloat2Dynamic : public UDMMaterialValueDynamic
 	GENERATED_BODY()
 
 public:
-	UDMMaterialValueFloat2Dynamic();
+	DYNAMICMATERIAL_API UDMMaterialValueFloat2Dynamic();
 
 	UFUNCTION(BlueprintPure, Category = "Material Designer")
 	const FVector2D& GetValue() const { return Value; }
diff --git a/Engine/Plugins/Experimental/DynamicMaterial/Source/DynamicMaterial/Public/Components/MaterialValuesDynamic/DMMaterialValueFloat3RGBDynamic.h b/Engine/Plugins/Experimental/DynamicMaterial/Source/DynamicMaterial/Public/Components/MaterialValuesDynamic/DMMaterialValueFloat3RGBDynamic.h
--- a/Engine/Plugins/Experimental/DynamicMaterial/Source/DynamicMaterial/Public/Components/MaterialValuesDynamic/DMMaterialValueFloat3RGBDynamic.h
+++ b/Engine/Plugins/Experimental/DynamicMaterial/Source/DynamicMaterial/Public/Components/MaterialValuesDynamic/DMMaterialValueFloat3RGBDynamic.h
@@ -15,7 +15,7 @@ class UDMMaterialValueFloat3RGBDynamic : public UDMMaterialValueDynamic
 	GENERATED_BODY()
 
 public:
-	UDMMaterialValueFloat3RGBDynamic();
+	DYNAMICMATERIAL_API UDMMaterialValueFloat3RGBDynamic();
 
 	UFUNCTION(BlueprintPure, Category = "Material Designer")
 	const FLinearColor& GetValue() const { return Value; }
diff --git a/Engine/Plugins/Experimental/DynamicMaterial/Source/DynamicMaterial/Public/Components/MaterialValuesDynamic/DMMaterialValueFloat3RPYDynamic.h b/Engine/Plugins/Experimental/DynamicMaterial/Source/DynamicMaterial/Public/Components/MaterialValuesDynamic/DMMaterialValueFloat3RPYDynamic.h
--- a/Engine/Plugins/Experimental/DynamicMaterial/Source/DynamicMaterial/Public/Components/MaterialValuesDynamic/DMMaterialValueFloat3RPYDynamic.h
+++ b/Engine/Plugins/Experimental/DynamicMaterial/Source/DynamicMaterial/Public/Components/MaterialValuesDynamic/DMMaterialValueFloat3RPYDynamic.h
@@ -15,7 +15,7 @@ class UDMMaterialValueFloat3RPYDynamic : public UDMMaterialValueDynamic
 	GENERATED_BODY()
 
 public:
-	UDMMaterialValueFloat3RPYDynamic();
+	DYNAMICMATERIAL_API UDMMaterialValueFloat3RPYDynamic();
 
 	UFUNCTION(BlueprintPure, Category = "Material Designer")
 	const FRotator& GetValue() const { return Value; }
diff --git a/Engine/Plugins/Experimental/DynamicMaterial/Source/DynamicMaterial/Public/Components/MaterialValuesDynamic/DMMaterialValueFloat3XYZDynamic.h b/Engine/Plugins/Experimental/DynamicMaterial/Source/DynamicMaterial/Public/Components/MaterialValuesDynamic/DMMaterialValueFloat3XYZDynamic.h
--- a/Engine/Plugins/Experimental/DynamicMaterial/Source/DynamicMaterial/Public/Components/MaterialValuesDynamic/DMMaterialValueFloat3XYZDynamic.h
+++ b/Engine/Plugins/Experimental/DynamicMaterial/Source/DynamicMaterial/Public/Components/MaterialValuesDynamic/DMMaterialValueFloat3XYZDynamic.h
@@ -15,7 +15,7 @@ class UDMMaterialValueFloat3XYZDynamic : public UDMMaterialValueDynamic
 	GENERATED_BODY()
 
 public:
-	UDMMaterialValueFloat3XYZDynamic();
+	DYNAMICMATERIAL_API UDMMaterialValueFloat3XYZDynamic();
 
 	UFUNCTION(BlueprintPure, Category = "Material Designer")
 	const FVector& GetValue() const { return Value; }
diff --git a/Engine/Plugins/Experimental/DynamicMaterial/Source/DynamicMaterial/Public/Components/MaterialValuesDynamic/DMMaterialValueFloat4Dynamic.h b/Engine/Plugins/Experimental/DynamicMaterial/Source/DynamicMaterial/Public/Components/MaterialValuesDynamic/DMMaterialValueFloat4Dynamic.h
--- a/Engine/Plugins/Experimental/DynamicMaterial/Source/DynamicMaterial/Public/Components/MaterialValuesDynamic/DMMaterialValueFloat4Dynamic.h
+++ b/Engine/Plugins/Experimental/DynamicMaterial/Source/DynamicMaterial/Public/Components/MaterialValuesDynamic/DMMaterialValueFloat4Dynamic.h
@@ -15,7 +15,7 @@ class UDMMaterialValueFloat4Dynamic : public UDMMaterialValueDynamic
 	GENERATED_BODY()
 
 public:
-	UDMMaterialValueFloat4Dynamic();
+	DYNAMICMATERIAL_API UDMMaterialValueFloat4Dynamic();
 
 	UFUNCTION(BlueprintPure, Category = "Material Designer")
 	const FLinearColor& GetValue() const { return Value; }
diff --git a/Engine/Plugins/Experimental/DynamicMaterial/Source/DynamicMaterial/Public/Components/MaterialValuesDynamic/DMMaterialValueRenderTargetDynamic.h b/Engine/Plugins/Experimental/DynamicMaterial/Source/DynamicMaterial/Public/Components/MaterialValuesDynamic/DMMaterialValueRenderTargetDynamic.h
--- a/Engine/Plugins/Experimental/DynamicMaterial/Source/DynamicMaterial/Public/Components/MaterialValuesDynamic/DMMaterialValueRenderTargetDynamic.h
+++ b/Engine/Plugins/Experimental/DynamicMaterial/Source/DynamicMaterial/Public/Components/MaterialValuesDynamic/DMMaterialValueRenderTargetDynamic.h
@@ -12,7 +12,7 @@ class UDMMaterialValueRenderTargetDynamic : public UDMMaterialValueDynamic
 	GENERATED_BODY()
 
 public:
-	UDMMaterialValueRenderTargetDynamic();
+	DYNAMICMATERIAL_API UDMMaterialValueRenderTargetDynamic();
 
 #if WITH_EDITOR
 	//~ Begin IDMJsonSerializable
diff --git a/Engine/Plugins/Experimental/DynamicMaterial/Source/DynamicMaterial/Public/Components/MaterialValuesDynamic/DMMaterialValueTextureDynamic.h b/Engine/Plugins/Experimental/DynamicMaterial/Source/DynamicMaterial/Public/Components/MaterialValuesDynamic/DMMaterialValueTextureDynamic.h
--- a/Engine/Plugins/Experimental/DynamicMaterial/Source/DynamicMaterial/Public/Components/MaterialValuesDynamic/DMMaterialValueTextureDynamic.h
+++ b/Engine/Plugins/Experimental/DynamicMaterial/Source/DynamicMaterial/Public/Components/MaterialValuesDynamic/DMMaterialValueTextureDynamic.h
@@ -16,7 +16,7 @@ class UDMMaterialValueTextureDynamic : public UDMMaterialValueDynamic
 	GENERATED_BODY()
 
 public:
-	UDMMaterialValueTextureDynamic();
+	DYNAMICMATERIAL_API UDMMaterialValueTextureDynamic();
 
 	UFUNCTION(BlueprintPure, Category = "Material Designer")
 	UTexture* GetValue() const { return Value; }
diff --git a/Engine/Plugins/Experimental/DynamicMaterial/Source/DynamicMaterial/Public/DMDefs.h b/Engine/Plugins/Experimental/DynamicMaterial/Source/DynamicMaterial/Public/DMDefs.h
--- a/Engine/Plugins/Experimental/DynamicMaterial/Source/DynamicMaterial/Public/DMDefs.h
+++ b/Engine/Plugins/Experimental/DynamicMaterial/Source/DynamicMaterial/Public/DMDefs.h
@@ -79,7 +79,8 @@ enum class EDMUpdateType : uint8
 {
 	Value = 0,
 	Structure = 1 << 0,
-	AllowParentUpdate = 1 << 1
+	AllowParentUpdate = 1 << 1,
+	RefreshDetailView = 1 << 2
 };
 ENUM_CLASS_FLAGS(EDMUpdateType)
 
diff --git a/Engine/Plugins/Experimental/DynamicMaterial/Source/DynamicMaterialEditor/Private/DynamicMaterialEditorModule.cpp b/Engine/Plugins/Experimental/DynamicMaterial/Source/DynamicMaterialEditor/Private/DynamicMaterialEditorModule.cpp
--- a/Engine/Plugins/Experimental/DynamicMaterial/Source/DynamicMaterialEditor/Private/DynamicMaterialEditorModule.cpp
+++ b/Engine/Plugins/Experimental/DynamicMaterial/Source/DynamicMaterialEditor/Private/DynamicMaterialEditorModule.cpp
@@ -45,7 +45,7 @@
 #include "UI/PropertyGenerators/DMTextureUVDynamicPropertyRowGenerator.h"
 #include "UI/PropertyGenerators/DMTextureUVPropertyRowGenerator.h"
 #include "UI/PropertyGenerators/DMThroughputPropertyRowGenerator.h"
-#include "UI/Utils/DMWidgetStatics.h"
+#include "UI/Utils/DMWidgetLibrary.h"
 #include "UI/Utils/DynamicMaterialInstanceThumbnailRenderer.h"
 #include "UI/Widgets/SDMMaterialDesigner.h"
 
@@ -235,16 +235,16 @@ void FDynamicMaterialEditorModule::StartupModule()
 
 	FDMMaterialModelDefaults::RegisterDefaultsDelegates();
 
-	RegisterComponentPropertyRowGeneratorDelegate<UDMMaterialComponent,            FDMComponentPropertyRowGenerator>();
-	RegisterComponentPropertyRowGeneratorDelegate<UDMMaterialStage,                FDMStagePropertyRowGenerator>();
-	RegisterComponentPropertyRowGeneratorDelegate<UDMMaterialValue,                FDMMaterialValuePropertyRowGenerator>();
-	RegisterComponentPropertyRowGeneratorDelegate<UDMMaterialValueDynamic,         FDMMaterialValueDynamicPropertyRowGenerator>();
-	RegisterComponentPropertyRowGeneratorDelegate<UDMTextureUV,                    FDMTextureUVPropertyRowGenerator>();
-	RegisterComponentPropertyRowGeneratorDelegate<UDMTextureUVDynamic,             FDMTextureUVDynamicPropertyRowGenerator>();
-	RegisterComponentPropertyRowGeneratorDelegate<UDMMaterialStageThroughput,      FDMThroughputPropertyRowGenerator>();
-	RegisterComponentPropertyRowGeneratorDelegate<UDMMaterialStageInputThroughput, FDMInputThroughputPropertyRowGenerator>();
-	RegisterComponentPropertyRowGeneratorDelegate<UDMMaterialEffectFunction,       FDMMaterialEffectFunctionPropertyRowGenerator>();
-	RegisterComponentPropertyRowGeneratorDelegate<UDMMaterialStageInputFunction,   FDMMaterialStageFunctionPropertyRowGenerator>();
+	IDynamicMaterialEditorModule::RegisterComponentPropertyRowGeneratorDelegate<UDMMaterialComponent,            FDMComponentPropertyRowGenerator>();
+	IDynamicMaterialEditorModule::RegisterComponentPropertyRowGeneratorDelegate<UDMMaterialStage,                FDMStagePropertyRowGenerator>();
+	IDynamicMaterialEditorModule::RegisterComponentPropertyRowGeneratorDelegate<UDMMaterialValue,                FDMMaterialValuePropertyRowGenerator>();
+	IDynamicMaterialEditorModule::RegisterComponentPropertyRowGeneratorDelegate<UDMMaterialValueDynamic,         FDMMaterialValueDynamicPropertyRowGenerator>();
+	IDynamicMaterialEditorModule::RegisterComponentPropertyRowGeneratorDelegate<UDMTextureUV,                    FDMTextureUVPropertyRowGenerator>();
+	IDynamicMaterialEditorModule::RegisterComponentPropertyRowGeneratorDelegate<UDMTextureUVDynamic,             FDMTextureUVDynamicPropertyRowGenerator>();
+	IDynamicMaterialEditorModule::RegisterComponentPropertyRowGeneratorDelegate<UDMMaterialStageThroughput,      FDMThroughputPropertyRowGenerator>();
+	IDynamicMaterialEditorModule::RegisterComponentPropertyRowGeneratorDelegate<UDMMaterialStageInputThroughput, FDMInputThroughputPropertyRowGenerator>();
+	IDynamicMaterialEditorModule::RegisterComponentPropertyRowGeneratorDelegate<UDMMaterialEffectFunction,       FDMMaterialEffectFunctionPropertyRowGenerator>();
+	IDynamicMaterialEditorModule::RegisterComponentPropertyRowGeneratorDelegate<UDMMaterialStageInputFunction,   FDMMaterialStageFunctionPropertyRowGenerator>();
 
 	UDMMaterialValueTexture::GetDefaultRGBTexture.BindLambda([]()
 		{
@@ -290,7 +290,7 @@ void FDynamicMaterialEditorModule::ShutdownModule()
 
 	FDMValueDetailsRowExtensions::Get().UnregisterRowExtensions();
 
-	FDMWidgetStatics::Get().ClearData();
+	FDMWidgetLibrary::Get().ClearData();
 }
 
 void FDynamicMaterialEditorModule::OpenMaterialModel(UDynamicMaterialModelBase* InMaterialModel, UWorld* InWorld, bool bInInvokeTab) const
diff --git a/Engine/Plugins/Experimental/DynamicMaterial/Source/DynamicMaterialEditor/Private/DynamicMaterialEditorModule.h b/Engine/Plugins/Experimental/DynamicMaterial/Source/DynamicMaterialEditor/Private/DynamicMaterialEditorModule.h
--- a/Engine/Plugins/Experimental/DynamicMaterial/Source/DynamicMaterialEditor/Private/DynamicMaterialEditorModule.h
+++ b/Engine/Plugins/Experimental/DynamicMaterial/Source/DynamicMaterialEditor/Private/DynamicMaterialEditorModule.h
@@ -35,10 +35,6 @@ DECLARE_MULTICAST_DELEGATE(FDMOnUIValueUpdate);
 /** Takes a UMaterialValue and returns the widget used to edit it. */
 DECLARE_DELEGATE_RetVal_TwoParams(TSharedPtr<SWidget>, FDMCreateValueEditWidgetDelegate, const TSharedPtr<SDMMaterialComponentEditor>&, UDMMaterialValue*);
 
-/** Creates property rows in the edit widget. */
-DECLARE_DELEGATE_FourParams(FDMComponentPropertyRowGeneratorDelegate, const TSharedRef<SDMMaterialComponentEditor>&, UDMMaterialComponent*,
-	TArray<FDMPropertyHandle>&, TSet<UDMMaterialComponent*>&)
-
 /**
  * Material Designer - Build your own materials in a slimline editor!
  */
@@ -50,8 +46,6 @@ public:
 
 	static FDynamicMaterialEditorModule& Get();
 
-	static void RegisterComponentPropertyRowGeneratorDelegate(UClass* InClass, FDMComponentPropertyRowGeneratorDelegate InComponentPropertyRowGeneratorDelegate);
-	template<class InObjClass, class InGenClass> static void RegisterComponentPropertyRowGeneratorDelegate();
 	static FDMComponentPropertyRowGeneratorDelegate GetComponentPropertyRowGeneratorDelegate(UClass* InClass);
 	static void GeneratorComponentPropertyRows(const TSharedRef<SDMMaterialComponentEditor>& InComponentEditorWidget, UDMMaterialComponent* InComponent, 
 		TArray<FDMPropertyHandle>& InOutPropertyRows, TSet<UDMMaterialComponent*>& InOutProcessedObjects);
@@ -68,6 +62,7 @@ public:
 	void OnWizardComplete(UDynamicMaterialModel* InModel);
 
 	//~ Begin IDynamicMaterialEditorModule
+	virtual void RegisterComponentPropertyRowGeneratorDelegate(UClass* InClass, FDMComponentPropertyRowGeneratorDelegate InComponentPropertyRowGeneratorDelegate) override;
 	virtual void RegisterCustomMaterialPropertyGenerator(UClass* InClass, FDMGetObjectMaterialPropertiesDelegate InGenerator) override;
 	virtual void RegisterMaterialModelCreatedCallback(const TSharedRef<IDMOnWizardCompleteCallback> InCallback)  override;
 	virtual void UnregisterMaterialModelCreatedCallback(const TSharedRef<IDMOnWizardCompleteCallback> InCallback) override;
@@ -96,15 +91,3 @@ protected:
 	void MapCommands();
 	void UnmapCommands();
 };
-
-template <class InObjClass, class InGenClass>
-void FDynamicMaterialEditorModule::RegisterComponentPropertyRowGeneratorDelegate()
-{
-	RegisterComponentPropertyRowGeneratorDelegate(
-		InObjClass::StaticClass(),
-		FDMComponentPropertyRowGeneratorDelegate::CreateSP(
-			InGenClass::Get(),
-			&InGenClass::AddComponentProperties
-		)
-	);
-}
diff --git a/Engine/Plugins/Experimental/DynamicMaterial/Source/DynamicMaterialEditor/Private/UI/Menus/DMMaterialStageSourceMenus.cpp b/Engine/Plugins/Experimental/DynamicMaterial/Source/DynamicMaterialEditor/Private/UI/Menus/DMMaterialStageSourceMenus.cpp
--- a/Engine/Plugins/Experimental/DynamicMaterial/Source/DynamicMaterialEditor/Private/UI/Menus/DMMaterialStageSourceMenus.cpp
+++ b/Engine/Plugins/Experimental/DynamicMaterial/Source/DynamicMaterialEditor/Private/UI/Menus/DMMaterialStageSourceMenus.cpp
@@ -64,7 +64,7 @@ TSharedRef<SWidget> FDMMaterialStageSourceMenus::MakeChangeSourceMenu(const TSha
 			return SNullWidget::NullWidget;
 		}
 
-		FToolMenuSection& NewSection = NewToolMenu->AddSection("ChangeMaterialStage", LOCTEXT("ChangeStageSource", "Change Stage Source"));
+		FToolMenuSection& NewSection = NewToolMenu->FindOrAddSection("ChangeMaterialStage", LOCTEXT("ChangeStageSource", "Change Stage Source"));
 		NewSection.AddDynamicEntry(NAME_None, FNewToolMenuSectionDelegate::CreateStatic(&CreateChangeMaterialStageSource));
 	}
 
@@ -670,7 +670,7 @@ void FDMMaterialStageSourceMenus::GenerateChangeSourceMenu_Slots(UToolMenu* InMe
 		}
 		else
 		{
-			FToolMenuSection& NewSection = InMenu->AddSection("ChangeSourceSlotTooltip", LOCTEXT("ChangeSourceSlot", "Change Source Slot"));
+			FToolMenuSection& NewSection = InMenu->FindOrAddSection("ChangeSourceSlotTooltip", LOCTEXT("ChangeSourceSlot", "Change Source Slot"));
 			NewSection.AddDynamicEntry(NAME_None,
 				FNewToolMenuSectionDelegate::CreateWeakLambda(
 					CurrentSlot,
@@ -817,7 +817,7 @@ void FDMMaterialStageSourceMenus::GenerateChangeSourceMenu_Advanced(UToolMenu* c
 		return;
 	}
 
-	FToolMenuSection& NewSection = InMenu->AddSection(
+	FToolMenuSection& NewSection = InMenu->FindOrAddSection(
 		NAME_None,
 		LOCTEXT("Advanced", "Advanced")
 	);
diff --git a/Engine/Plugins/Experimental/DynamicMaterial/Source/DynamicMaterialEditor/Private/UI/Menus/DMMenuContext.h b/Engine/Plugins/Experimental/DynamicMaterial/Source/DynamicMaterialEditor/Private/UI/Menus/DMMenuContext.h
deleted file mode 100644
diff --git a/Engine/Plugins/Experimental/DynamicMaterial/Source/DynamicMaterialEditor/Private/UI/PropertyGenerators/DMComponentPropertyRowGenerator.cpp b/Engine/Plugins/Experimental/DynamicMaterial/Source/DynamicMaterialEditor/Private/UI/PropertyGenerators/DMComponentPropertyRowGenerator.cpp
--- a/Engine/Plugins/Experimental/DynamicMaterial/Source/DynamicMaterialEditor/Private/UI/PropertyGenerators/DMComponentPropertyRowGenerator.cpp
+++ b/Engine/Plugins/Experimental/DynamicMaterial/Source/DynamicMaterialEditor/Private/UI/PropertyGenerators/DMComponentPropertyRowGenerator.cpp
@@ -6,7 +6,7 @@
 #include "DynamicMaterialEditorModule.h"
 #include "Model/DynamicMaterialModelDynamic.h"
 #include "Runtime/Launch/Resources/Version.h"
-#include "UI/Utils/DMWidgetStatics.h"
+#include "UI/Utils/DMWidgetLibrary.h"
 #include "UI/Widgets/Editor/SDMMaterialComponentEditor.h"
 #include "UI/Widgets/SDMMaterialEditor.h"
 
@@ -99,7 +99,7 @@ void FDMComponentPropertyRowGenerator::AddPropertyEditRows(const TSharedRef<SDMM
 		}
 	}
 
-	FDMPropertyHandle& Handle = InOutPropertyRows.Add_GetRef(FDMWidgetStatics::Get().GetPropertyHandle(&*InComponentEditorWidget, InComponent, InProperty->GetFName()));
+	FDMPropertyHandle& Handle = InOutPropertyRows.Add_GetRef(FDMWidgetLibrary::Get().GetPropertyHandle(&*InComponentEditorWidget, InComponent, InProperty->GetFName()));
 	Handle.bEnabled = !IsDynamic(InComponentEditorWidget);
 }
 
diff --git a/Engine/Plugins/Experimental/DynamicMaterial/Source/DynamicMaterialEditor/Private/UI/PropertyGenerators/DMMaterialModelPropertyRowGenerator.cpp b/Engine/Plugins/Experimental/DynamicMaterial/Source/DynamicMaterialEditor/Private/UI/PropertyGenerators/DMMaterialModelPropertyRowGenerator.cpp
--- a/Engine/Plugins/Experimental/DynamicMaterial/Source/DynamicMaterialEditor/Private/UI/PropertyGenerators/DMMaterialModelPropertyRowGenerator.cpp
+++ b/Engine/Plugins/Experimental/DynamicMaterial/Source/DynamicMaterialEditor/Private/UI/PropertyGenerators/DMMaterialModelPropertyRowGenerator.cpp
@@ -7,7 +7,7 @@
 #include "Model/DynamicMaterialModel.h"
 #include "Model/DynamicMaterialModelDynamic.h"
 #include "Model/DynamicMaterialModelEditorOnlyData.h"
-#include "UI/Utils/DMWidgetStatics.h"
+#include "UI/Utils/DMWidgetLibrary.h"
 #include "UI/Widgets/Editor/SDMMaterialGlobalSettingsEditor.h"
 #include "Utils/DMPrivate.h"
 
@@ -83,7 +83,7 @@ void FDMMaterialModelPropertyRowGenerator::AddGlobalValue(const TSharedRef<SDMMa
 		}
 	}
 
-	FDMPropertyHandle& ComponentHandle = InOutPropertyRows.Add_GetRef(FDMWidgetStatics::Get().GetPropertyHandle(
+	FDMPropertyHandle& ComponentHandle = InOutPropertyRows.Add_GetRef(FDMWidgetLibrary::Get().GetPropertyHandle(
 		&*InGlobalSettingEditorWidget,
 		InComponent,
 		UDMMaterialValue::ValueName
@@ -113,7 +113,7 @@ void FDMMaterialModelPropertyRowGenerator::AddGlobalValue(const TSharedRef<SDMMa
 void FDMMaterialModelPropertyRowGenerator::AddVariable(const TSharedRef<SDMMaterialGlobalSettingsEditor>& InGlobalSettingEditorWidget, 
 	UDynamicMaterialModelBase* InMaterialModelBase, TArray<FDMPropertyHandle>& InOutPropertyRows, UObject* InObject, FName InPropertyName)
 {
-	FDMPropertyHandle& ValueHandle = InOutPropertyRows.Add_GetRef(FDMWidgetStatics::Get().GetPropertyHandle(
+	FDMPropertyHandle& ValueHandle = InOutPropertyRows.Add_GetRef(FDMWidgetLibrary::Get().GetPropertyHandle(
 		&*InGlobalSettingEditorWidget, 
 		InObject, InPropertyName
 	));
diff --git a/Engine/Plugins/Experimental/DynamicMaterial/Source/DynamicMaterialEditor/Private/UI/PropertyGenerators/DMMaterialValueDynamicPropertyRowGenerator.cpp b/Engine/Plugins/Experimental/DynamicMaterial/Source/DynamicMaterialEditor/Private/UI/PropertyGenerators/DMMaterialValueDynamicPropertyRowGenerator.cpp
--- a/Engine/Plugins/Experimental/DynamicMaterial/Source/DynamicMaterialEditor/Private/UI/PropertyGenerators/DMMaterialValueDynamicPropertyRowGenerator.cpp
+++ b/Engine/Plugins/Experimental/DynamicMaterial/Source/DynamicMaterialEditor/Private/UI/PropertyGenerators/DMMaterialValueDynamicPropertyRowGenerator.cpp
@@ -7,7 +7,7 @@
 #include "Components/DMMaterialValueDynamic.h"
 #include "DMEDefs.h"
 #include "IDetailPropertyRow.h"
-#include "UI/Utils/DMWidgetStatics.h"
+#include "UI/Utils/DMWidgetLibrary.h"
 #include "UI/Widgets/Editor/SDMMaterialComponentEditor.h"
 
 const TSharedRef<FDMMaterialValueDynamicPropertyRowGenerator>& FDMMaterialValueDynamicPropertyRowGenerator::Get()
@@ -53,7 +53,7 @@ void FDMMaterialValueDynamicPropertyRowGenerator::AddComponentProperties(const T
 
 	if (ParentValue->AllowEditValue())
 	{
-		FDMPropertyHandle Handle = FDMWidgetStatics::Get().GetPropertyHandle(&*InComponentEditorWidget, ValueDynamic, UDMMaterialValue::ValueName);
+		FDMPropertyHandle Handle = FDMWidgetLibrary::Get().GetPropertyHandle(&*InComponentEditorWidget, ValueDynamic, UDMMaterialValue::ValueName);
 
 		Handle.ResetToDefaultOverride = FResetToDefaultOverride::Create(
 			FIsResetToDefaultVisible::CreateUObject(ValueDynamic, &UDMMaterialValueDynamic::CanResetToDefault),
diff --git a/Engine/Plugins/Experimental/DynamicMaterial/Source/DynamicMaterialEditor/Private/UI/PropertyGenerators/DMMaterialValuePropertyRowGenerator.cpp b/Engine/Plugins/Experimental/DynamicMaterial/Source/DynamicMaterialEditor/Private/UI/PropertyGenerators/DMMaterialValuePropertyRowGenerator.cpp
--- a/Engine/Plugins/Experimental/DynamicMaterial/Source/DynamicMaterialEditor/Private/UI/PropertyGenerators/DMMaterialValuePropertyRowGenerator.cpp
+++ b/Engine/Plugins/Experimental/DynamicMaterial/Source/DynamicMaterialEditor/Private/UI/PropertyGenerators/DMMaterialValuePropertyRowGenerator.cpp
@@ -9,7 +9,7 @@
 #include "IDetailPropertyRow.h"
 #include "Model/DynamicMaterialModelBase.h"
 #include "Model/DynamicMaterialModelDynamic.h"
-#include "UI/Utils/DMWidgetStatics.h"
+#include "UI/Utils/DMWidgetLibrary.h"
 #include "UI/Widgets/Editor/SDMMaterialComponentEditor.h"
 #include "UI/Widgets/SDMMaterialEditor.h"
 
@@ -65,7 +65,7 @@ void FDMMaterialValuePropertyRowGenerator::AddComponentProperties(const TSharedR
 
 	if (Value->AllowEditValue())
 	{
-		FDMPropertyHandle Handle = FDMWidgetStatics::Get().GetPropertyHandle(&*InComponentEditorWidget, Value, UDMMaterialValue::ValueName);
+		FDMPropertyHandle Handle = FDMWidgetLibrary::Get().GetPropertyHandle(&*InComponentEditorWidget, Value, UDMMaterialValue::ValueName);
 
 		Handle.ResetToDefaultOverride = FResetToDefaultOverride::Create(
 			FIsResetToDefaultVisible::CreateUObject(Value, &UDMMaterialValue::CanResetToDefault),
diff --git a/Engine/Plugins/Experimental/DynamicMaterial/Source/DynamicMaterialEditor/Private/UI/PropertyGenerators/DMTextureUVDynamicPropertyRowGenerator.cpp b/Engine/Plugins/Experimental/DynamicMaterial/Source/DynamicMaterialEditor/Private/UI/PropertyGenerators/DMTextureUVDynamicPropertyRowGenerator.cpp
--- a/Engine/Plugins/Experimental/DynamicMaterial/Source/DynamicMaterialEditor/Private/UI/PropertyGenerators/DMTextureUVDynamicPropertyRowGenerator.cpp
+++ b/Engine/Plugins/Experimental/DynamicMaterial/Source/DynamicMaterialEditor/Private/UI/PropertyGenerators/DMTextureUVDynamicPropertyRowGenerator.cpp
@@ -10,7 +10,7 @@
 #include "Modules/ModuleManager.h"
 #include "PropertyCustomizationHelpers.h"
 #include "PropertyEditorModule.h"
-#include "UI/Utils/DMWidgetStatics.h"
+#include "UI/Utils/DMWidgetLibrary.h"
 #include "UI/Widgets/Editor/SDMMaterialComponentEditor.h"
 #include "UI/Widgets/SDMMaterialEditor.h"
 #include "UI/Widgets/Visualizers/SDMTextureUVVisualizerProperty.h"
@@ -26,10 +26,10 @@ const TSharedRef<FDMTextureUVDynamicPropertyRowGenerator>& FDMTextureUVDynamicPr
 
 namespace UE::DynamicMaterialEditor::Private
 {
-	void AddTextureUVDynamicPropertyRow(const TSharedRef<SWidget>& InComponentEditorWidget, UDMMaterialComponent* InTextureUVDynamic, 
+	void AddTextureUVDynamicPropertyRow(const TSharedRef<SWidget>& InParentWidget, UDMMaterialComponent* InTextureUVDynamic,
 		FName InProperty, bool bInEnabled, TArray<FDMPropertyHandle>& InOutPropertyRows);
 
-	void AddTextureUVDynamicVisualizerRow(const TSharedRef<SWidget>& InComponentEditorWidget, UDMTextureUVDynamic* InTextureUVDynamic, 
+	void AddTextureUVDynamicVisualizerRow(const TSharedRef<SWidget>& InParentWidget, UDMTextureUVDynamic* InTextureUVDynamic,
 		TArray<FDMPropertyHandle>& InOutPropertyRows);
 
 	bool CanResetTextureUVDynamicDynamicPropertyToDefault(TSharedPtr<IPropertyHandle> InPropertyHandle);
@@ -124,10 +124,10 @@ bool FDMTextureUVDynamicPropertyRowGenerator::AllowKeyframeButton(UDMMaterialCom
 	return FDMComponentPropertyRowGenerator::AllowKeyframeButton(InComponent, InProperty);
 }
 
-void UE::DynamicMaterialEditor::Private::AddTextureUVDynamicPropertyRow(const TSharedRef<SWidget>& InComponentEditorWidget, 
+void UE::DynamicMaterialEditor::Private::AddTextureUVDynamicPropertyRow(const TSharedRef<SWidget>& InParentWidget, 
 	UDMMaterialComponent* InTextureUVDynamic, FName InProperty, bool bInEnabled, TArray<FDMPropertyHandle>& InOutPropertyRows)
 {
-	FDMPropertyHandle& NewHandle = InOutPropertyRows.Add_GetRef(FDMWidgetStatics::Get().GetPropertyHandle(&*InComponentEditorWidget, InTextureUVDynamic, InProperty));
+	FDMPropertyHandle& NewHandle = InOutPropertyRows.Add_GetRef(FDMWidgetLibrary::Get().GetPropertyHandle(&*InParentWidget, InTextureUVDynamic, InProperty));
 
 	NewHandle.ResetToDefaultOverride = FResetToDefaultOverride::Create(
 		FIsResetToDefaultVisible::CreateStatic(&UE::DynamicMaterialEditor::Private::CanResetTextureUVDynamicDynamicPropertyToDefault),
@@ -137,7 +137,7 @@ void UE::DynamicMaterialEditor::Private::AddTextureUVDynamicPropertyRow(const TS
 	NewHandle.bEnabled = bInEnabled;
 }
 
-void UE::DynamicMaterialEditor::Private::AddTextureUVDynamicVisualizerRow(const TSharedRef<SWidget>& InComponentEditorWidget, UDMTextureUVDynamic* InTextureUVDynamic,
+void UE::DynamicMaterialEditor::Private::AddTextureUVDynamicVisualizerRow(const TSharedRef<SWidget>& InParentWidget, UDMTextureUVDynamic* InTextureUVDynamic,
 	TArray<FDMPropertyHandle>& InOutPropertyRows)
 {
 	UDMTextureUV* TextureUV = InTextureUVDynamic->GetParentTextureUV();
@@ -155,12 +155,12 @@ void UE::DynamicMaterialEditor::Private::AddTextureUVDynamicVisualizerRow(const
 		return;
 	}
 
-	if (InComponentEditorWidget->GetWidgetClass().GetWidgetType() != SDMMaterialComponentEditor::StaticWidgetClass().GetWidgetType())
+	if (InParentWidget->GetWidgetClass().GetWidgetType() != SDMMaterialComponentEditor::StaticWidgetClass().GetWidgetType())
 	{
 		return;
 	}
 
-	TSharedPtr<SDMMaterialEditor> EditorWidget = StaticCastSharedRef<SDMMaterialComponentEditor>(InComponentEditorWidget)->GetEditorWidget();
+	TSharedPtr<SDMMaterialEditor> EditorWidget = StaticCastSharedRef<SDMMaterialComponentEditor>(InParentWidget)->GetEditorWidget();
 
 	if (!EditorWidget.IsValid())
 	{
diff --git a/Engine/Plugins/Experimental/DynamicMaterial/Source/DynamicMaterialEditor/Private/UI/PropertyGenerators/DMTextureUVPropertyRowGenerator.cpp b/Engine/Plugins/Experimental/DynamicMaterial/Source/DynamicMaterialEditor/Private/UI/PropertyGenerators/DMTextureUVPropertyRowGenerator.cpp
--- a/Engine/Plugins/Experimental/DynamicMaterial/Source/DynamicMaterialEditor/Private/UI/PropertyGenerators/DMTextureUVPropertyRowGenerator.cpp
+++ b/Engine/Plugins/Experimental/DynamicMaterial/Source/DynamicMaterialEditor/Private/UI/PropertyGenerators/DMTextureUVPropertyRowGenerator.cpp
@@ -11,7 +11,7 @@
 #include "Modules/ModuleManager.h"
 #include "PropertyCustomizationHelpers.h"
 #include "PropertyEditorModule.h"
-#include "UI/Utils/DMWidgetStatics.h"
+#include "UI/Utils/DMWidgetLibrary.h"
 #include "UI/Widgets/Editor/SDMMaterialComponentEditor.h"
 #include "UI/Widgets/SDMMaterialEditor.h"
 #include "UI/Widgets/Visualizers/SDMTextureUVVisualizerProperty.h"
@@ -27,10 +27,10 @@ const TSharedRef<FDMTextureUVPropertyRowGenerator>& FDMTextureUVPropertyRowGener
 
 namespace UE::DynamicMaterialEditor::Private
 {
-	void AddTextureUVPropertyRow(const TSharedRef<SWidget>& InComponentEditorWidget, UDMTextureUV* InTextureUV, FName InProperty, 
+	void AddTextureUVPropertyRow(const TSharedRef<SWidget>& InParentWidget, UDMTextureUV* InTextureUV, FName InProperty,
 		TArray<FDMPropertyHandle>& InOutPropertyRows);
 
-	void AddTextureUVVisualizerRow(const TSharedRef<SWidget>& InComponentEditorWidget, UDMTextureUV* InTextureUV, 
+	void AddTextureUVVisualizerRow(const TSharedRef<SWidget>& InParentWidget, UDMTextureUV* InTextureUV,
 		TArray<FDMPropertyHandle>& InOutPropertyRows);
 
 	bool CanResetTextureUVPropertyToDefault(TSharedPtr<IPropertyHandle> InPropertyHandle);
@@ -127,10 +127,10 @@ bool FDMTextureUVPropertyRowGenerator::AllowKeyframeButton(UDMMaterialComponent*
 	return FDMComponentPropertyRowGenerator::AllowKeyframeButton(InComponent, InProperty);
 }
 
-void UE::DynamicMaterialEditor::Private::AddTextureUVPropertyRow(const TSharedRef<SWidget>& InComponentEditorWidget, UDMTextureUV* InTextureUV, 
+void UE::DynamicMaterialEditor::Private::AddTextureUVPropertyRow(const TSharedRef<SWidget>& InParentWidget, UDMTextureUV* InTextureUV,
 	FName InProperty, TArray<FDMPropertyHandle>& InOutPropertyRows)
 {
-	FDMPropertyHandle& NewHandle = InOutPropertyRows.Add_GetRef(FDMWidgetStatics::Get().GetPropertyHandle(&*InComponentEditorWidget, InTextureUV, InProperty));
+	FDMPropertyHandle& NewHandle = InOutPropertyRows.Add_GetRef(FDMWidgetLibrary::Get().GetPropertyHandle(&*InParentWidget, InTextureUV, InProperty));
 
 	NewHandle.ResetToDefaultOverride = FResetToDefaultOverride::Create(
 		FIsResetToDefaultVisible::CreateStatic(&UE::DynamicMaterialEditor::Private::CanResetTextureUVPropertyToDefault),
@@ -140,7 +140,7 @@ void UE::DynamicMaterialEditor::Private::AddTextureUVPropertyRow(const TSharedRe
 	NewHandle.bEnabled = true;
 }
 
-void UE::DynamicMaterialEditor::Private::AddTextureUVVisualizerRow(const TSharedRef<SWidget>& InComponentEditorWidget, UDMTextureUV* InTextureUV, 
+void UE::DynamicMaterialEditor::Private::AddTextureUVVisualizerRow(const TSharedRef<SWidget>& InParentWidget, UDMTextureUV* InTextureUV,
 	TArray<FDMPropertyHandle>& InOutPropertyRows)
 {
 	// Make sure we don't get a substage
@@ -151,12 +151,12 @@ void UE::DynamicMaterialEditor::Private::AddTextureUVVisualizerRow(const TShared
 		return;
 	}
 
-	if (InComponentEditorWidget->GetWidgetClass().GetWidgetType() != SDMMaterialComponentEditor::StaticWidgetClass().GetWidgetType())
+	if (InParentWidget->GetWidgetClass().GetWidgetType() != SDMMaterialComponentEditor::StaticWidgetClass().GetWidgetType())
 	{
 		return;
 	}
 
-	TSharedPtr<SDMMaterialEditor> EditorWidget = StaticCastSharedRef<SDMMaterialComponentEditor>(InComponentEditorWidget)->GetEditorWidget();
+	TSharedPtr<SDMMaterialEditor> EditorWidget = StaticCastSharedRef<SDMMaterialComponentEditor>(InParentWidget)->GetEditorWidget();
 
 	if (!EditorWidget.IsValid())
 	{
diff --git a/Engine/Plugins/Experimental/DynamicMaterial/Source/DynamicMaterialEditor/Private/UI/Utils/DMWidgetLibrary.cpp b/Engine/Plugins/Experimental/DynamicMaterial/Source/DynamicMaterialEditor/Private/UI/Utils/DMWidgetLibrary.cpp
new file mode 100644
--- /dev/null
+++ b/Engine/Plugins/Experimental/DynamicMaterial/Source/DynamicMaterialEditor/Private/UI/Utils/DMWidgetLibrary.cpp
@@ -0,0 +1,345 @@
+// Copyright Epic Games, Inc. All Rights Reserved.
+
+#include "UI/Utils/DMWidgetLibrary.h"
+
+#include "Components/MaterialValues/DMMaterialValueFloat.h"
+#include "IDetailTreeNode.h"
+#include "IPropertyRowGenerator.h"
+#include "Modules/ModuleManager.h"
+#include "PropertyEditorModule.h"
+#include "UI/Widgets/Editor/SDMMaterialComponentEditor.h"
+
+IDMWidgetLibrary& IDMWidgetLibrary::Get()
+{
+	return FDMWidgetLibrary::Get();
+}
+
+const FLazyName FDMWidgetLibrary::PropertyValueWidget = TEXT("SPropertyValueWidget");
+
+FDMWidgetLibrary& FDMWidgetLibrary::Get()
+{
+	 static FDMWidgetLibrary Instance;
+	 return Instance;
+}
+
+bool FDMWidgetLibrary::GetExpansionState(UObject* InOwner, FName InName, bool& bOutExpanded)
+{
+	const FExpansionItem ExpansionItem = {InOwner, InName};
+
+	if (const bool* State = ExpansionStates.Find(ExpansionItem))
+	{
+		bOutExpanded = *State;
+		return true;
+	}
+
+	return false;
+}
+
+void FDMWidgetLibrary::SetExpansionState(UObject* InOwner, FName InName, bool bInIsExpanded)
+{
+	const FExpansionItem ExpansionItem = {InOwner, InName};
+	ExpansionStates.FindOrAdd(ExpansionItem) = bInIsExpanded;
+}
+
+FDMPropertyHandle FDMWidgetLibrary::GetPropertyHandle(const SWidget* InOwningWidget, UObject* InObject, FName InPropertyName)
+{
+	TArray<FDMPropertyHandle>& PropertyHandles = PropertyHandleMap.FindOrAdd(InOwningWidget);
+
+	for (const FDMPropertyHandle& ExistingHandle : PropertyHandles)
+	{
+		if (ExistingHandle.PropertyHandle && ExistingHandle.PropertyHandle->GetProperty()->GetFName() == InPropertyName)
+		{
+			TArray<UObject*> Outers;
+			ExistingHandle.PropertyHandle->GetOuterObjects(Outers);
+
+			if (Outers.IsEmpty() == false && Outers[0] == InObject)
+			{
+				return ExistingHandle;
+			}
+		}
+	}
+
+	if (TSharedPtr<IPropertyRowGenerator> PropertyRowGenerator = SearchForGenerator(PropertyHandles, InObject))
+	{
+		FDMPropertyHandle PropertyHandle;
+		PropertyHandle.PropertyRowGenerator = PropertyRowGenerator;
+
+		if (TSharedPtr<IDetailTreeNode> DetailTreeNode = SearchGeneratorForNode(PropertyRowGenerator.ToSharedRef(), InPropertyName))
+		{
+			PropertyHandle.DetailTreeNode = DetailTreeNode;
+			PropertyHandle.PropertyHandle = DetailTreeNode->CreatePropertyHandle();
+
+			AddPropertyMetaData(InObject, InPropertyName, PropertyHandle);
+
+			return PropertyHandle;
+		}
+
+		return PropertyHandle;
+	}
+
+	FDMPropertyHandle NewHandle = CreatePropertyHandle(InOwningWidget, InObject, InPropertyName);
+
+	if (!NewHandle.PropertyHandle.IsValid() && NewHandle.DetailTreeNode.IsValid())
+	{
+		NewHandle.PropertyHandle = NewHandle.DetailTreeNode->CreatePropertyHandle();
+	}
+
+	AddPropertyMetaData(InObject, InPropertyName, NewHandle);
+
+	PropertyHandles.Add(NewHandle);
+
+	return NewHandle;
+}
+
+FDMPropertyHandle FDMWidgetLibrary::GetPropertyHandle(const SDMMaterialComponentEditor* InOwningWidget, UObject* InObject, FName InPropertyName)
+{
+	return GetPropertyHandle(static_cast<const SWidget*>(InOwningWidget), InObject, InPropertyName);
+}
+
+void FDMWidgetLibrary::ClearPropertyHandles(const SWidget* InOwningWidget)
+{
+	PropertyHandleMap.Remove(InOwningWidget);
+}
+
+TSharedPtr<SWidget> FDMWidgetLibrary::FindWidgetInHierarchy(const TSharedRef<SWidget>& InParent, const FName& InName)
+{
+	if (InParent->GetType() == InName)
+	{
+		return InParent;
+	}
+
+	FChildren* Children = InParent->GetChildren();
+
+	if (!Children)
+	{
+		return nullptr;
+	}
+
+	const int32 ChildNum = Children->Num();
+
+	for (int32 Index = 0; Index < ChildNum; ++Index)
+	{
+		const TSharedRef<SWidget>& Widget = Children->GetChildAt(Index);
+
+		if (Widget->GetType() == InName)
+		{
+			return Widget;
+		}
+	}
+
+	for (int32 Index = 0; Index < ChildNum; ++Index)
+	{
+		if (TSharedPtr<SWidget> FoundChild = FindWidgetInHierarchy(Children->GetChildAt(Index), InName))
+		{
+			return FoundChild;
+		}
+	}
+
+	return nullptr;
+}
+
+TSharedPtr<SWidget> FDMWidgetLibrary::GetInnerPropertyValueWidget(const TSharedRef<SWidget>& InWidget)
+{
+	if (FChildren* Children = InWidget->GetChildren())
+	{
+		if (Children->Num() > 0)
+		{
+			return Children->GetChildAt(0);
+		}
+	}
+
+	return nullptr;
+}
+
+void FDMWidgetLibrary::ClearData()
+{
+	ExpansionStates.Empty();
+	PropertyHandleMap.Empty();
+}
+
+FDMPropertyHandle FDMWidgetLibrary::CreatePropertyHandle(const void* InOwningWidget, UObject* InObject, FName InPropertyName)
+{
+	FDMPropertyHandle PropertyHandle;
+
+	FPropertyEditorModule& PropertyEditor = FModuleManager::Get().LoadModuleChecked<FPropertyEditorModule>(TEXT("PropertyEditor"));
+
+	FPropertyRowGeneratorArgs RowGeneratorArgs;
+
+	if (UDMMaterialComponent* Component = Cast<UDMMaterialComponent>(InObject))
+	{
+		RowGeneratorArgs.NotifyHook = Component;
+	}
+
+	PropertyHandle.PropertyRowGenerator = PropertyEditor.CreatePropertyRowGenerator(RowGeneratorArgs);
+	PropertyHandle.PropertyRowGenerator->SetObjects({InObject});
+
+	if (const TSharedPtr<IDetailTreeNode> FoundTreeNode = SearchGeneratorForNode(
+		PropertyHandle.PropertyRowGenerator.ToSharedRef(), InPropertyName))
+	{
+		PropertyHandle.DetailTreeNode = FoundTreeNode;
+		PropertyHandle.PropertyHandle = FoundTreeNode->CreatePropertyHandle();
+		return PropertyHandle;
+	}
+
+	return PropertyHandle;
+}
+
+TSharedPtr<IDetailTreeNode> FDMWidgetLibrary::SearchNodesForProperty(const TArray<TSharedRef<IDetailTreeNode>>& InNodes, FName InPropertyName)
+{
+	for (const TSharedRef<IDetailTreeNode>& ChildNode : InNodes)
+	{
+		switch (ChildNode->GetNodeType())
+		{
+			case EDetailNodeType::Category:
+			{
+				TArray<TSharedRef<IDetailTreeNode>> CategoryChildNodes;
+				ChildNode->GetChildren(CategoryChildNodes);
+
+				if (TSharedPtr<IDetailTreeNode> FoundNode = SearchNodesForProperty(CategoryChildNodes, InPropertyName))
+				{
+					return FoundNode;
+				}
+
+				break;
+			}
+
+			case EDetailNodeType::Item:
+				if (ChildNode->GetNodeName() == InPropertyName)
+				{
+					return ChildNode;
+				}
+				break;
+
+			default:
+				// Do nothing
+				break;
+		}
+	}
+
+	return nullptr;
+}
+
+TSharedPtr<IDetailTreeNode> FDMWidgetLibrary::SearchGeneratorForNode(const TSharedRef<IPropertyRowGenerator>& InGenerator, FName InPropertyName)
+{
+	return SearchNodesForProperty(InGenerator->GetRootTreeNodes(), InPropertyName);
+}
+
+TSharedPtr<IPropertyRowGenerator> FDMWidgetLibrary::SearchForGenerator(const TArray<FDMPropertyHandle>& InPropertyHandles, UObject* InObject)
+{
+	if (!InObject)
+	{
+		return nullptr;
+	}
+
+	for (const FDMPropertyHandle& PropertyHandle : InPropertyHandles)
+	{
+		if (PropertyHandle.PropertyRowGenerator.IsValid())
+		{
+			for (const TWeakObjectPtr<UObject>& WeakObject : PropertyHandle.PropertyRowGenerator->GetSelectedObjects())
+			{
+				if (WeakObject.Get() == InObject)
+				{
+					return PropertyHandle.PropertyRowGenerator;
+				}
+			}
+		}
+	}
+
+	return nullptr;
+}
+
+void FDMWidgetLibrary::AddPropertyMetaData(UObject* InObject, FName InPropertyName, FDMPropertyHandle& InPropertyHandle)
+{
+	FProperty* Property = nullptr;
+
+	if (InPropertyHandle.PropertyHandle.IsValid())
+	{
+		TSharedRef<IPropertyHandle> PropertyHandleRef = InPropertyHandle.PropertyHandle.ToSharedRef();
+
+		InPropertyHandle.Priority = GetPriority(PropertyHandleRef);
+		InPropertyHandle.bKeyframeable = IsKeyframeable(PropertyHandleRef);
+
+		Property = InPropertyHandle.PropertyHandle->GetProperty();
+
+		if (UDMMaterialValueFloat* FloatValue = Cast<UDMMaterialValueFloat>(InObject))
+		{
+			if (FloatValue->HasValueRange())
+			{
+				const FName UIMin = FName("UIMin");
+				const FName UIMax = FName("UIMax");
+				const FName ClampMin = FName("ClampMin");
+				const FName ClampMax = FName("ClampMax");
+
+				InPropertyHandle.PropertyHandle->SetInstanceMetaData(UIMin, FString::SanitizeFloat(FloatValue->GetValueRange().Min));
+				InPropertyHandle.PropertyHandle->SetInstanceMetaData(ClampMin, FString::SanitizeFloat(FloatValue->GetValueRange().Min));
+				InPropertyHandle.PropertyHandle->SetInstanceMetaData(UIMax, FString::SanitizeFloat(FloatValue->GetValueRange().Max));
+				InPropertyHandle.PropertyHandle->SetInstanceMetaData(ClampMax, FString::SanitizeFloat(FloatValue->GetValueRange().Max));
+			}
+		}
+	}
+	else if (IsValid(InObject))
+	{
+		Property = InObject->GetClass()->FindPropertyByName(InPropertyName);
+	}
+
+	if (FStructProperty* StructProperty = CastField<FStructProperty>(Property))
+	{
+		uint8 ComponentCount = 1;
+
+		if (StructProperty->Struct == TBaseStructure<FVector2D>::Get()
+			|| StructProperty->Struct == TVariantStructure<FVector2f>::Get())
+		{
+			ComponentCount = 2;
+		}
+
+		if (StructProperty->Struct == TBaseStructure<FVector>::Get()
+			|| StructProperty->Struct == TVariantStructure<FVector3f>::Get()
+			|| StructProperty->Struct == TBaseStructure<FRotator>::Get())
+		{
+			ComponentCount = 3;
+		}
+
+		// FLinearColor doesn't need the extra space
+		if (StructProperty->Struct == TBaseStructure<FVector4>::Get()
+			|| StructProperty->Struct == TVariantStructure<FVector4f>::Get())
+		{
+			ComponentCount = 4;
+		}
+
+		switch (ComponentCount)
+		{
+			case 0:
+			case 1:
+				break;
+
+			case 2:
+				InPropertyHandle.MaxWidth = 200.f;
+				break;
+
+				// 3 and above
+			default:
+				InPropertyHandle.MaxWidth = 275.f;
+				break;
+		}
+	}
+}
+
+EDMPropertyHandlePriority FDMWidgetLibrary::GetPriority(const TSharedRef<IPropertyHandle>& InPropertyHandle)
+{
+	if (InPropertyHandle->HasMetaData("HighPriority"))
+	{
+		return EDMPropertyHandlePriority::High;
+	}
+
+	if (InPropertyHandle->HasMetaData("LowPriority"))
+	{
+		return EDMPropertyHandlePriority::Low;
+	}
+
+	return EDMPropertyHandlePriority::Normal;
+}
+
+bool FDMWidgetLibrary::IsKeyframeable(const TSharedRef<IPropertyHandle>& InPropertyHandle)
+{
+	return !InPropertyHandle->HasMetaData("NotKeyframeable");
+}
diff --git a/Engine/Plugins/Experimental/DynamicMaterial/Source/DynamicMaterialEditor/Private/UI/Utils/DMWidgetLibrary.h b/Engine/Plugins/Experimental/DynamicMaterial/Source/DynamicMaterialEditor/Private/UI/Utils/DMWidgetLibrary.h
new file mode 100644
--- /dev/null
+++ b/Engine/Plugins/Experimental/DynamicMaterial/Source/DynamicMaterialEditor/Private/UI/Utils/DMWidgetLibrary.h
@@ -0,0 +1,82 @@
+// Copyright Epic Games, Inc. All Rights Reserved.
+
+#pragma once
+
+#include "Containers/Array.h"
+#include "Containers/Map.h"
+#include "DMEDefs.h"
+#include "Templates/SharedPointerFwd.h"
+#include "UI/Utils/IDMWidgetLibrary.h"
+#include "UObject/ObjectKey.h"
+
+class FName;
+class SWidget;
+class UObject;
+
+class FDMWidgetLibrary : public IDMWidgetLibrary
+{
+public:
+	static const FLazyName PropertyValueWidget;
+
+	static FDMWidgetLibrary& Get();
+
+	virtual ~FDMWidgetLibrary() = default;
+
+	bool GetExpansionState(UObject* InOwner, FName InName, bool& bOutExpanded);
+
+	void SetExpansionState(UObject* InOwner, FName InName, bool bInIsExpanded);
+
+	void ClearPropertyHandles(const SWidget* InOwningWidget);
+
+	TSharedPtr<SWidget> FindWidgetInHierarchy(const TSharedRef<SWidget>& InParent, const FName& InName);
+
+	TSharedPtr<SWidget> GetInnerPropertyValueWidget(const TSharedRef<SWidget>& InWidget);
+
+	void ClearData();
+
+	//~ Begin IDMWidgetLibrary
+	virtual FDMPropertyHandle GetPropertyHandle(const SWidget* InOwningWidget, UObject* InObject, FName InPropertyName) override;
+	virtual FDMPropertyHandle GetPropertyHandle(const SDMMaterialComponentEditor* InOwningWidget, UObject* InObject, FName InPropertyName) override;
+	//~ End IDMWidgetLibrary
+
+private:
+	struct FExpansionItem
+	{
+		TObjectKey<UObject> Owner;
+		FName Name;
+
+		bool operator==(const FExpansionItem& InOther) const
+		{
+			return Owner == InOther.Owner
+				&& Name == InOther.Name;
+		}
+
+		friend uint32 GetTypeHash(const FExpansionItem& InItem)
+		{
+			return HashCombineFast(
+				GetTypeHash(InItem.Owner),
+				GetTypeHash(InItem.Name)
+			);
+		}
+	};
+
+	TMap<FExpansionItem, bool> ExpansionStates;
+
+	TMap<const SWidget*, TArray<FDMPropertyHandle>> PropertyHandleMap;
+
+	FDMPropertyHandle CreatePropertyHandle(const void* InOwningWidget, UObject* InObject, FName InPropertyName);
+
+	static TSharedPtr<IDetailTreeNode> SearchNodesForProperty(const TArray<TSharedRef<IDetailTreeNode>>& InNodes, FName InPropertyName);
+
+	static TSharedPtr<IDetailTreeNode> SearchGeneratorForNode(const TSharedRef<IPropertyRowGenerator>& InGenerator, FName InPropertyName);
+
+	static TSharedPtr<IPropertyRowGenerator> SearchForGenerator(const TArray<FDMPropertyHandle>& InPropertyHandles, UObject* InObject);
+
+	static void AddPropertyMetaData(UObject* InObject, FName InPropertyName, FDMPropertyHandle& InPropertyHandle);
+
+	/** Checks the meta data of a property meta data to check for high and low priority specifiers. */
+	static EDMPropertyHandlePriority GetPriority(const TSharedRef<IPropertyHandle>& InPropertyHandle);
+
+	/** Checks for the NotKeyframeable meta data. */
+	static bool IsKeyframeable(const TSharedRef<IPropertyHandle>& InPropertyHandle);
+};
diff --git a/Engine/Plugins/Experimental/DynamicMaterial/Source/DynamicMaterialEditor/Private/UI/Utils/DMWidgetStatics.cpp b/Engine/Plugins/Experimental/DynamicMaterial/Source/DynamicMaterialEditor/Private/UI/Utils/DMWidgetStatics.cpp
deleted file mode 100644
diff --git a/Engine/Plugins/Experimental/DynamicMaterial/Source/DynamicMaterialEditor/Private/UI/Utils/DMWidgetStatics.h b/Engine/Plugins/Experimental/DynamicMaterial/Source/DynamicMaterialEditor/Private/UI/Utils/DMWidgetStatics.h
deleted file mode 100644
diff --git a/Engine/Plugins/Experimental/DynamicMaterial/Source/DynamicMaterialEditor/Private/UI/Widgets/Editor/SDMMaterialComponentEditor.cpp b/Engine/Plugins/Experimental/DynamicMaterial/Source/DynamicMaterialEditor/Private/UI/Widgets/Editor/SDMMaterialComponentEditor.cpp
--- a/Engine/Plugins/Experimental/DynamicMaterial/Source/DynamicMaterialEditor/Private/UI/Widgets/Editor/SDMMaterialComponentEditor.cpp
+++ b/Engine/Plugins/Experimental/DynamicMaterial/Source/DynamicMaterialEditor/Private/UI/Widgets/Editor/SDMMaterialComponentEditor.cpp
@@ -10,7 +10,7 @@
 #include "ICustomDetailsView.h"
 #include "Items/ICustomDetailsViewCustomCategoryItem.h"
 #include "Items/ICustomDetailsViewItem.h"
-#include "UI/Utils/DMWidgetStatics.h"
+#include "UI/Utils/DMWidgetLibrary.h"
 #include "UI/Widgets/SDMMaterialEditor.h"
 #include "Widgets/Input/SComboButton.h"
 
@@ -57,7 +57,7 @@ UDMMaterialComponent* SDMMaterialComponentEditor::GetComponent() const
 
 void SDMMaterialComponentEditor::OnComponentUpdated(UDMMaterialComponent* InComponent, UDMMaterialComponent* InSource, EDMUpdateType InUpdateType)
 {
-	if (EnumHasAnyFlags(InUpdateType, EDMUpdateType::Structure))
+	if (EnumHasAnyFlags(InUpdateType, EDMUpdateType::Structure | EDMUpdateType::RefreshDetailView))
 	{
 		if (TSharedPtr<SDMMaterialEditor> EditorWidget = GetEditorWidget())
 		{
@@ -86,7 +86,7 @@ TSharedRef<ICustomDetailsViewItem> SDMMaterialComponentEditor::GetDefaultCategor
 		InDetailsView->ExtendTree(InRootId, ECustomDetailsTreeInsertPosition::Child, DefaultCategoryItem.ToSharedRef());
 
 		bool bExpansionState = true;
-		FDMWidgetStatics::Get().GetExpansionState(ObjectWeak.Get(), DefaultCategoryName, bExpansionState);
+		FDMWidgetLibrary::Get().GetExpansionState(ObjectWeak.Get(), DefaultCategoryName, bExpansionState);
 
 		InDetailsView->SetItemExpansionState(
 			DefaultCategoryItem->GetItemId(),
diff --git a/Engine/Plugins/Experimental/DynamicMaterial/Source/DynamicMaterialEditor/Private/UI/Widgets/Editor/SDMMaterialProperties.cpp b/Engine/Plugins/Experimental/DynamicMaterial/Source/DynamicMaterialEditor/Private/UI/Widgets/Editor/SDMMaterialProperties.cpp
--- a/Engine/Plugins/Experimental/DynamicMaterial/Source/DynamicMaterialEditor/Private/UI/Widgets/Editor/SDMMaterialProperties.cpp
+++ b/Engine/Plugins/Experimental/DynamicMaterial/Source/DynamicMaterialEditor/Private/UI/Widgets/Editor/SDMMaterialProperties.cpp
@@ -23,7 +23,7 @@
 #include "Model/DynamicMaterialModelEditorOnlyData.h"
 #include "SAssetDropTarget.h"
 #include "Styling/StyleColors.h"
-#include "UI/Utils/DMWidgetStatics.h"
+#include "UI/Utils/DMWidgetLibrary.h"
 #include "UI/Widgets/Editor/SDMMaterialPropertySelector.h"
 #include "UI/Widgets/SDMMaterialEditor.h"
 #include "UI/Widgets/Visualizers/SDMMaterialComponentPreview.h"
@@ -126,7 +126,7 @@ TSharedRef<SWidget> SDMMaterialProperties::CreateSlot_Content()
 			DetailsView->ExtendTree(RootItem->GetItemId(), ECustomDetailsTreeInsertPosition::Child, CategoryItem);
 
 			bool bExpansionState = true;
-			FDMWidgetStatics::Get().GetExpansionState(MaterialModelBase, InName, bExpansionState);
+			FDMWidgetLibrary::Get().GetExpansionState(MaterialModelBase, InName, bExpansionState);
 
 			DetailsView->SetItemExpansionState(
 				CategoryItem->GetItemId(),
@@ -625,7 +625,7 @@ TSharedRef<SWidget> SDMMaterialProperties::CreateGlobalSlider(UDMMaterialPropert
 	TSharedRef<ICustomDetailsView> DetailsView = ICustomDetailsViewModule::Get().CreateCustomDetailsView(Args);
 	FCustomDetailsViewItemId RootId = DetailsView->GetRootItem()->GetItemId();
 
-	FDMPropertyHandle PropertyHandle = FDMWidgetStatics::Get().GetPropertyHandle(this, AlphaObject, UDMMaterialValue::ValueName);
+	FDMPropertyHandle PropertyHandle = FDMWidgetLibrary::Get().GetPropertyHandle(this, AlphaObject, UDMMaterialValue::ValueName);
 
 	TSharedRef<ICustomDetailsViewItem> AlphaItem = DetailsView->CreateDetailTreeItem(DetailsView->GetRootItem(), PropertyHandle.DetailTreeNode.ToSharedRef());
 
@@ -653,9 +653,9 @@ TSharedRef<SWidget> SDMMaterialProperties::CreateGlobalSlider(UDMMaterialPropert
 
 	if (ValueWidget.IsValid())
 	{
-		if (TSharedPtr<SWidget> FoundPropertyValueWidget = FDMWidgetStatics::Get().FindWidgetInHierarchy(ValueWidget.ToSharedRef(), FDMWidgetStatics::PropertyValueWidget))
+		if (TSharedPtr<SWidget> FoundPropertyValueWidget = FDMWidgetLibrary::Get().FindWidgetInHierarchy(ValueWidget.ToSharedRef(), FDMWidgetLibrary::PropertyValueWidget))
 		{
-			if (TSharedPtr<SWidget> InnerPropertyWidget = FDMWidgetStatics::Get().GetInnerPropertyValueWidget(FoundPropertyValueWidget.ToSharedRef()))
+			if (TSharedPtr<SWidget> InnerPropertyWidget = FDMWidgetLibrary::Get().GetInnerPropertyValueWidget(FoundPropertyValueWidget.ToSharedRef()))
 			{
 				ValueWidget = InnerPropertyWidget;
 			}
@@ -725,7 +725,7 @@ void SDMMaterialProperties::OnExpansionStateChanged(const TSharedRef<ICustomDeta
 		return;
 	}
 
-	FDMWidgetStatics::Get().SetExpansionState(MaterialModelBase, *ItemId.GetItemName(), bInExpansionState);
+	FDMWidgetLibrary::Get().SetExpansionState(MaterialModelBase, *ItemId.GetItemName(), bInExpansionState);
 }
 
 void SDMMaterialProperties::OnSettingsUpdated(const FPropertyChangedEvent& InPropertyChangedEvent)
diff --git a/Engine/Plugins/Experimental/DynamicMaterial/Source/DynamicMaterialEditor/Private/UI/Widgets/Editor/SDMMaterialSlotEditor.cpp b/Engine/Plugins/Experimental/DynamicMaterial/Source/DynamicMaterialEditor/Private/UI/Widgets/Editor/SDMMaterialSlotEditor.cpp
--- a/Engine/Plugins/Experimental/DynamicMaterial/Source/DynamicMaterialEditor/Private/UI/Widgets/Editor/SDMMaterialSlotEditor.cpp
+++ b/Engine/Plugins/Experimental/DynamicMaterial/Source/DynamicMaterialEditor/Private/UI/Widgets/Editor/SDMMaterialSlotEditor.cpp
@@ -37,7 +37,7 @@
 #include "UI/Menus/DMMaterialSlotLayerAddEffectMenus.h"
 #include "UI/Menus/DMMaterialSlotLayerMenus.h"
 #include "UI/Utils/DMDropTargetPrivateSetter.h"
-#include "UI/Utils/DMWidgetStatics.h"
+#include "UI/Utils/DMWidgetLibrary.h"
 #include "UI/Widgets/Editor/SDMMaterialComponentEditor.h"
 #include "UI/Widgets/Editor/SlotEditor/SDMMaterialSlotLayerView.h"
 #include "UI/Widgets/SDMMaterialEditor.h"
@@ -60,7 +60,7 @@ void SDMMaterialSlotEditor::PrivateRegisterAttributes(FSlateAttributeDescriptor:
 
 SDMMaterialSlotEditor::~SDMMaterialSlotEditor()
 {
-	FDMWidgetStatics::Get().ClearPropertyHandles(this);
+	FDMWidgetLibrary::Get().ClearPropertyHandles(this);
 
 	if (!FDynamicMaterialModule::AreUObjectsSafe())
 	{
@@ -649,7 +649,7 @@ TSharedRef<SWidget> SDMMaterialSlotEditor::CreateSlot_Container()
 
 TSharedRef<SWidget> SDMMaterialSlotEditor::CreateSlot_SlotSettings()
 {
-	FDMWidgetStatics::Get().ClearPropertyHandles(this);
+	FDMWidgetLibrary::Get().ClearPropertyHandles(this);
 
 	return CreateSlot_LayerOpacity();
 }
@@ -693,7 +693,7 @@ TSharedRef<SWidget> SDMMaterialSlotEditor::CreateSlot_LayerOpacity()
 							}
 						}
 
-						FDMPropertyHandle PropertyHandle = FDMWidgetStatics::Get().GetPropertyHandle(this, OpacityValue, UDMMaterialValue::ValueName);
+						FDMPropertyHandle PropertyHandle = FDMWidgetLibrary::Get().GetPropertyHandle(this, OpacityValue, UDMMaterialValue::ValueName);
 
 						FCustomDetailsViewArgs Args;
 						Args.KeyframeHandler = KeyframeHandler;
@@ -725,9 +725,9 @@ TSharedRef<SWidget> SDMMaterialSlotEditor::CreateSlot_LayerOpacity()
 
 	if (ValueWidget.IsValid())
 	{
-		if (TSharedPtr<SWidget> FoundPropertyValueWidget = FDMWidgetStatics::Get().FindWidgetInHierarchy(ValueWidget.ToSharedRef(), FDMWidgetStatics::PropertyValueWidget))
+		if (TSharedPtr<SWidget> FoundPropertyValueWidget = FDMWidgetLibrary::Get().FindWidgetInHierarchy(ValueWidget.ToSharedRef(), FDMWidgetLibrary::PropertyValueWidget))
 		{
-			if (TSharedPtr<SWidget> InnerPropertyWidget = FDMWidgetStatics::Get().GetInnerPropertyValueWidget(FoundPropertyValueWidget.ToSharedRef()))
+			if (TSharedPtr<SWidget> InnerPropertyWidget = FDMWidgetLibrary::Get().GetInnerPropertyValueWidget(FoundPropertyValueWidget.ToSharedRef()))
 			{
 				ValueWidget = InnerPropertyWidget;
 			}
diff --git a/Engine/Plugins/Experimental/DynamicMaterial/Source/DynamicMaterialEditor/Private/UI/Widgets/Editor/SDMObjectEditorWidgetBase.cpp b/Engine/Plugins/Experimental/DynamicMaterial/Source/DynamicMaterialEditor/Private/UI/Widgets/Editor/SDMObjectEditorWidgetBase.cpp
--- a/Engine/Plugins/Experimental/DynamicMaterial/Source/DynamicMaterialEditor/Private/UI/Widgets/Editor/SDMObjectEditorWidgetBase.cpp
+++ b/Engine/Plugins/Experimental/DynamicMaterial/Source/DynamicMaterialEditor/Private/UI/Widgets/Editor/SDMObjectEditorWidgetBase.cpp
@@ -12,7 +12,7 @@
 #include "Items/ICustomDetailsViewCustomItem.h"
 #include "Items/ICustomDetailsViewItem.h"
 #include "Model/DynamicMaterialModelBase.h"
-#include "UI/Utils/DMWidgetStatics.h"
+#include "UI/Utils/DMWidgetLibrary.h"
 #include "UI/Widgets/SDMMaterialEditor.h"
 #include "Widgets/Layout/SBox.h"
 #include "Widgets/Layout/SScrollBox.h"
@@ -26,7 +26,7 @@ void SDMObjectEditorWidgetBase::PrivateRegisterAttributes(FSlateAttributeDescrip
 
 SDMObjectEditorWidgetBase::~SDMObjectEditorWidgetBase()
 {
-	FDMWidgetStatics::Get().ClearPropertyHandles(this);
+	FDMWidgetLibrary::Get().ClearPropertyHandles(this);
 }
 
 void SDMObjectEditorWidgetBase::Construct(const FArguments& InArgs, const TSharedRef<SDMMaterialEditor>& InEditorWidget, UObject* InObject)
@@ -82,7 +82,7 @@ void SDMObjectEditorWidgetBase::Validate()
 
 TSharedRef<SWidget> SDMObjectEditorWidgetBase::CreateWidget()
 {
-	FDMWidgetStatics::Get().ClearPropertyHandles(this);
+	FDMWidgetLibrary::Get().ClearPropertyHandles(this);
 
 	UObject* Object = ObjectWeak.Get();
 
@@ -152,7 +152,7 @@ TSharedRef<ICustomDetailsViewItem> SDMObjectEditorWidgetBase::GetDefaultCategory
 		InDetailsView->ExtendTree(InRootId, ECustomDetailsTreeInsertPosition::Child, DefaultCategoryItem.ToSharedRef());
 
 		bool bExpansionState = true;
-		FDMWidgetStatics::Get().GetExpansionState(ObjectWeak.Get(), DefaultCategoryName, bExpansionState);
+		FDMWidgetLibrary::Get().GetExpansionState(ObjectWeak.Get(), DefaultCategoryName, bExpansionState);
 
 		InDetailsView->SetItemExpansionState(
 			DefaultCategoryItem->GetItemId(), 
@@ -206,7 +206,7 @@ TSharedRef<ICustomDetailsViewItem> SDMObjectEditorWidgetBase::GetCategoryForRow(
 	InDetailsView->ExtendTree(InRootId, ECustomDetailsTreeInsertPosition::Child, CategoryItem.ToSharedRef());
 
 	bool bExpansionState = true;
-	FDMWidgetStatics::Get().GetExpansionState(ObjectWeak.Get(), CategoryName, bExpansionState);
+	FDMWidgetLibrary::Get().GetExpansionState(ObjectWeak.Get(), CategoryName, bExpansionState);
 
 	InDetailsView->SetItemExpansionState(
 		CategoryItem->GetItemId(),
@@ -307,7 +307,7 @@ void SDMObjectEditorWidgetBase::OnExpansionStateChanged(const TSharedRef<ICustom
 		return;
 	}
 
-	FDMWidgetStatics::Get().SetExpansionState(ObjectWeak.Get(), *ItemId.GetItemName(), bInExpansionState);
+	FDMWidgetLibrary::Get().SetExpansionState(ObjectWeak.Get(), *ItemId.GetItemName(), bInExpansionState);
 }
 
 #undef LOCTEXT_NAMESPACE
diff --git a/Engine/Plugins/Experimental/DynamicMaterial/Source/DynamicMaterialEditor/Private/UI/Widgets/Editor/SlotEditor/SDMMaterialSlotLayerItem.cpp b/Engine/Plugins/Experimental/DynamicMaterial/Source/DynamicMaterialEditor/Private/UI/Widgets/Editor/SlotEditor/SDMMaterialSlotLayerItem.cpp
--- a/Engine/Plugins/Experimental/DynamicMaterial/Source/DynamicMaterialEditor/Private/UI/Widgets/Editor/SlotEditor/SDMMaterialSlotLayerItem.cpp
+++ b/Engine/Plugins/Experimental/DynamicMaterial/Source/DynamicMaterialEditor/Private/UI/Widgets/Editor/SlotEditor/SDMMaterialSlotLayerItem.cpp
@@ -17,7 +17,7 @@
 #include "Styling/StyleColors.h"
 #include "UI/DragDrop/DMSlotLayerDragDropOperation.h"
 #include "UI/Menus/DMMaterialStageSourceMenus.h"
-#include "UI/Utils/DMWidgetStatics.h"
+#include "UI/Utils/DMWidgetLibrary.h"
 #include "UI/Widgets/Editor/SDMMaterialComponentEditor.h"
 #include "UI/Widgets/Editor/SDMMaterialSlotEditor.h"
 #include "UI/Widgets/Editor/SlotEditor/SDMMaterialLayerBlendMode.h"
@@ -114,7 +114,7 @@ bool SDMMaterialSlotLayerItem::AreEffectsExpanded() const
 
 	if (UDMMaterialLayerObject* Layer = GetLayer())
 	{
-		FDMWidgetStatics::Get().GetExpansionState(Layer, EffectViewName, bExpanded);
+		FDMWidgetLibrary::Get().GetExpansionState(Layer, EffectViewName, bExpanded);
 	}
 
 	return bExpanded;
@@ -124,7 +124,7 @@ void SDMMaterialSlotLayerItem::SetEffectsExpanded(bool bInExpanded)
 {
 	if (UDMMaterialLayerObject* Layer = GetLayer())
 	{
-		FDMWidgetStatics::Get().SetExpansionState(Layer, EffectViewName, bInExpanded);
+		FDMWidgetLibrary::Get().SetExpansionState(Layer, EffectViewName, bInExpanded);
 	}
 }
 
diff --git a/Engine/Plugins/Experimental/DynamicMaterial/Source/DynamicMaterialEditor/Private/UI/Widgets/Visualizers/SDMTextureUVVisualizerPopout.cpp b/Engine/Plugins/Experimental/DynamicMaterial/Source/DynamicMaterialEditor/Private/UI/Widgets/Visualizers/SDMTextureUVVisualizerPopout.cpp
--- a/Engine/Plugins/Experimental/DynamicMaterial/Source/DynamicMaterialEditor/Private/UI/Widgets/Visualizers/SDMTextureUVVisualizerPopout.cpp
+++ b/Engine/Plugins/Experimental/DynamicMaterial/Source/DynamicMaterialEditor/Private/UI/Widgets/Visualizers/SDMTextureUVVisualizerPopout.cpp
@@ -17,7 +17,7 @@
 #include "Styling/SlateIconFinder.h"
 #include "UI/PropertyGenerators/DMTextureUVDynamicPropertyRowGenerator.h"
 #include "UI/PropertyGenerators/DMTextureUVPropertyRowGenerator.h"
-#include "UI/Utils/DMWidgetStatics.h"
+#include "UI/Utils/DMWidgetLibrary.h"
 #include "UI/Widgets/SDMMaterialEditor.h"
 #include "UI/Widgets/Visualizers/SDMTextureUVVisualizer.h"
 #include "Widgets/Colors/SColorBlock.h"
@@ -240,7 +240,7 @@ FVector2D SDMTextureUVVisualizerPopout::GetSideBlockSize() const
 
 TSharedRef<SWidget> SDMTextureUVVisualizerPopout::CreatePropertyWidget(UDMMaterialComponent* InComponent)
 {
-	FDMWidgetStatics::Get().ClearPropertyHandles(this);
+	FDMWidgetLibrary::Get().ClearPropertyHandles(this);
 
 	if (!InComponent)
 	{
diff --git a/Engine/Plugins/Experimental/DynamicMaterial/Source/DynamicMaterialEditor/Public/IDynamicMaterialEditorModule.h b/Engine/Plugins/Experimental/DynamicMaterial/Source/DynamicMaterialEditor/Public/IDynamicMaterialEditorModule.h
--- a/Engine/Plugins/Experimental/DynamicMaterial/Source/DynamicMaterialEditor/Public/IDynamicMaterialEditorModule.h
+++ b/Engine/Plugins/Experimental/DynamicMaterial/Source/DynamicMaterialEditor/Public/IDynamicMaterialEditorModule.h
@@ -9,15 +9,22 @@
 #include "Templates/SharedPointerFwd.h"
 
 class AActor;
+class SDMMaterialComponentEditor;
 class UClass;
+class UDMMaterialComponent;
 class UDynamicMaterialInstance;
 class UDynamicMaterialModelBase;
 class UObject;
 struct FDMObjectMaterialProperty;
+struct FDMPropertyHandle;
 struct IDMOnWizardCompleteCallback;
 
 DECLARE_DELEGATE_RetVal_OneParam(TArray<FDMObjectMaterialProperty>, FDMGetObjectMaterialPropertiesDelegate, UObject* InObject)
 
+/** Creates property rows in the edit widget. */
+DECLARE_DELEGATE_FourParams(FDMComponentPropertyRowGeneratorDelegate, const TSharedRef<SDMMaterialComponentEditor>&, UDMMaterialComponent*,
+	TArray<FDMPropertyHandle>&, TSet<UDMMaterialComponent*>&)
+
 /** Material Designer - Build your own materials in a slimline editor! */
 class IDynamicMaterialEditorModule : public IModuleInterface
 {
@@ -35,6 +42,20 @@ public:
 		return FModuleManager::LoadModuleChecked<IDynamicMaterialEditorModule>(ModuleName);
 	}
 
+	virtual void RegisterComponentPropertyRowGeneratorDelegate(UClass* InClass, FDMComponentPropertyRowGeneratorDelegate InComponentPropertyRowGeneratorDelegate) = 0;
+
+	template <class InObjClass, class InGenClass>
+	void RegisterComponentPropertyRowGeneratorDelegate()
+	{
+		RegisterComponentPropertyRowGeneratorDelegate(
+			InObjClass::StaticClass(),
+			FDMComponentPropertyRowGeneratorDelegate::CreateSP(
+				InGenClass::Get(),
+				&InGenClass::AddComponentProperties
+			)
+		);
+	}
+
 	virtual void RegisterCustomMaterialPropertyGenerator(UClass* InClass, FDMGetObjectMaterialPropertiesDelegate InGenerator) = 0;
 
 	virtual void RegisterMaterialModelCreatedCallback(const TSharedRef<IDMOnWizardCompleteCallback> InCallback) = 0;
diff --git a/Engine/Plugins/Experimental/DynamicMaterial/Source/DynamicMaterialEditor/Public/UI/Menus/DMMenuContext.h b/Engine/Plugins/Experimental/DynamicMaterial/Source/DynamicMaterialEditor/Public/UI/Menus/DMMenuContext.h
new file mode 100644
--- /dev/null
+++ b/Engine/Plugins/Experimental/DynamicMaterial/Source/DynamicMaterialEditor/Public/UI/Menus/DMMenuContext.h
@@ -0,0 +1,75 @@
+// Copyright Epic Games, Inc. All Rights Reserved.
+
+#pragma once
+
+#include "UObject/Object.h"
+
+#include "Templates/SharedPointer.h"
+
+#include "DMMenuContext.generated.h"
+
+class FName;
+class SDMMaterialEditor;
+class SDMMaterialStage;
+class UDMMaterialLayerObject;
+class UDMMaterialSlot;
+class UDMMaterialStage;
+class UDMMaterialStageBlend;
+class UDMMaterialStageSource;
+class UDynamicMaterialModel;
+class UDynamicMaterialModelBase;
+class UToolMenu;
+
+UCLASS(MinimalAPI)
+class UDMMenuContext : public UObject
+{
+	GENERATED_BODY()
+
+public:
+	static UDMMenuContext* CreateEmpty();
+
+	static UDMMenuContext* CreateEditor(const TWeakPtr<SDMMaterialEditor>& InEditorWidget);
+
+	static UDMMenuContext* CreateLayer(const TWeakPtr<SDMMaterialEditor>& InEditorWidget, UDMMaterialLayerObject* InLayerObject);
+
+	static UDMMenuContext* CreateStage(const TWeakPtr<SDMMaterialEditor>& InEditorWidget, const TWeakPtr<SDMMaterialStage>& InStageWidget);
+
+	static UToolMenu* GenerateContextMenuDefault(const FName InMenuName);
+
+	static UToolMenu* GenerateContextMenuEditor(const FName InMenuName, const TWeakPtr<SDMMaterialEditor>& InEditorWidget);
+
+	static UToolMenu* GenerateContextMenuLayer(const FName InMenuName, const TWeakPtr<SDMMaterialEditor>& InEditorWidget, 
+		UDMMaterialLayerObject* InLayerObject);
+
+	static UToolMenu* GenerateContextMenuStage(const FName InMenuName, const TWeakPtr<SDMMaterialEditor>& InEditorWidget, 
+		const TWeakPtr<SDMMaterialStage>& InStageWidget);
+
+	const TSharedPtr<SDMMaterialEditor> GetEditorWidget() const { return EditorWidgetWeak.Pin(); }
+
+	const TSharedPtr<SDMMaterialStage> GetStageWidget() const { return StageWidgetWeak.Pin(); }
+
+	DYNAMICMATERIALEDITOR_API UDMMaterialSlot* GetSlot() const;
+
+	DYNAMICMATERIALEDITOR_API UDynamicMaterialModelBase* GetModelBase() const;
+
+	DYNAMICMATERIALEDITOR_API UDynamicMaterialModel* GetModel() const;
+
+	DYNAMICMATERIALEDITOR_API UDMMaterialStage* GetStage() const;
+
+	DYNAMICMATERIALEDITOR_API UDMMaterialStageSource* GetStageSource() const;
+
+	DYNAMICMATERIALEDITOR_API UDMMaterialStageBlend* GetStageSourceAsBlend() const;
+
+	DYNAMICMATERIALEDITOR_API UDMMaterialLayerObject* GetLayer() const;
+
+protected:
+	TWeakPtr<SDMMaterialEditor> EditorWidgetWeak;
+	TWeakPtr<SDMMaterialStage> StageWidgetWeak;
+	TWeakObjectPtr<UDMMaterialLayerObject> LayerObjectWeak;
+
+	static UDMMenuContext* Create(const TWeakPtr<SDMMaterialEditor>& InEditorWidget, const TWeakPtr<SDMMaterialStage>& InStageWidget,
+		UDMMaterialLayerObject* InLayerObject);
+
+	static UToolMenu* GenerateContextMenu(const FName InMenuName, const TWeakPtr<SDMMaterialEditor>& InEditorWidget, 
+		const TWeakPtr<SDMMaterialStage>& InStageWidget, UDMMaterialLayerObject* InLayerObject);
+};
diff --git a/Engine/Plugins/Experimental/DynamicMaterial/Source/DynamicMaterialEditor/Public/UI/Utils/IDMWidgetLibrary.h b/Engine/Plugins/Experimental/DynamicMaterial/Source/DynamicMaterialEditor/Public/UI/Utils/IDMWidgetLibrary.h
new file mode 100644
--- /dev/null
+++ b/Engine/Plugins/Experimental/DynamicMaterial/Source/DynamicMaterialEditor/Public/UI/Utils/IDMWidgetLibrary.h
@@ -0,0 +1,22 @@
+// Copyright Epic Games, Inc. All Rights Reserved.
+
+#pragma once
+
+#include "HAL/Platform.h"
+
+class FName;
+class SDMMaterialComponentEditor;
+class SWidget;
+class UObject;
+struct FDMPropertyHandle;
+
+class IDMWidgetLibrary
+{
+public:
+	DYNAMICMATERIALEDITOR_API static IDMWidgetLibrary& Get();
+
+	// Creates a DM property handle for use in the details panel section of the Material Designer
+	virtual FDMPropertyHandle GetPropertyHandle(const SWidget* InOwningWidget, UObject* InObject, FName InPropertyName) = 0;
+
+	virtual FDMPropertyHandle GetPropertyHandle(const SDMMaterialComponentEditor* InOwningWidget, UObject* InObject, FName InPropertyName) = 0;
+};
diff --git a/Engine/Plugins/Experimental/DynamicMaterialMediaStreamBridge/DynamicMaterialMediaStreamBridge.uplugin b/Engine/Plugins/Experimental/DynamicMaterialMediaStreamBridge/DynamicMaterialMediaStreamBridge.uplugin
new file mode 100644
--- /dev/null
+++ b/Engine/Plugins/Experimental/DynamicMaterialMediaStreamBridge/DynamicMaterialMediaStreamBridge.uplugin
@@ -0,0 +1,39 @@
+{
+	"FileVersion": 3,
+	"Version": 1,
+	"VersionName": "1.0",
+	"FriendlyName": "Material Designer Media Stream Bridge",
+	"Description": "Integrates the Media Stream plugin with the Material Designer.",
+	"Category": "Experimental",
+	"CreatedBy": "Epic Games, Inc.",
+	"CreatedByURL": "https://epicgames.com",
+	"DocsURL": "",
+	"MarketplaceURL": "",
+	"SupportURL": "",
+	"CanContainContent": true,
+	"IsBetaVersion": false,
+	"IsExperimentalVersion": true,
+	"Installed": false,
+	"Modules": [
+		{
+			"Name": "DynamicMaterialMediaStreamBridge",
+			"Type": "Runtime",
+			"LoadingPhase": "Default"
+		},
+		{
+			"Name": "DynamicMaterialMediaStreamBridgeEditor",
+			"Type": "Editor",
+			"LoadingPhase": "Default"
+		}
+	],
+	"Plugins": [
+		{
+			"Name": "DynamicMaterial",
+			"Enabled": true
+		},
+		{
+			"Name": "MediaStream",
+			"Enabled": true
+		}
+	]
+}
diff --git a/Engine/Plugins/Experimental/DynamicMaterialMediaStreamBridge/Source/DynamicMaterialMediaStreamBridge/DynamicMaterialMediaStreamBridge.Build.cs b/Engine/Plugins/Experimental/DynamicMaterialMediaStreamBridge/Source/DynamicMaterialMediaStreamBridge/DynamicMaterialMediaStreamBridge.Build.cs
new file mode 100644
--- /dev/null
+++ b/Engine/Plugins/Experimental/DynamicMaterialMediaStreamBridge/Source/DynamicMaterialMediaStreamBridge/DynamicMaterialMediaStreamBridge.Build.cs
@@ -0,0 +1,30 @@
+// Copyright Epic Games, Inc. All Rights Reserved.
+
+using UnrealBuildTool;
+
+public class DynamicMaterialMediaStreamBridge : ModuleRules
+{
+	public DynamicMaterialMediaStreamBridge(ReadOnlyTargetRules Target) : base(Target)
+	{
+		PCHUsage = ModuleRules.PCHUsageMode.UseExplicitOrSharedPCHs;
+
+		PublicDependencyModuleNames.AddRange(
+			new string[]
+			{
+				"Core",
+				"CoreUObject",
+				"Engine",
+				"DynamicMaterial"
+			}
+		);
+
+		PrivateDependencyModuleNames.AddRange(
+			new string[]
+			{
+				"MediaStream"
+			}
+		);
+
+		ShortName = "DMMSBridge";
+	}
+}
diff --git a/Engine/Plugins/Experimental/DynamicMaterialMediaStreamBridge/Source/DynamicMaterialMediaStreamBridge/Private/DMMaterialValueMediaStream.cpp b/Engine/Plugins/Experimental/DynamicMaterialMediaStreamBridge/Source/DynamicMaterialMediaStreamBridge/Private/DMMaterialValueMediaStream.cpp
new file mode 100644
--- /dev/null
+++ b/Engine/Plugins/Experimental/DynamicMaterialMediaStreamBridge/Source/DynamicMaterialMediaStreamBridge/Private/DMMaterialValueMediaStream.cpp
@@ -0,0 +1,237 @@
+// Copyright Epic Games, Inc. All Rights Reserved.
+
+#include "DMMaterialValueMediaStream.h"
+
+#include "DMMaterialValueMediaStreamDynamic.h"
+#include "IMediaStreamPlayer.h"
+#include "MediaStream.h"
+
+#define LOCTEXT_NAMESPACE "DMMaterialValueMediaStream"
+
+UDMMaterialValueMediaStream::UDMMaterialValueMediaStream()
+{
+	MediaStream = CreateDefaultSubobject<UMediaStream>("MediaStream");
+
+#if WITH_EDITOR
+	EditableProperties.Add(GET_MEMBER_NAME_CHECKED(UDMMaterialValueMediaStream, MediaStream));
+#endif
+}
+
+UDMMaterialValueMediaStream::~UDMMaterialValueMediaStream()
+{
+}
+
+UMediaStream* UDMMaterialValueMediaStream::GetMediaStream() const
+{
+	return MediaStream;
+}
+
+#if WITH_EDITOR
+TSharedPtr<FJsonValue> UDMMaterialValueMediaStream::JsonSerialize() const
+{
+	if (!MediaStream)
+	{
+		return nullptr;
+	}
+
+	TMap<FString, TSharedPtr<FJsonValue>> Data;
+
+	const FMediaStreamSource& Source = MediaStream->GetSource();
+
+	Data.Add(GET_MEMBER_NAME_STRING_CHECKED(FMediaStreamSource, Scheme), FDMJsonUtils::Serialize(Source.Scheme));
+	Data.Add(GET_MEMBER_NAME_STRING_CHECKED(FMediaStreamSource, Path), FDMJsonUtils::Serialize(Source.Path));
+
+	if (TScriptInterface<IMediaStreamPlayer> Player = MediaStream->GetPlayer())
+	{
+		const FMediaStreamPlayerConfig& PlayerConfig = Player->GetPlayerConfig();
+		Data.Add(TEXT("PlayerConfig"), FDMJsonUtils::Serialize<FMediaStreamPlayerConfig>(PlayerConfig));
+
+		const FMediaStreamTextureConfig& TextureConfig = Player->GetTextureConfig();
+		Data.Add(TEXT("TextureConfig"), FDMJsonUtils::Serialize<FMediaStreamTextureConfig>(TextureConfig));
+	}
+
+	return FDMJsonUtils::Serialize(Data);
+}
+
+bool UDMMaterialValueMediaStream::JsonDeserialize(const TSharedPtr<FJsonValue>& InJsonValue)
+{
+	if (!MediaStream)
+	{
+		return false;
+	}
+
+	TMap<FString, TSharedPtr<FJsonValue>> Data;
+
+	if (!FDMJsonUtils::Deserialize(InJsonValue, Data))
+	{
+		return false;
+	}
+
+	bool bSuccess = true;
+
+	FMediaStreamSource Source;
+
+	if (const TSharedPtr<FJsonValue>* JsonValue = Data.Find(GET_MEMBER_NAME_STRING_CHECKED(FMediaStreamSource, Scheme)))
+	{
+		FDMJsonUtils::Deserialize(*JsonValue, Source.Scheme);
+	}
+	else
+	{
+		bSuccess = false;
+	}
+
+	if (const TSharedPtr<FJsonValue>* JsonValue = Data.Find(GET_MEMBER_NAME_STRING_CHECKED(FMediaStreamSource, Path)))
+	{
+		FDMJsonUtils::Deserialize(*JsonValue, Source.Path);
+	}
+	else
+	{
+		bSuccess = false;
+	}
+
+	if (bSuccess)
+	{
+		MediaStream->SetSource(Source);
+	}
+
+	if (TScriptInterface<IMediaStreamPlayer> Player = MediaStream->GetPlayer())
+	{
+		if (const TSharedPtr<FJsonValue>* JsonValue = Data.Find(TEXT("PlayerConfig")))
+		{
+			FMediaStreamPlayerConfig PlayerConfig;
+			FDMJsonUtils::Deserialize<FMediaStreamPlayerConfig>(*JsonValue, PlayerConfig);
+			Player->SetPlayerConfig(PlayerConfig);
+		}
+
+		if (const TSharedPtr<FJsonValue>* JsonValue = Data.Find(TEXT("TextureConfig")))
+		{
+			FMediaStreamTextureConfig TextureConfig;
+			FDMJsonUtils::Deserialize<FMediaStreamTextureConfig>(*JsonValue, TextureConfig);
+			Player->SetTextureConfig(TextureConfig);
+		}
+	}
+
+	return bSuccess;
+}
+
+UDMMaterialValueDynamic* UDMMaterialValueMediaStream::ToDynamic(UDynamicMaterialModelDynamic* InMaterialModelDynamic)
+{
+	UDMMaterialValueMediaStreamDynamic* ValueDynamic = UDMMaterialValueDynamic::CreateValueDynamic<UDMMaterialValueMediaStreamDynamic>(InMaterialModelDynamic, this);
+
+	if (MediaStream)
+	{
+		if (UMediaStream* MediaStreamDynamic = ValueDynamic->GetMediaStream())
+		{
+			MediaStreamDynamic->SetSource(MediaStream->GetSource());
+		}
+	}
+
+	return ValueDynamic;
+}
+#endif
+
+#if WITH_EDITOR
+FString UDMMaterialValueMediaStream::GetComponentPathComponent() const
+{
+	return TEXT("MediaStream");
+}
+
+FText UDMMaterialValueMediaStream::GetComponentDescription() const
+{
+	return LOCTEXT("MediaStream", "Media Stream");
+}
+
+void UDMMaterialValueMediaStream::PostEditorDuplicate(UDynamicMaterialModel* InMaterialModel, UDMMaterialComponent* InParent)
+{
+	Super::PostEditorDuplicate(InMaterialModel, InParent);
+
+	UpdatePlayer();
+}
+
+void UDMMaterialValueMediaStream::PostEditChangeProperty(FPropertyChangedEvent& InPropertyChangedEvent)
+{
+	Super::PostEditChangeProperty(InPropertyChangedEvent);
+
+	UpdatePlayer();
+}
+
+void UDMMaterialValueMediaStream::PostLoad()
+{
+	Super::PostLoad();
+
+	UpdatePlayer();
+}
+
+void UDMMaterialValueMediaStream::UpdateTextureFromMediaStream()
+{
+	UTexture* MediaTexture = nullptr;
+
+	if (IsValid(MediaStream))
+	{
+		if (IMediaStreamPlayer* MediaStreamPlayer = MediaStream->GetPlayer().GetInterface())
+		{
+			MediaTexture = MediaStreamPlayer->GetMediaTexture();
+		}
+	}
+
+	SetValue(MediaTexture);
+}
+
+void UDMMaterialValueMediaStream::SubscribeToPlayerChange()
+{
+	if (IsValid(MediaStream))
+	{
+		MediaStream->GetOnPlayerChanged().AddDynamic(this, &UDMMaterialValueMediaStream::OnPlayerChanged);
+		SubscribedStreamWeak = MediaStream;
+	}
+	else
+	{
+		SubscribedStreamWeak.Reset();
+	}
+}
+
+void UDMMaterialValueMediaStream::UnsubscribeFromPlayerChange()
+{
+	if (UMediaStream* SubscribedStream = SubscribedStreamWeak.Get())
+	{
+		SubscribedStream->GetOnPlayerChanged().RemoveAll(this);
+	}
+
+	if (IsValid(MediaStream))
+	{
+		MediaStream->GetOnPlayerChanged().RemoveAll(this);
+	}
+}
+
+void UDMMaterialValueMediaStream::OnPlayerChanged(UMediaStream* InMediaStream, TScriptInterface<IMediaStreamPlayer> InPlayer)
+{
+	UpdatePlayer();
+
+	Update(this, EDMUpdateType::RefreshDetailView);
+}
+
+void UDMMaterialValueMediaStream::UpdatePlayer()
+{
+	UnsubscribeFromPlayerChange();
+	SubscribeToPlayerChange();
+	UpdateTextureFromMediaStream();
+}
+
+void UDMMaterialValueMediaStream::OnComponentAdded()
+{
+	Super::OnComponentAdded();
+
+	UpdatePlayer();
+}
+
+void UDMMaterialValueMediaStream::OnComponentRemoved()
+{
+	Super::OnComponentRemoved();
+
+	UnsubscribeFromPlayerChange();
+
+	SetValue(nullptr);
+}
+#endif
+
+#undef LOCTEXT_NAMESPACE
diff --git a/Engine/Plugins/Experimental/DynamicMaterialMediaStreamBridge/Source/DynamicMaterialMediaStreamBridge/Private/DMMaterialValueMediaStreamDynamic.cpp b/Engine/Plugins/Experimental/DynamicMaterialMediaStreamBridge/Source/DynamicMaterialMediaStreamBridge/Private/DMMaterialValueMediaStreamDynamic.cpp
new file mode 100644
--- /dev/null
+++ b/Engine/Plugins/Experimental/DynamicMaterialMediaStreamBridge/Source/DynamicMaterialMediaStreamBridge/Private/DMMaterialValueMediaStreamDynamic.cpp
@@ -0,0 +1,211 @@
+// Copyright Epic Games, Inc. All Rights Reserved.
+
+#include "DMMaterialValueMediaStreamDynamic.h"
+
+#include "DMMaterialValueMediaStream.h"
+#include "IMediaStreamPlayer.h"
+#include "MediaStream.h"
+
+#define LOCTEXT_NAMESPACE "DMMaterialValueMediaStreamDynamic"
+
+UDMMaterialValueMediaStreamDynamic::UDMMaterialValueMediaStreamDynamic()
+{
+	MediaStream = CreateDefaultSubobject<UMediaStream>("MediaStream");
+
+#if WITH_EDITOR
+	EditableProperties.Add(GET_MEMBER_NAME_CHECKED(UDMMaterialValueMediaStreamDynamic, MediaStream));
+#endif
+}
+
+#if WITH_EDITOR
+TSharedPtr<FJsonValue> UDMMaterialValueMediaStreamDynamic::JsonSerialize() const
+{
+	if (!MediaStream)
+	{
+		return nullptr;
+	}
+
+	TMap<FString, TSharedPtr<FJsonValue>> Data;
+
+	const FMediaStreamSource& Source = MediaStream->GetSource();
+
+	Data.Add(GET_MEMBER_NAME_STRING_CHECKED(FMediaStreamSource, Scheme), FDMJsonUtils::Serialize(Source.Scheme));
+	Data.Add(GET_MEMBER_NAME_STRING_CHECKED(FMediaStreamSource, Path), FDMJsonUtils::Serialize(Source.Path));
+
+	if (TScriptInterface<IMediaStreamPlayer> Player = MediaStream->GetPlayer())
+	{
+		const FMediaStreamPlayerConfig& PlayerConfig = Player->GetPlayerConfig();
+		Data.Add(TEXT("PlayerConfig"), FDMJsonUtils::Serialize<FMediaStreamPlayerConfig>(PlayerConfig));
+
+		const FMediaStreamTextureConfig& TextureConfig = Player->GetTextureConfig();
+		Data.Add(TEXT("TextureConfig"), FDMJsonUtils::Serialize<FMediaStreamTextureConfig>(TextureConfig));
+	}
+
+	return FDMJsonUtils::Serialize(Data);
+}
+
+bool UDMMaterialValueMediaStreamDynamic::JsonDeserialize(const TSharedPtr<FJsonValue>& InJsonValue)
+{
+	if (!MediaStream)
+	{
+		return false;
+	}
+
+	TMap<FString, TSharedPtr<FJsonValue>> Data;
+
+	if (!FDMJsonUtils::Deserialize(InJsonValue, Data))
+	{
+		return false;
+	}
+
+	bool bSuccess = true;
+
+	FMediaStreamSource Source;
+
+	if (const TSharedPtr<FJsonValue>* JsonValue = Data.Find(GET_MEMBER_NAME_STRING_CHECKED(FMediaStreamSource, Scheme)))
+	{
+		FDMJsonUtils::Deserialize(*JsonValue, Source.Scheme);
+	}
+	else
+	{
+		bSuccess = false;
+	}
+
+	if (const TSharedPtr<FJsonValue>* JsonValue = Data.Find(GET_MEMBER_NAME_STRING_CHECKED(FMediaStreamSource, Path)))
+	{
+		FDMJsonUtils::Deserialize(*JsonValue, Source.Path);
+	}
+	else
+	{
+		bSuccess = false;
+	}
+
+	if (bSuccess)
+	{
+		MediaStream->SetSource(Source);
+	}
+
+	if (TScriptInterface<IMediaStreamPlayer> Player = MediaStream->GetPlayer())
+	{
+		if (const TSharedPtr<FJsonValue>* JsonValue = Data.Find(TEXT("PlayerConfig")))
+		{
+			FMediaStreamPlayerConfig PlayerConfig;
+			FDMJsonUtils::Deserialize<FMediaStreamPlayerConfig>(*JsonValue, PlayerConfig);
+			Player->SetPlayerConfig(PlayerConfig);
+		}
+
+		if (const TSharedPtr<FJsonValue>* JsonValue = Data.Find(TEXT("TextureConfig")))
+		{
+			FMediaStreamTextureConfig TextureConfig;
+			FDMJsonUtils::Deserialize<FMediaStreamTextureConfig>(*JsonValue, TextureConfig);
+			Player->SetTextureConfig(TextureConfig);
+		}
+	}
+
+	return bSuccess;
+}
+
+FString UDMMaterialValueMediaStreamDynamic::GetComponentPathComponent() const
+{
+	return TEXT("MediaStream");
+}
+
+FText UDMMaterialValueMediaStreamDynamic::GetComponentDescription() const
+{
+	return LOCTEXT("MediaStream", "Media Stream");
+}
+
+void UDMMaterialValueMediaStreamDynamic::PostEditorDuplicate(UDynamicMaterialModel* InMaterialModel, UDMMaterialComponent* InParent)
+{
+	Super::PostEditorDuplicate(InMaterialModel, InParent);
+
+	UpdatePlayer();
+}
+
+void UDMMaterialValueMediaStreamDynamic::PostEditChangeProperty(FPropertyChangedEvent& InPropertyChangedEvent)
+{
+	Super::PostEditChangeProperty(InPropertyChangedEvent);
+
+	UpdatePlayer();
+}
+
+void UDMMaterialValueMediaStreamDynamic::PostLoad()
+{
+	Super::PostLoad();
+
+	UpdatePlayer();
+}
+
+void UDMMaterialValueMediaStreamDynamic::UpdateTextureFromMediaStream()
+{
+	UTexture* MediaTexture = nullptr;
+
+	if (IsValid(MediaStream))
+	{
+		if (IMediaStreamPlayer* MediaStreamPlayer = MediaStream->GetPlayer().GetInterface())
+		{
+			MediaTexture = MediaStreamPlayer->GetMediaTexture();
+		}
+	}
+
+	SetValue(MediaTexture);
+}
+
+void UDMMaterialValueMediaStreamDynamic::SubscribeToPlayerChange()
+{
+	if (IsValid(MediaStream))
+	{
+		MediaStream->GetOnPlayerChanged().AddDynamic(this, &UDMMaterialValueMediaStreamDynamic::OnPlayerChanged);
+		SubscribedStreamWeak = MediaStream;
+	}
+	else
+	{
+		SubscribedStreamWeak.Reset();
+	}
+}
+
+void UDMMaterialValueMediaStreamDynamic::UnsubscribeFromPlayerChange()
+{
+	if (UMediaStream* SubscribedStream = SubscribedStreamWeak.Get())
+	{
+		SubscribedStream->GetOnPlayerChanged().RemoveAll(this);
+	}
+
+	if (IsValid(MediaStream))
+	{
+		MediaStream->GetOnPlayerChanged().RemoveAll(this);
+	}
+}
+
+void UDMMaterialValueMediaStreamDynamic::OnPlayerChanged(UMediaStream* InMediaStream, TScriptInterface<IMediaStreamPlayer> InPlayer)
+{
+	UpdatePlayer();
+
+	Update(this, EDMUpdateType::RefreshDetailView);
+}
+
+void UDMMaterialValueMediaStreamDynamic::UpdatePlayer()
+{
+	UnsubscribeFromPlayerChange();
+	SubscribeToPlayerChange();
+	UpdateTextureFromMediaStream();
+}
+
+void UDMMaterialValueMediaStreamDynamic::OnComponentAdded()
+{
+	Super::OnComponentAdded();
+
+	UpdatePlayer();
+}
+
+void UDMMaterialValueMediaStreamDynamic::OnComponentRemoved()
+{
+	Super::OnComponentRemoved();
+
+	UnsubscribeFromPlayerChange();
+
+	SetValue(nullptr);
+}
+#endif
+
+#undef LOCTEXT_NAMESPACE
diff --git a/Engine/Plugins/Experimental/DynamicMaterialMediaStreamBridge/Source/DynamicMaterialMediaStreamBridge/Private/DynamicMaterialMediaStreamBridgeModule.cpp b/Engine/Plugins/Experimental/DynamicMaterialMediaStreamBridge/Source/DynamicMaterialMediaStreamBridge/Private/DynamicMaterialMediaStreamBridgeModule.cpp
new file mode 100644
--- /dev/null
+++ b/Engine/Plugins/Experimental/DynamicMaterialMediaStreamBridge/Source/DynamicMaterialMediaStreamBridge/Private/DynamicMaterialMediaStreamBridgeModule.cpp
@@ -0,0 +1,15 @@
+// Copyright Epic Games, Inc. All Rights Reserved.
+
+#include "DynamicMaterialMediaStreamBridgeModule.h"
+
+#include "Modules/ModuleManager.h"
+
+void FDynamicMaterialMediaStreamBridgeModule::StartupModule()
+{
+}
+
+void FDynamicMaterialMediaStreamBridgeModule::ShutdownModule()
+{
+}
+
+IMPLEMENT_MODULE(FDynamicMaterialMediaStreamBridgeModule, DynamicMaterialMediaStreamBridge)
diff --git a/Engine/Plugins/Experimental/DynamicMaterialMediaStreamBridge/Source/DynamicMaterialMediaStreamBridge/Private/DynamicMaterialMediaStreamBridgeModule.h b/Engine/Plugins/Experimental/DynamicMaterialMediaStreamBridge/Source/DynamicMaterialMediaStreamBridge/Private/DynamicMaterialMediaStreamBridgeModule.h
new file mode 100644
--- /dev/null
+++ b/Engine/Plugins/Experimental/DynamicMaterialMediaStreamBridge/Source/DynamicMaterialMediaStreamBridge/Private/DynamicMaterialMediaStreamBridgeModule.h
@@ -0,0 +1,17 @@
+// Copyright Epic Games, Inc. All Rights Reserved.
+
+#pragma once
+
+#include "Modules/ModuleInterface.h"
+
+/**
+ * Media Stream Material Designer Bridge - Integrates the Media Stream plugin with the Material Designer.
+ */
+class FDynamicMaterialMediaStreamBridgeModule : public IModuleInterface
+{
+public:
+	//~ Begin IModuleInterface
+	virtual void StartupModule() override;
+	virtual void ShutdownModule() override;
+	//~ End IModuleInterface
+};
diff --git a/Engine/Plugins/Experimental/DynamicMaterialMediaStreamBridge/Source/DynamicMaterialMediaStreamBridge/Public/DMMaterialValueMediaStream.h b/Engine/Plugins/Experimental/DynamicMaterialMediaStreamBridge/Source/DynamicMaterialMediaStreamBridge/Public/DMMaterialValueMediaStream.h
new file mode 100644
--- /dev/null
+++ b/Engine/Plugins/Experimental/DynamicMaterialMediaStreamBridge/Source/DynamicMaterialMediaStreamBridge/Public/DMMaterialValueMediaStream.h
@@ -0,0 +1,75 @@
+// Copyright Epic Games, Inc. All Rights Reserved.
+
+#pragma once
+
+#include "Components/MaterialValues/DMMaterialValueTexture.h"
+#include "UObject/ScriptInterface.h"
+#include "DMMaterialValueMediaStream.generated.h"
+
+class IMediaStreamPlayer;
+class UMediaStream;
+
+/**
+ * Component representing a render target texture value. Manages its own parameter.
+ */
+UCLASS(MinimalAPI, BlueprintType, ClassGroup = "Material Designer")
+class UDMMaterialValueMediaStream : public UDMMaterialValueTexture
+{
+	GENERATED_BODY()
+
+public:
+	UDMMaterialValueMediaStream();
+
+	virtual ~UDMMaterialValueMediaStream() override;
+
+	UFUNCTION(BlueprintPure, Category = "Material Designer")
+	DYNAMICMATERIALMEDIASTREAMBRIDGE_API  UMediaStream* GetMediaStream() const;
+
+#if WITH_EDITOR
+	//~ Begin IDMJsonSerializable
+	DYNAMICMATERIALMEDIASTREAMBRIDGE_API  virtual TSharedPtr<FJsonValue> JsonSerialize() const override;
+	DYNAMICMATERIALMEDIASTREAMBRIDGE_API  virtual bool JsonDeserialize(const TSharedPtr<FJsonValue>& InJsonValue) override;
+	//~ End IDMJsonSerializable
+
+	//~ Begin UDMMaterialValue
+	virtual bool AllowEditValue() const { return false; }
+	DYNAMICMATERIALMEDIASTREAMBRIDGE_API  virtual UDMMaterialValueDynamic* ToDynamic(UDynamicMaterialModelDynamic* InMaterialModelDynamic) override;
+	//~ End UDMMaterialValue
+
+	//~ Begin UDMMaterialComponent
+	DYNAMICMATERIALMEDIASTREAMBRIDGE_API virtual FString GetComponentPathComponent() const override;
+	DYNAMICMATERIALMEDIASTREAMBRIDGE_API virtual FText GetComponentDescription() const override;
+	DYNAMICMATERIALMEDIASTREAMBRIDGE_API virtual void PostEditorDuplicate(UDynamicMaterialModel* InMaterialModel, UDMMaterialComponent* InParent) override;
+	//~ End UDMMaterialComponent
+
+	//~ Begin UObject
+	DYNAMICMATERIALMEDIASTREAMBRIDGE_API virtual void PostEditChangeProperty(FPropertyChangedEvent& InPropertyChangedEvent) override;
+	DYNAMICMATERIALMEDIASTREAMBRIDGE_API virtual void PostLoad() override;
+	//~ End UObject
+#endif
+
+protected:
+	UPROPERTY(VisibleAnywhere, BlueprintReadOnly, Category = "Material Designer|Media Stream",
+		meta = (AllowPrivateAccess = "true", NotKeyframeable))
+	TObjectPtr<UMediaStream> MediaStream;
+
+	TWeakObjectPtr<UMediaStream> SubscribedStreamWeak;
+
+#if WITH_EDITOR
+	void UpdateTextureFromMediaStream();
+
+	void SubscribeToPlayerChange();
+
+	void UnsubscribeFromPlayerChange();
+
+	UFUNCTION()
+	void OnPlayerChanged(UMediaStream* InMediaStream, TScriptInterface<IMediaStreamPlayer> InPlayer);
+
+	void UpdatePlayer();
+
+	//~ Begin UDMMaterialComponent
+	DYNAMICMATERIALMEDIASTREAMBRIDGE_API  virtual void OnComponentAdded() override;
+	DYNAMICMATERIALMEDIASTREAMBRIDGE_API  virtual void OnComponentRemoved() override;
+	//~ End UDMMaterialComponent
+#endif
+};
diff --git a/Engine/Plugins/Experimental/DynamicMaterialMediaStreamBridge/Source/DynamicMaterialMediaStreamBridge/Public/DMMaterialValueMediaStreamDynamic.h b/Engine/Plugins/Experimental/DynamicMaterialMediaStreamBridge/Source/DynamicMaterialMediaStreamBridge/Public/DMMaterialValueMediaStreamDynamic.h
new file mode 100644
--- /dev/null
+++ b/Engine/Plugins/Experimental/DynamicMaterialMediaStreamBridge/Source/DynamicMaterialMediaStreamBridge/Public/DMMaterialValueMediaStreamDynamic.h
@@ -0,0 +1,72 @@
+// Copyright Epic Games, Inc. All Rights Reserved.
+
+#pragma once
+
+#include "Components/MaterialValuesDynamic/DMMaterialValueTextureDynamic.h"
+
+#include "DMMaterialValueMediaStreamDynamic.generated.h"
+
+class IMediaStreamPlayer;
+class UMediaStream;
+
+/**
+ * Link to a UDMMaterialValueMediaStream for Material Designer Model Dynamics.
+ */
+UCLASS(MinimalAPI, BlueprintType, ClassGroup = "Material Designer")
+class UDMMaterialValueMediaStreamDynamic : public UDMMaterialValueTextureDynamic
+{
+	GENERATED_BODY()
+
+public:
+#if WITH_EDITOR
+	using Super::PostEditorDuplicate;
+#endif
+
+	UDMMaterialValueMediaStreamDynamic();
+
+	UFUNCTION(BlueprintPure, Category = "Material Designer")
+	UMediaStream* GetMediaStream() const { return MediaStream; }
+
+#if WITH_EDITOR
+	//~ Begin IDMJsonSerializable
+	DYNAMICMATERIALMEDIASTREAMBRIDGE_API virtual TSharedPtr<FJsonValue> JsonSerialize() const override;
+	DYNAMICMATERIALMEDIASTREAMBRIDGE_API virtual bool JsonDeserialize(const TSharedPtr<FJsonValue>& InJsonValue) override;
+	//~ End IDMJsonSerializable
+
+	//~ Begin UDMMaterialComponent
+	DYNAMICMATERIALMEDIASTREAMBRIDGE_API virtual FString GetComponentPathComponent() const override;
+	DYNAMICMATERIALMEDIASTREAMBRIDGE_API virtual FText GetComponentDescription() const override;
+	DYNAMICMATERIALMEDIASTREAMBRIDGE_API virtual void PostEditorDuplicate(UDynamicMaterialModel* InMaterialModel, UDMMaterialComponent* InParent) override;
+	//~ End UDMMaterialComponent
+
+	//~ Begin UObject
+	DYNAMICMATERIALMEDIASTREAMBRIDGE_API virtual void PostEditChangeProperty(FPropertyChangedEvent& InPropertyChangedEvent) override;
+	DYNAMICMATERIALMEDIASTREAMBRIDGE_API virtual void PostLoad() override;
+	//~ End UObject
+#endif
+
+protected:
+	UPROPERTY(VisibleAnywhere, BlueprintReadOnly, Category = "Material Designer|Media Stream",
+		meta = (AllowPrivateAccess = "true", NotKeyframeable))
+	TObjectPtr<UMediaStream> MediaStream;
+
+	TWeakObjectPtr<UMediaStream> SubscribedStreamWeak;
+
+#if WITH_EDITOR
+	void UpdateTextureFromMediaStream();
+
+	void SubscribeToPlayerChange();
+
+	void UnsubscribeFromPlayerChange();
+
+	UFUNCTION()
+	void OnPlayerChanged(UMediaStream* InMediaStream, TScriptInterface<IMediaStreamPlayer> InPlayer);
+
+	void UpdatePlayer();
+
+	//~ Begin UDMMaterialComponent
+	DYNAMICMATERIALMEDIASTREAMBRIDGE_API  virtual void OnComponentAdded() override;
+	DYNAMICMATERIALMEDIASTREAMBRIDGE_API  virtual void OnComponentRemoved() override;
+	//~ End UDMMaterialComponent
+#endif	
+};
diff --git a/Engine/Plugins/Experimental/DynamicMaterialMediaStreamBridge/Source/DynamicMaterialMediaStreamBridgeEditor/DynamicMaterialMediaStreamBridgeEditor.Build.cs b/Engine/Plugins/Experimental/DynamicMaterialMediaStreamBridge/Source/DynamicMaterialMediaStreamBridgeEditor/DynamicMaterialMediaStreamBridgeEditor.Build.cs
new file mode 100644
--- /dev/null
+++ b/Engine/Plugins/Experimental/DynamicMaterialMediaStreamBridge/Source/DynamicMaterialMediaStreamBridgeEditor/DynamicMaterialMediaStreamBridgeEditor.Build.cs
@@ -0,0 +1,36 @@
+// Copyright Epic Games, Inc. All Rights Reserved.
+
+using UnrealBuildTool;
+
+public class DynamicMaterialMediaStreamBridgeEditor : ModuleRules
+{
+	public DynamicMaterialMediaStreamBridgeEditor(ReadOnlyTargetRules Target) : base(Target)
+	{
+		PCHUsage = ModuleRules.PCHUsageMode.UseExplicitOrSharedPCHs;
+
+		PublicDependencyModuleNames.AddRange(
+			new string[]
+			{
+				"Core",
+				"CoreUObject",
+				"DynamicMaterial",
+				"Engine"
+			}
+		);
+
+		PrivateDependencyModuleNames.AddRange(
+			new string[]
+			{
+				"DynamicMaterialEditor",
+				"DynamicMaterialMediaStreamBridge",
+				"MediaStream",
+				"Slate",
+				"SlateCore",
+				"ToolMenus",
+				"UnrealEd"
+			}
+		);
+
+		ShortName = "DMMSBridgeEditor";
+	}
+}
diff --git a/Engine/Plugins/Experimental/DynamicMaterialMediaStreamBridge/Source/DynamicMaterialMediaStreamBridgeEditor/Private/DMMaterialValueMediaStreamPropertyRowGenerator.cpp b/Engine/Plugins/Experimental/DynamicMaterialMediaStreamBridge/Source/DynamicMaterialMediaStreamBridgeEditor/Private/DMMaterialValueMediaStreamPropertyRowGenerator.cpp
new file mode 100644
--- /dev/null
+++ b/Engine/Plugins/Experimental/DynamicMaterialMediaStreamBridge/Source/DynamicMaterialMediaStreamBridgeEditor/Private/DMMaterialValueMediaStreamPropertyRowGenerator.cpp
@@ -0,0 +1,89 @@
+// Copyright Epic Games, Inc. All Rights Reserved.
+
+#include "DMMaterialValueMediaStreamPropertyRowGenerator.h"
+
+#include "Components/DMMaterialComponent.h"
+#include "DMEDefs.h"
+#include "DMMaterialValueMediaStream.h"
+#include "IMediaStreamPlayer.h"
+#include "MediaStream.h"
+#include "UI/Utils/IDMWidgetLibrary.h"
+#include "UObject/Object.h"
+
+const TSharedRef<FDMMaterialValueMediaStreamPropertyRowGenerator>& FDMMaterialValueMediaStreamPropertyRowGenerator::Get()
+{
+	static TSharedRef<FDMMaterialValueMediaStreamPropertyRowGenerator> Generator = MakeShared<FDMMaterialValueMediaStreamPropertyRowGenerator>();
+	return Generator;
+}
+
+void FDMMaterialValueMediaStreamPropertyRowGenerator::AddComponentProperties(const TSharedRef<SDMMaterialComponentEditor>& InComponentEditorWidget, UDMMaterialComponent* InComponent,
+	TArray<FDMPropertyHandle>& InOutPropertyRows, TSet<UDMMaterialComponent*>& InOutProcessedObjects)
+{
+	if (!IsValid(InComponent))
+	{
+		return;
+	}
+
+	if (InOutProcessedObjects.Contains(InComponent))
+	{
+		return;
+	}
+
+	UDMMaterialValueMediaStream* MediaStreamValue = Cast<UDMMaterialValueMediaStream>(InComponent);
+
+	if (!MediaStreamValue)
+	{
+		return;
+	}
+
+	InOutProcessedObjects.Add(InComponent);
+
+	UMediaStream* MediaStream = MediaStreamValue->GetMediaStream();
+
+	if (!MediaStream)
+	{
+		return;
+	}
+
+	AddSourceProperty(InComponentEditorWidget, InOutPropertyRows, MediaStream);
+	AddPlayerProperty(InComponentEditorWidget, InOutPropertyRows, MediaStream);
+}
+
+void FDMMaterialValueMediaStreamPropertyRowGenerator::AddSourceProperty(const TSharedRef<SDMMaterialComponentEditor>& InComponentEditorWidget, 
+	TArray<FDMPropertyHandle>& InOutPropertyRows, UMediaStream* InMediaStream)
+{
+	FDMPropertyHandle& NewHandle = InOutPropertyRows.Add_GetRef(
+		IDMWidgetLibrary::Get().GetPropertyHandle(&*InComponentEditorWidget, InMediaStream, TEXT("Source"))
+	);
+}
+
+void FDMMaterialValueMediaStreamPropertyRowGenerator::AddPlayerProperty(const TSharedRef<SDMMaterialComponentEditor>& InComponentEditorWidget,
+	TArray<FDMPropertyHandle>& InOutPropertyRows, UMediaStream* InMediaStream)
+{
+	IMediaStreamPlayer* Player = InMediaStream->GetPlayer().GetInterface();
+
+	if (!Player || Player->IsReadOnly())
+	{
+		return;
+	}
+
+	FDMPropertyHandle& ControlHandle = InOutPropertyRows.Add_GetRef(
+		IDMWidgetLibrary::Get().GetPropertyHandle(&*InComponentEditorWidget, InMediaStream, TEXT("Control"))
+	);
+	ControlHandle.CategoryOverrideName = TEXT("Media Stream Player");
+
+	FDMPropertyHandle& TextureHandle = InOutPropertyRows.Add_GetRef(
+		IDMWidgetLibrary::Get().GetPropertyHandle(&*InComponentEditorWidget, InMediaStream, TEXT("MediaTexture"))
+	);
+	TextureHandle.CategoryOverrideName = TEXT("Media Stream Player");
+
+	FDMPropertyHandle& TextureConfigHandle = InOutPropertyRows.Add_GetRef(
+		IDMWidgetLibrary::Get().GetPropertyHandle(&*InComponentEditorWidget, InMediaStream, TEXT("TextureConfig"))
+	);
+	TextureConfigHandle.CategoryOverrideName = TEXT("Media Stream Player");
+
+	FDMPropertyHandle& PlayerConfigHandle = InOutPropertyRows.Add_GetRef(
+		IDMWidgetLibrary::Get().GetPropertyHandle(&*InComponentEditorWidget, InMediaStream, TEXT("PlayerConfig"))
+	);
+	PlayerConfigHandle.CategoryOverrideName = TEXT("Media Stream Player");
+}
diff --git a/Engine/Plugins/Experimental/DynamicMaterialMediaStreamBridge/Source/DynamicMaterialMediaStreamBridgeEditor/Private/DMMaterialValueMediaStreamPropertyRowGenerator.h b/Engine/Plugins/Experimental/DynamicMaterialMediaStreamBridge/Source/DynamicMaterialMediaStreamBridgeEditor/Private/DMMaterialValueMediaStreamPropertyRowGenerator.h
new file mode 100644
--- /dev/null
+++ b/Engine/Plugins/Experimental/DynamicMaterialMediaStreamBridge/Source/DynamicMaterialMediaStreamBridgeEditor/Private/DMMaterialValueMediaStreamPropertyRowGenerator.h
@@ -0,0 +1,26 @@
+// Copyright Epic Games, Inc. All Rights Reserved.
+
+#pragma once
+
+#include "UI/PropertyGenerators/DMComponentPropertyRowGenerator.h"
+
+class UMediaStream;
+
+class FDMMaterialValueMediaStreamPropertyRowGenerator : public FDMComponentPropertyRowGenerator
+{
+public:
+	static const TSharedRef<FDMMaterialValueMediaStreamPropertyRowGenerator>& Get();
+
+	virtual ~FDMMaterialValueMediaStreamPropertyRowGenerator() override = default;
+
+	//~ Begin FDMComponentPropertyRowGenerator
+	virtual void AddComponentProperties(const TSharedRef<SDMMaterialComponentEditor>& InComponentEditorWidget, UDMMaterialComponent* InComponent,
+		TArray<FDMPropertyHandle>& InOutPropertyRows, TSet<UDMMaterialComponent*>& InOutProcessedObjects) override;
+	//~ End FDMComponentPropertyRowGenerator
+
+	void AddSourceProperty(const TSharedRef<SDMMaterialComponentEditor>& InComponentEditorWidget, TArray<FDMPropertyHandle>& InOutPropertyRows, 
+		UMediaStream* InMediaStream);
+
+	void AddPlayerProperty(const TSharedRef<SDMMaterialComponentEditor>& InComponentEditorWidget, TArray<FDMPropertyHandle>& InOutPropertyRows,
+		UMediaStream* InMediaStream);
+};
diff --git a/Engine/Plugins/Experimental/DynamicMaterialMediaStreamBridge/Source/DynamicMaterialMediaStreamBridgeEditor/Private/DMMediaStreamStageSourceMenuExtender.cpp b/Engine/Plugins/Experimental/DynamicMaterialMediaStreamBridge/Source/DynamicMaterialMediaStreamBridgeEditor/Private/DMMediaStreamStageSourceMenuExtender.cpp
new file mode 100644
--- /dev/null
+++ b/Engine/Plugins/Experimental/DynamicMaterialMediaStreamBridge/Source/DynamicMaterialMediaStreamBridgeEditor/Private/DMMediaStreamStageSourceMenuExtender.cpp
@@ -0,0 +1,152 @@
+// Copyright Epic Games, Inc. All Rights Reserved.
+
+#include "DMMediaStreamStageSourceMenuExtender.h"
+
+#include "Components/DMMaterialStage.h"
+#include "Components/DMMaterialStageBlend.h"
+#include "Components/DMMaterialStageThroughputLayerBlend.h"
+#include "Components/DMMaterialSubStage.h"
+#include "Components/MaterialStageExpressions/DMMSETextureSample.h"
+#include "Components/MaterialStageInputs/DMMSIExpression.h"
+#include "Components/MaterialStageInputs/DMMSIValue.h"
+#include "DMMaterialValueMediaStream.h"
+#include "ScopedTransaction.h"
+#include "ToolMenu.h"
+#include "ToolMenus.h"
+#include "UI/Menus/DMMenuContext.h"
+
+#define LOCTEXT_NAMESPACE "DMMEdiaStreamStageSourceMenuExtender"
+
+void FDMMEdiaStreamStageSourceMenuExtender::Integrate()
+{
+	if (bIntegrated)
+	{
+		return;
+	}
+
+	UToolMenus* ToolMenus = UToolMenus::Get();
+
+	if (!ToolMenus)
+	{
+		return;
+	}
+
+	UToolMenu* SourceMenu = ToolMenus->ExtendMenu(TEXT("MaterialDesigner.MaterialStage.ChangeSource.Advanced"));
+
+	if (!SourceMenu)
+	{
+		return;
+	}
+
+	SourceMenu->AddDynamicSection(
+		TEXT("Advanced"), 
+		FNewToolMenuDelegate::CreateRaw(this, &FDMMEdiaStreamStageSourceMenuExtender::ExtendMenu)
+	);
+
+	bIntegrated = true;
+}
+
+void FDMMEdiaStreamStageSourceMenuExtender::ExtendMenu(UToolMenu* InMenu)
+{
+	UDMMenuContext* MenuContext = InMenu->FindContext<UDMMenuContext>();
+
+	if (!ensure(MenuContext))
+	{
+		return;
+	}
+
+	FToolMenuSection& AdvancedSection = InMenu->AddSection(NAME_None);
+
+	AdvancedSection.AddMenuEntry("MediaStream",
+		LOCTEXT("ChangeSourceMediaStream", "Media Stream"),
+		LOCTEXT("ChangeSourceMediaStreamTooltip", "Change the source of this stage to a Media Stream."),
+		GetDefault<UDMMaterialValueMediaStream>()->GetComponentIcon(),
+		FUIAction(
+			FExecuteAction::CreateRaw(
+				this,
+				&FDMMEdiaStreamStageSourceMenuExtender::ChangeSourceToMediaStreamFromContext,
+				MenuContext
+			)
+		)
+	);
+}
+
+void FDMMEdiaStreamStageSourceMenuExtender::ChangeSourceToMediaStreamFromContext(UDMMenuContext* InMenuContext)
+{
+	if (!IsValid(InMenuContext))
+	{
+		return;
+	}
+
+	UDMMaterialStage* const Stage = InMenuContext->GetStage();
+
+	if (!Stage)
+	{
+		return;
+	}
+
+	UDMMaterialStageSource* const StageSource = Stage->GetSource();
+
+	if (!StageSource)
+	{
+		return;
+	}
+
+	UDMMaterialStageInputExpression* InputExpression = nullptr;
+
+	if (StageSource->IsA<UDMMaterialStageBlend>())
+	{
+		FScopedTransaction Transaction(LOCTEXT("SetStageInputBase", "Set Material Designer Base Source"));
+		Stage->Modify();
+
+		InputExpression = UDMMaterialStageInputExpression::ChangeStageInput_Expression(
+			Stage,
+			UDMMaterialStageExpressionTextureSample::StaticClass(),
+			UDMMaterialStageBlend::InputB,
+			FDMMaterialStageConnectorChannel::WHOLE_CHANNEL,
+			0,
+			FDMMaterialStageConnectorChannel::THREE_CHANNELS
+		);
+	}
+	else if (StageSource->IsA<UDMMaterialStageThroughputLayerBlend>())
+	{
+		FScopedTransaction Transaction(LOCTEXT("SetStageInputMask", "Set Material Designer Mask Source"));
+		Stage->Modify();
+
+		InputExpression = UDMMaterialStageInputExpression::ChangeStageInput_Expression(
+			Stage,
+			UDMMaterialStageExpressionTextureSample::StaticClass(),
+			UDMMaterialStageThroughputLayerBlend::InputMaskSource,
+			FDMMaterialStageConnectorChannel::WHOLE_CHANNEL,
+			0,
+			FDMMaterialStageConnectorChannel::WHOLE_CHANNEL
+		);
+	}
+	else
+	{
+		ensureMsgf(false, TEXT("Invalid stage type (%s)"), *StageSource->GetClass()->GetName());
+	}
+
+	if (!InputExpression)
+	{
+		return;
+	}
+
+	UDMMaterialSubStage* SubStage = InputExpression->GetSubStage();
+
+	UDMMaterialStageInputValue::ChangeStageInput_NewLocalValue(
+		SubStage,
+		/* Texture input */ 0,
+		FDMMaterialStageConnectorChannel::WHOLE_CHANNEL,
+		UDMMaterialValueMediaStream::StaticClass(),
+		FDMMaterialStageConnectorChannel::WHOLE_CHANNEL
+	);
+}
+
+FDMMEdiaStreamStageSourceMenuExtender& FDMMEdiaStreamStageSourceMenuExtender::Get()
+{
+	static FDMMEdiaStreamStageSourceMenuExtender MenuExtension;
+	return MenuExtension;
+}
+
+#undef LOCTEXT_NAMESPACE
diff --git a/Engine/Plugins/Experimental/DynamicMaterialMediaStreamBridge/Source/DynamicMaterialMediaStreamBridgeEditor/Private/DMMediaStreamStageSourceMenuExtender.h b/Engine/Plugins/Experimental/DynamicMaterialMediaStreamBridge/Source/DynamicMaterialMediaStreamBridgeEditor/Private/DMMediaStreamStageSourceMenuExtender.h
new file mode 100644
--- /dev/null
+++ b/Engine/Plugins/Experimental/DynamicMaterialMediaStreamBridge/Source/DynamicMaterialMediaStreamBridgeEditor/Private/DMMediaStreamStageSourceMenuExtender.h
@@ -0,0 +1,21 @@
+// Copyright Epic Games, Inc. All Rights Reserved.
+
+#pragma once
+
+class UDMMenuContext;
+class UToolMenu;
+
+class FDMMEdiaStreamStageSourceMenuExtender
+{
+public:
+	static FDMMEdiaStreamStageSourceMenuExtender& Get();
+
+	void Integrate();
+
+private:
+	bool bIntegrated = false;
+
+	void ExtendMenu(UToolMenu* InMenu);
+
+	void ChangeSourceToMediaStreamFromContext(UDMMenuContext* InMenuContext);
+};
diff --git a/Engine/Plugins/Experimental/DynamicMaterialMediaStreamBridge/Source/DynamicMaterialMediaStreamBridgeEditor/Private/DynamicMaterialMediaStreamBridgeEditorModule.cpp b/Engine/Plugins/Experimental/DynamicMaterialMediaStreamBridge/Source/DynamicMaterialMediaStreamBridgeEditor/Private/DynamicMaterialMediaStreamBridgeEditorModule.cpp
new file mode 100644
--- /dev/null
+++ b/Engine/Plugins/Experimental/DynamicMaterialMediaStreamBridge/Source/DynamicMaterialMediaStreamBridgeEditor/Private/DynamicMaterialMediaStreamBridgeEditorModule.cpp
@@ -0,0 +1,19 @@
+// Copyright Epic Games, Inc. All Rights Reserved.
+
+#include "DynamicMaterialMediaStreamBridgeEditorModule.h"
+
+#include "DMMaterialValueMediaStream.h"
+#include "DMMaterialValueMediaStreamPropertyRowGenerator.h"
+#include "DMMediaStreamStageSourceMenuExtender.h"
+#include "IDynamicMaterialEditorModule.h"
+#include "Modules/ModuleManager.h"
+
+void FDynamicMaterialMediaStreamBridgeEditorModule::StartupModule()
+{
+	IDynamicMaterialEditorModule& DMEditorModule = IDynamicMaterialEditorModule::Get();
+	DMEditorModule.RegisterComponentPropertyRowGeneratorDelegate<UDMMaterialValueMediaStream, FDMMaterialValueMediaStreamPropertyRowGenerator>();
+
+	FDMMEdiaStreamStageSourceMenuExtender::Get().Integrate();
+}
+
+IMPLEMENT_MODULE(FDynamicMaterialMediaStreamBridgeEditorModule, DynamicMaterialMediaStreamBridgeEditor)
diff --git a/Engine/Plugins/Experimental/DynamicMaterialMediaStreamBridge/Source/DynamicMaterialMediaStreamBridgeEditor/Private/DynamicMaterialMediaStreamBridgeEditorModule.h b/Engine/Plugins/Experimental/DynamicMaterialMediaStreamBridge/Source/DynamicMaterialMediaStreamBridgeEditor/Private/DynamicMaterialMediaStreamBridgeEditorModule.h
new file mode 100644
--- /dev/null
+++ b/Engine/Plugins/Experimental/DynamicMaterialMediaStreamBridge/Source/DynamicMaterialMediaStreamBridgeEditor/Private/DynamicMaterialMediaStreamBridgeEditorModule.h
@@ -0,0 +1,16 @@
+// Copyright Epic Games, Inc. All Rights Reserved.
+
+#pragma once
+
+#include "Modules/ModuleInterface.h"
+
+/**
+ * Media Stream Material Designer Bridge Editor - Integrates the Media Stream plugin with the Material Designer.
+ */
+class FDynamicMaterialMediaStreamBridgeEditorModule : public IModuleInterface
+{
+public:
+	//~ Begin IModuleInterface
+	virtual void StartupModule() override;
+	//~ End IModuleInterface
+};
