PSD Importer: Full mask support
- Added mask support to full materials, material designer quads, instances and materials (as well as normal material quads).
- Fixed stages generating extremely large lists of used nodes.
- Fixes setting parameter names for vectors on the right click menu.

#virtualized

diff --git a/Engine/Plugins/Experimental/DynamicMaterial/Source/DynamicMaterialEditor/Private/Components/DMMaterialStageThroughput.cpp b/Engine/Plugins/Experimental/DynamicMaterial/Source/DynamicMaterialEditor/Private/Components/DMMaterialStageThroughput.cpp
--- a/Engine/Plugins/Experimental/DynamicMaterial/Source/DynamicMaterialEditor/Private/Components/DMMaterialStageThroughput.cpp
+++ b/Engine/Plugins/Experimental/DynamicMaterial/Source/DynamicMaterialEditor/Private/Components/DMMaterialStageThroughput.cpp
@@ -708,6 +708,7 @@ int32 UDMMaterialStageThroughput::ResolveInputChannel(const TSharedRef<FDMMateri
 	int32 InnateMaskOutput = INDEX_NONE;
 	int32 NodeOutputIndex = INDEX_NONE;
 	int32 OutputChannelOverride = INDEX_NONE;
+	UMaterialExpression* LastExpression = nullptr;
 
 	if (OutChannel.SourceIndex == FDMMaterialStageConnectorChannel::PREVIOUS_STAGE)
 	{
@@ -718,8 +719,9 @@ int32 UDMMaterialStageThroughput::ResolveInputChannel(const TSharedRef<FDMMateri
 		{
 			PreviousLayer->GenerateExpressions(InBuildState);
 
-			OutExpressions.Append(InBuildState->GetLayerExpressions(PreviousLayer));
-			check(!OutExpressions.IsEmpty());
+			const TArray<UMaterialExpression*>& PreviousLayerExpressions = InBuildState->GetLayerExpressions(PreviousLayer);
+			check(!PreviousLayerExpressions.IsEmpty());
+			LastExpression = PreviousLayerExpressions.Last();
 
 			UDMMaterialStage* PreviousStage = PreviousLayer->GetLastValidStage(EDMMaterialLayerStage::All);
 			check(PreviousStage);
@@ -738,6 +740,7 @@ int32 UDMMaterialStageThroughput::ResolveInputChannel(const TSharedRef<FDMMateri
 			UMaterialExpressionConstant3Vector* Black = InBuildState->GetBuildUtils().CreateExpression<UMaterialExpressionConstant3Vector>(UE_DM_NodeComment_Default);
 			Black->Constant = FLinearColor(0, 0, 0, 1);
 
+			LastExpression = Black;
 			OutExpressions.Add(Black);
 
 			NodeOutputIndex = 0;
@@ -756,6 +759,7 @@ int32 UDMMaterialStageThroughput::ResolveInputChannel(const TSharedRef<FDMMateri
 
 		OutExpressions.Append(InBuildState->GetStageSourceExpressions(InputValue));
 		check(!OutExpressions.IsEmpty());
+		LastExpression = OutExpressions.Last();
 
 		InputSourceOutputConnectors = &(InputValue->GetOutputConnectors());
 		check(InputSourceOutputConnectors->IsValidIndex(OutChannel.OutputIndex));
@@ -803,7 +807,7 @@ int32 UDMMaterialStageThroughput::ResolveInputChannel(const TSharedRef<FDMMateri
 		&& bOutputIsFloatType && OutputFloatCount > 0 && InputFloatCount > 0 && OutputFloatCount < InputFloatCount)
 	{
 		UMaterialExpressionAppendVector* Append = InBuildState->GetBuildUtils().CreateExpression<UMaterialExpressionAppendVector>(UE_DM_NodeComment_Default);
-		Append->A.Expression = OutExpressions.Last();
+		Append->A.Expression = LastExpression;
 		Append->A.OutputIndex = NodeOutputIndex;
 
 		switch (InputFloatCount - OutputFloatCount)
diff --git a/Engine/Plugins/Experimental/DynamicMaterial/Source/DynamicMaterialEditor/Private/Components/DMMaterialStageThroughputLayerBlend.cpp b/Engine/Plugins/Experimental/DynamicMaterial/Source/DynamicMaterialEditor/Private/Components/DMMaterialStageThroughputLayerBlend.cpp
--- a/Engine/Plugins/Experimental/DynamicMaterial/Source/DynamicMaterialEditor/Private/Components/DMMaterialStageThroughputLayerBlend.cpp
+++ b/Engine/Plugins/Experimental/DynamicMaterial/Source/DynamicMaterialEditor/Private/Components/DMMaterialStageThroughputLayerBlend.cpp
@@ -245,7 +245,11 @@ int32 UDMMaterialStageThroughputLayerBlend::ResolveInput(const TSharedRef<FDMMat
 			if (UDMMaterialLayerObject* PreviousLayer = Layer->GetPreviousLayer(StageProperty, EDMMaterialLayerStage::Base))
 			{
 				PreviousLayer->GenerateExpressions(InBuildState);
-				OutExpressions.Append(InBuildState->GetLayerExpressions(PreviousLayer));
+
+				const TArray<UMaterialExpression*>& PreviousLayerExpressions = InBuildState->GetLayerExpressions(PreviousLayer);
+				check(!PreviousLayerExpressions.IsEmpty());
+
+				OutExpressions.Add(PreviousLayerExpressions.Last());
 				return 0;
 			}
 
@@ -271,7 +275,9 @@ int32 UDMMaterialStageThroughputLayerBlend::ResolveInput(const TSharedRef<FDMMat
 			check(BaseStage);
  
 			BaseStage->GenerateExpressions(InBuildState);
-			OutExpressions.Append(InBuildState->GetStageExpressions(BaseStage));
+
+			const TArray<UMaterialExpression*>& BaseStageExpressions = InBuildState->GetStageExpressions(BaseStage);
+			OutExpressions.Add(BaseStageExpressions.Last());
  
 			return OutChannel.OutputIndex;
 		}
diff --git a/Engine/Plugins/Experimental/DynamicMaterial/Source/DynamicMaterialEditor/Private/DetailsPanel/DMValueDetailsRowExtensions.cpp b/Engine/Plugins/Experimental/DynamicMaterial/Source/DynamicMaterialEditor/Private/DetailsPanel/DMValueDetailsRowExtensions.cpp
--- a/Engine/Plugins/Experimental/DynamicMaterial/Source/DynamicMaterialEditor/Private/DetailsPanel/DMValueDetailsRowExtensions.cpp
+++ b/Engine/Plugins/Experimental/DynamicMaterial/Source/DynamicMaterialEditor/Private/DetailsPanel/DMValueDetailsRowExtensions.cpp
@@ -228,7 +228,7 @@ void FDMValueDetailsRowExtensions::FillPropertyRightClickMenu_TextureUV(FToolMen
 						.ClearKeyboardFocusOnCommit(true)
 						.MinDesiredWidth(100.f)
 						.OnTextCommitted_Lambda(
-							[TextureUVWeak, InPropertyName](const FText& InText, ETextCommit::Type InCommitType)
+							[TextureUVWeak, InPropertyName, InComponent](const FText& InText, ETextCommit::Type InCommitType)
 							{
 								if (InCommitType != ETextCommit::OnEnter)
 								{
@@ -249,7 +249,7 @@ void FDMValueDetailsRowExtensions::FillPropertyRightClickMenu_TextureUV(FToolMen
 									return;
 								}
 
-								SetTextureUVParameterName(TextureUVWeak, InPropertyName, ComponentX, FName(*InText.ToString()));
+								SetTextureUVParameterName(TextureUVWeak, InPropertyName, InComponent, FName(*InText.ToString()));
 							});
 				}));
 
