Material Designer: Create Button made optional
- There is now a 'Show Create Button' option in the editor settings for the Material Designer. This can be used to disable the details view buttons.

diff --git a/Engine/Plugins/Experimental/DynamicMaterial/Source/DynamicMaterialEditor/Private/DetailsPanel/Widgets/SDMDetailsPanelMaterialInterfaceWidget.cpp b/Engine/Plugins/Experimental/DynamicMaterial/Source/DynamicMaterialEditor/Private/DetailsPanel/Widgets/SDMDetailsPanelMaterialInterfaceWidget.cpp
--- a/Engine/Plugins/Experimental/DynamicMaterial/Source/DynamicMaterialEditor/Private/DetailsPanel/Widgets/SDMDetailsPanelMaterialInterfaceWidget.cpp
+++ b/Engine/Plugins/Experimental/DynamicMaterial/Source/DynamicMaterialEditor/Private/DetailsPanel/Widgets/SDMDetailsPanelMaterialInterfaceWidget.cpp
@@ -4,6 +4,7 @@
 
 #include "Components/PrimitiveComponent.h"
 #include "DMObjectMaterialProperty.h"
+#include "DynamicMaterialEditorSettings.h"
 #include "DynamicMaterialEditorStyle.h"
 #include "IDynamicMaterialEditorModule.h"
 #include "Material/DynamicMaterialInstance.h"
@@ -38,10 +39,7 @@ void SDMDetailsPanelMaterialInterfaceWidget::Construct(const FArguments& InArgs,
 		? ToRawPtr(ObjectProperty->PropertyClass)
 		: ToRawPtr(UMaterialInterface::StaticClass());
 
-	// @formatter:off
-	ChildSlot
-	[
-		SNew(SVerticalBox)
+	TSharedRef<SVerticalBox> Container = SNew(SVerticalBox)
 		+ SVerticalBox::Slot()
 		.AutoHeight()
 		.HAlign(EHorizontalAlignment::HAlign_Left)
@@ -58,25 +56,34 @@ void SDMDetailsPanelMaterialInterfaceWidget::Construct(const FArguments& InArgs,
 			.EnableContentPicker(true)
 			.PropertyHandle(PropertyHandle)
 			.ThumbnailPool(InArgs._ThumbnailPool)
-		]
-		+ SVerticalBox::Slot()
-		.HAlign(EHorizontalAlignment::HAlign_Left)
-		.VAlign(EVerticalAlignment::VAlign_Center)
-		.Padding(10.f, 5.f, 10.f, 5.f)
-		.AutoHeight()
-		[
-			SNew(SButton)
-			.OnClicked(this, &SDMDetailsPanelMaterialInterfaceWidget::OnButtonClicked)
-			.IsEnabled(InPropertyHandle, &IPropertyHandle::IsEditable)
-			.Content()
+		];
+
+	const UDynamicMaterialEditorSettings* Settings = UDynamicMaterialEditorSettings::Get();
+
+	if (Settings && Settings->bAddDetailsPanelButton)
+	{
+		Container->AddSlot()
+			.HAlign(EHorizontalAlignment::HAlign_Left)
+			.VAlign(EVerticalAlignment::VAlign_Center)
+			.Padding(10.f, 5.f, 10.f, 5.f)
+			.AutoHeight()
 			[
-				SNew(STextBlock)
-				.TextStyle(FDynamicMaterialEditorStyle::Get(), "RegularFont")
-				.Text(this, &SDMDetailsPanelMaterialInterfaceWidget::GetButtonText)
-			]
-		]
+				SNew(SButton)
+				.OnClicked(this, &SDMDetailsPanelMaterialInterfaceWidget::OnButtonClicked)
+				.IsEnabled(InPropertyHandle, &IPropertyHandle::IsEditable)
+				.Content()
+				[
+					SNew(STextBlock)
+					.TextStyle(FDynamicMaterialEditorStyle::Get(), "RegularFont")
+					.Text(this, &SDMDetailsPanelMaterialInterfaceWidget::GetButtonText)
+				]
+			];
+	}
+
+	ChildSlot
+	[
+		Container
 	];
-	// @formatter:on
 }
 
 UObject* SDMDetailsPanelMaterialInterfaceWidget::GetAsset() const
diff --git a/Engine/Plugins/Experimental/DynamicMaterial/Source/DynamicMaterialEditor/Private/DetailsPanel/Widgets/SDMDetailsPanelTabSpawner.cpp b/Engine/Plugins/Experimental/DynamicMaterial/Source/DynamicMaterialEditor/Private/DetailsPanel/Widgets/SDMDetailsPanelTabSpawner.cpp
--- a/Engine/Plugins/Experimental/DynamicMaterial/Source/DynamicMaterialEditor/Private/DetailsPanel/Widgets/SDMDetailsPanelTabSpawner.cpp
+++ b/Engine/Plugins/Experimental/DynamicMaterial/Source/DynamicMaterialEditor/Private/DetailsPanel/Widgets/SDMDetailsPanelTabSpawner.cpp
@@ -4,6 +4,7 @@
 
 #include "AssetToolsModule.h"
 #include "DMObjectMaterialProperty.h"
+#include "DynamicMaterialEditorSettings.h"
 #include "DynamicMaterialEditorStyle.h"
 #include "IAssetTools.h"
 #include "IDynamicMaterialEditorModule.h"
@@ -30,10 +31,7 @@ void SDMDetailsPanelTabSpawner::Construct(const FArguments& InArgs, const TShare
 	UObject* Value = nullptr;
 	InPropertyHandle->GetValue(Value);
 
-	// @formatter:off
-	ChildSlot
-	[
-		SNew(SVerticalBox)
+	TSharedRef<SVerticalBox> Container = SNew(SVerticalBox)
 		+ SVerticalBox::Slot()
 		.AutoHeight()
 		.HAlign(EHorizontalAlignment::HAlign_Left)
@@ -50,24 +48,33 @@ void SDMDetailsPanelTabSpawner::Construct(const FArguments& InArgs, const TShare
 			.EnableContentPicker(true)
 			.ObjectPath(this, &SDMDetailsPanelTabSpawner::GetEditorPath)
 			.OnObjectChanged(this, &SDMDetailsPanelTabSpawner::OnEditorChanged)
-		]
-		+ SVerticalBox::Slot()
-		.HAlign(EHorizontalAlignment::HAlign_Left)
-		.VAlign(EVerticalAlignment::VAlign_Center)
-		.Padding(10.f, 5.f, 10.f, 5.f)
-		.AutoHeight()
-		[
-			SNew(SButton)
-			.OnClicked(this, &SDMDetailsPanelTabSpawner::OnButtonClicked)
-			.Content()
+		];
+
+	const UDynamicMaterialEditorSettings* Settings = UDynamicMaterialEditorSettings::Get();
+
+	if (Settings && Settings->bAddDetailsPanelButton)
+	{
+		Container->AddSlot()
+			.HAlign(EHorizontalAlignment::HAlign_Left)
+			.VAlign(EVerticalAlignment::VAlign_Center)
+			.Padding(10.f, 5.f, 10.f, 5.f)
+			.AutoHeight()
 			[
-				SNew(STextBlock)
-				.TextStyle(FDynamicMaterialEditorStyle::Get(), "RegularFont")
-				.Text(this, &SDMDetailsPanelTabSpawner::GetButtonText)
-			]
-		]
+				SNew(SButton)
+				.OnClicked(this, &SDMDetailsPanelTabSpawner::OnButtonClicked)
+				.Content()
+				[
+					SNew(STextBlock)
+					.TextStyle(FDynamicMaterialEditorStyle::Get(), "RegularFont")
+					.Text(this, &SDMDetailsPanelTabSpawner::GetButtonText)
+				]
+			];
+	}
+
+	ChildSlot
+	[
+		Container
 	];
-	// @formatter:on
 }
 
 UDynamicMaterialModelBase* SDMDetailsPanelTabSpawner::GetMaterialModelBase() const
diff --git a/Engine/Plugins/Experimental/DynamicMaterial/Source/DynamicMaterialEditor/Private/DynamicMaterialEditorModule.cpp b/Engine/Plugins/Experimental/DynamicMaterial/Source/DynamicMaterialEditor/Private/DynamicMaterialEditorModule.cpp
--- a/Engine/Plugins/Experimental/DynamicMaterial/Source/DynamicMaterialEditor/Private/DynamicMaterialEditorModule.cpp
+++ b/Engine/Plugins/Experimental/DynamicMaterial/Source/DynamicMaterialEditor/Private/DynamicMaterialEditorModule.cpp
@@ -60,6 +60,13 @@ namespace UE::DynamicMaterialEditor::Private
 		IDetailLayoutBuilder& InDetailBuilder,
 		TArray<TSharedPtr<SWidget>>& OutExtensions)
 	{
+		const UDynamicMaterialEditorSettings* Settings = UDynamicMaterialEditorSettings::Get();
+
+		if (!Settings || !Settings->bAddDetailsPanelButton)
+		{
+			return;
+		}
+
 		if (UPrimitiveComponent* PrimitiveComponent = Cast<UPrimitiveComponent>(InCurrentComponent))
 		{
 			OutExtensions.Add(SNew(SDMMaterialListExtensionWidget,
diff --git a/Engine/Plugins/Experimental/DynamicMaterial/Source/DynamicMaterialEditor/Private/DynamicMaterialEditorSettings.cpp b/Engine/Plugins/Experimental/DynamicMaterial/Source/DynamicMaterialEditor/Private/DynamicMaterialEditorSettings.cpp
--- a/Engine/Plugins/Experimental/DynamicMaterial/Source/DynamicMaterialEditor/Private/DynamicMaterialEditorSettings.cpp
+++ b/Engine/Plugins/Experimental/DynamicMaterial/Source/DynamicMaterialEditor/Private/DynamicMaterialEditorSettings.cpp
@@ -133,6 +133,7 @@ UDynamicMaterialEditorSettings::UDynamicMaterialEditorSettings()
 	SectionName = TEXT("Material Designer");
 
 	bFollowSelection = true;
+	bAddDetailsPanelButton = false;
 	bUseLinearColorForVectors = true;
 
 	ResetAllLayoutSettings();
diff --git a/Engine/Plugins/Experimental/DynamicMaterial/Source/DynamicMaterialEditor/Private/DynamicMaterialEditorSettings.h b/Engine/Plugins/Experimental/DynamicMaterial/Source/DynamicMaterialEditor/Private/DynamicMaterialEditorSettings.h
--- a/Engine/Plugins/Experimental/DynamicMaterial/Source/DynamicMaterialEditor/Private/DynamicMaterialEditorSettings.h
+++ b/Engine/Plugins/Experimental/DynamicMaterial/Source/DynamicMaterialEditor/Private/DynamicMaterialEditorSettings.h
@@ -185,6 +185,13 @@ public:
 	UPROPERTY(Config, EditAnywhere, BlueprintReadWrite, Category = "Editor")
 	bool bFollowSelection;
 
+	/**
+	 * If true, a "Create with Material Designer" button will show on properties and material lists. The edit
+	 * buttons will always appear.
+	 */
+	UPROPERTY(Config, EditAnywhere, BlueprintReadWrite, Category = "Editor")
+	bool bAddDetailsPanelButton;
+
 	UPROPERTY(Config, EditAnywhere, BlueprintReadWrite, Category = "Layout")
 	EDMMaterialEditorLayout Layout;
 
