Material Designer: Renamed menu option
- The "Reset all to defaults" menu option has been renamed to "Reset UI Settings".

diff --git a/Engine/Plugins/Experimental/DynamicMaterial/Source/DynamicMaterialEditor/Private/UI/Menus/DMToolBarMenus.cpp b/Engine/Plugins/Experimental/DynamicMaterial/Source/DynamicMaterialEditor/Private/UI/Menus/DMToolBarMenus.cpp
--- a/Engine/Plugins/Experimental/DynamicMaterial/Source/DynamicMaterialEditor/Private/UI/Menus/DMToolBarMenus.cpp
+++ b/Engine/Plugins/Experimental/DynamicMaterial/Source/DynamicMaterialEditor/Private/UI/Menus/DMToolBarMenus.cpp
@@ -265,8 +265,8 @@ void FDMToolBarMenus::AddAdvancedSection(UToolMenu* InMenu)
 
 	NewSection.AddMenuEntry(
 		NAME_None,
-		LOCTEXT("ResetAllSettingsToDefaults", "Reset All To Defaults"),
-		LOCTEXT("ResetAllSettingsToDefaultsTooltip", "Resets all the Material Designer settings to their default values."),
+		LOCTEXT("ResetAllSettingsToDefaults", "Reset UI Settings"),
+		LOCTEXT("ResetAllSettingsToDefaultsTooltip", "Resets all the Material Designer UI settings to their default values."),
 		FSlateIcon(),
 		FUIAction(FExecuteAction::CreateUObject(
 			UDynamicMaterialEditorSettings::Get(), 
