Material Designer: Fixed layer blend tooltip.
- Changed the layer blend selection tooltip to include the name of the blend mode.

diff --git a/Engine/Plugins/Experimental/DynamicMaterial/Source/DynamicMaterialEditor/Private/UI/Widgets/Editor/SlotEditor/SDMMaterialLayerBlendMode.cpp b/Engine/Plugins/Experimental/DynamicMaterial/Source/DynamicMaterialEditor/Private/UI/Widgets/Editor/SlotEditor/SDMMaterialLayerBlendMode.cpp
--- a/Engine/Plugins/Experimental/DynamicMaterial/Source/DynamicMaterialEditor/Private/UI/Widgets/Editor/SlotEditor/SDMMaterialLayerBlendMode.cpp
+++ b/Engine/Plugins/Experimental/DynamicMaterial/Source/DynamicMaterialEditor/Private/UI/Widgets/Editor/SlotEditor/SDMMaterialLayerBlendMode.cpp
@@ -371,10 +371,12 @@ void SDMMaterialLayerBlendMode::MakeSourceBlendMenu(UToolMenu* InToolMenu)
 
 	auto AddClassToSection = [&SourceBlendWidget](FToolMenuSection& InSection, const FText& InName, UClass* InBlendClass)
 	{
+		const FText ToolTipFormat = LOCTEXT("ChangeSourceBlendTooltip", "Change the source of this stage to a {0} Material Blend.");
+
 		InSection.AddMenuEntry(
 			NAME_None,
 			InName,
-			LOCTEXT("ChangeSourceBlendTooltip", "Change the source of this stage to a Material Blend."),
+			FText::Format(ToolTipFormat, InName),
 			FSlateIcon(),
 			FUIAction(
 				FExecuteAction::CreateSP(SourceBlendWidget.ToSharedRef(), &SDMMaterialLayerBlendMode::OnBlendModeSelected, InBlendClass),
