Material Designer: Blend descriptions
- Added descriptions / tooltips to the different blend modes.

diff --git a/Engine/Plugins/Experimental/DynamicMaterial/Source/DynamicMaterialEditor/Private/Components/DMMaterialStageBlend.cpp b/Engine/Plugins/Experimental/DynamicMaterial/Source/DynamicMaterialEditor/Private/Components/DMMaterialStageBlend.cpp
--- a/Engine/Plugins/Experimental/DynamicMaterial/Source/DynamicMaterialEditor/Private/Components/DMMaterialStageBlend.cpp
+++ b/Engine/Plugins/Experimental/DynamicMaterial/Source/DynamicMaterialEditor/Private/Components/DMMaterialStageBlend.cpp
@@ -24,12 +24,13 @@
 TArray<TStrongObjectPtr<UClass>> UDMMaterialStageBlend::Blends = {};
  
 UDMMaterialStageBlend::UDMMaterialStageBlend()
-	: UDMMaterialStageBlend(FText::GetEmpty())
+	: UDMMaterialStageBlend(FText::GetEmpty(), FText::GetEmpty())
 {
 }
  
-UDMMaterialStageBlend::UDMMaterialStageBlend(const FText& InName)
+UDMMaterialStageBlend::UDMMaterialStageBlend(const FText& InName, const FText& InDescription)
 	: UDMMaterialStageThroughput(InName)
+	, BlendDescription(InDescription)
 {
 	bInputRequired = true;
 	bAllowNestedInputs = true;
@@ -79,6 +80,11 @@ void UDMMaterialStageBlend::SetBaseChannelOverride(EAvaColorChannel InMaskChanne
 	Update(this, EDMUpdateType::Structure);
 }
 
+const FText& UDMMaterialStageBlend::GetBlendDescription() const
+{
+	return BlendDescription;
+}
+
 int32 UDMMaterialStageBlend::GetDefaultBaseChannelOverrideOutputIndex() const
 {
 	UDMMaterialStageInput* StageInputB = GetInputB();
diff --git a/Engine/Plugins/Experimental/DynamicMaterial/Source/DynamicMaterialEditor/Private/Components/DMMaterialStageBlendFunction.cpp b/Engine/Plugins/Experimental/DynamicMaterial/Source/DynamicMaterialEditor/Private/Components/DMMaterialStageBlendFunction.cpp
--- a/Engine/Plugins/Experimental/DynamicMaterial/Source/DynamicMaterialEditor/Private/Components/DMMaterialStageBlendFunction.cpp
+++ b/Engine/Plugins/Experimental/DynamicMaterial/Source/DynamicMaterialEditor/Private/Components/DMMaterialStageBlendFunction.cpp
@@ -13,18 +13,18 @@
 #define LOCTEXT_NAMESPACE "DMMaterialStageBlendFunction"
 
 UDMMaterialStageBlendFunction::UDMMaterialStageBlendFunction()
-	: UDMMaterialStageBlendFunction(LOCTEXT("BlendFunction", "Blend Function"), nullptr)
+	: UDMMaterialStageBlendFunction(LOCTEXT("BlendFunction", "Blend Function"), FText::GetEmpty(), nullptr)
 {
 }
 
-UDMMaterialStageBlendFunction::UDMMaterialStageBlendFunction(const FText& InName, UMaterialFunctionInterface* InMaterialFunction)
-	: UDMMaterialStageBlend(InName)
+UDMMaterialStageBlendFunction::UDMMaterialStageBlendFunction(const FText& InName, const FText& InDescription, UMaterialFunctionInterface* InMaterialFunction)
+	: UDMMaterialStageBlend(InName, InDescription)
 	, MaterialFunction(InMaterialFunction)
 {
 }
 
-UDMMaterialStageBlendFunction::UDMMaterialStageBlendFunction(const FText& InName, const FName& InFunctionName, const FString& InFunctionPath)
-	: UDMMaterialStageBlend(InName)
+UDMMaterialStageBlendFunction::UDMMaterialStageBlendFunction(const FText& InName, const FText& InDescription, const FName& InFunctionName, const FString& InFunctionPath)
+	: UDMMaterialStageBlend(InName, InDescription)
 	, MaterialFunction(FDMMaterialFunctionLibrary::Get().GetFunction(InFunctionName, InFunctionPath))
 {
 	check(MaterialFunction);
diff --git a/Engine/Plugins/Experimental/DynamicMaterial/Source/DynamicMaterialEditor/Private/Components/MaterialStageBlends/DMMSBAdd.cpp b/Engine/Plugins/Experimental/DynamicMaterial/Source/DynamicMaterialEditor/Private/Components/MaterialStageBlends/DMMSBAdd.cpp
--- a/Engine/Plugins/Experimental/DynamicMaterial/Source/DynamicMaterialEditor/Private/Components/MaterialStageBlends/DMMSBAdd.cpp
+++ b/Engine/Plugins/Experimental/DynamicMaterial/Source/DynamicMaterialEditor/Private/Components/MaterialStageBlends/DMMSBAdd.cpp
@@ -6,7 +6,8 @@
 
 UDMMaterialStageBlendAdd::UDMMaterialStageBlendAdd()
 	: UDMMaterialStageBlendFunction(
-		LOCTEXT("BlendAdd", "Add"), 
+		LOCTEXT("BlendAdd", "Add"),
+		LOCTEXT("BlendAddDescription", "Add increases the brightness of the base layer based on the blend layer on a per channel basis."),
 		"DM_Blend_Add", 
 		TEXT("/Script/Engine.MaterialFunction'/DynamicMaterial/MaterialFunctions/Blends/MF_DM_Blend_Add.MF_DM_Blend_Add'")
 	)
diff --git a/Engine/Plugins/Experimental/DynamicMaterial/Source/DynamicMaterialEditor/Private/Components/MaterialStageBlends/DMMSBColor.cpp b/Engine/Plugins/Experimental/DynamicMaterial/Source/DynamicMaterialEditor/Private/Components/MaterialStageBlends/DMMSBColor.cpp
--- a/Engine/Plugins/Experimental/DynamicMaterial/Source/DynamicMaterialEditor/Private/Components/MaterialStageBlends/DMMSBColor.cpp
+++ b/Engine/Plugins/Experimental/DynamicMaterial/Source/DynamicMaterialEditor/Private/Components/MaterialStageBlends/DMMSBColor.cpp
@@ -6,7 +6,8 @@
 
 UDMMaterialStageBlendColor::UDMMaterialStageBlendColor()
 	: UDMMaterialStageBlendFunction(
-		LOCTEXT("BlendColor", "Color"), 
+		LOCTEXT("BlendColor", "Color"),
+		LOCTEXT("BlendColorDescription", "Color combines the blend layer's hue and saturation with the base layer's luminosity."),
 		"DM_Blend_Color", 
 		TEXT("/Script/Engine.MaterialFunction'/DynamicMaterial/MaterialFunctions/Blends/MF_DM_Blend_Color.MF_DM_Blend_Color'")
 	)
diff --git a/Engine/Plugins/Experimental/DynamicMaterial/Source/DynamicMaterialEditor/Private/Components/MaterialStageBlends/DMMSBColorBurn.cpp b/Engine/Plugins/Experimental/DynamicMaterial/Source/DynamicMaterialEditor/Private/Components/MaterialStageBlends/DMMSBColorBurn.cpp
--- a/Engine/Plugins/Experimental/DynamicMaterial/Source/DynamicMaterialEditor/Private/Components/MaterialStageBlends/DMMSBColorBurn.cpp
+++ b/Engine/Plugins/Experimental/DynamicMaterial/Source/DynamicMaterialEditor/Private/Components/MaterialStageBlends/DMMSBColorBurn.cpp
@@ -7,6 +7,7 @@
 UDMMaterialStageBlendColorBurn::UDMMaterialStageBlendColorBurn()
 	: UDMMaterialStageBlendFunction(
 		LOCTEXT("BlendColorBurn", "ColorBurn"),
+		LOCTEXT("BlendColorBurnDescription", "Color Burn is similar to Multiply, but increases the contrast between the base and the blend layers, resulting in more highly saturated mid-tones and reduced highlights."),
 		"DM_Blend_ColorBurn",
 		TEXT("/Script/Engine.MaterialFunction'/DynamicMaterial/MaterialFunctions/Blends/MF_DM_Blend_ColorBurn.MF_DM_Blend_ColorBurn'")
 	)
diff --git a/Engine/Plugins/Experimental/DynamicMaterial/Source/DynamicMaterialEditor/Private/Components/MaterialStageBlends/DMMSBColorDodge.cpp b/Engine/Plugins/Experimental/DynamicMaterial/Source/DynamicMaterialEditor/Private/Components/MaterialStageBlends/DMMSBColorDodge.cpp
--- a/Engine/Plugins/Experimental/DynamicMaterial/Source/DynamicMaterialEditor/Private/Components/MaterialStageBlends/DMMSBColorDodge.cpp
+++ b/Engine/Plugins/Experimental/DynamicMaterial/Source/DynamicMaterialEditor/Private/Components/MaterialStageBlends/DMMSBColorDodge.cpp
@@ -7,6 +7,7 @@
 UDMMaterialStageBlendColorDodge::UDMMaterialStageBlendColorDodge()
 	: UDMMaterialStageBlendFunction(
 		LOCTEXT("BlendColorDodge", "ColorDodge"),
+		LOCTEXT("BlendColorDodgeDescription", "Color Dodge is similar to Screen, but decreases contrast between the base and blend layer, resulting in saturated mid-tones and blown highlights."),
 		"DM_Blend_ColorDodge",
 		TEXT("/Script/Engine.MaterialFunction'/DynamicMaterial/MaterialFunctions/Blends/MF_DM_Blend_ColorDodge.MF_DM_Blend_ColorDodge'")
 	)
diff --git a/Engine/Plugins/Experimental/DynamicMaterial/Source/DynamicMaterialEditor/Private/Components/MaterialStageBlends/DMMSBContrastBase.cpp b/Engine/Plugins/Experimental/DynamicMaterial/Source/DynamicMaterialEditor/Private/Components/MaterialStageBlends/DMMSBContrastBase.cpp
--- a/Engine/Plugins/Experimental/DynamicMaterial/Source/DynamicMaterialEditor/Private/Components/MaterialStageBlends/DMMSBContrastBase.cpp
+++ b/Engine/Plugins/Experimental/DynamicMaterial/Source/DynamicMaterialEditor/Private/Components/MaterialStageBlends/DMMSBContrastBase.cpp
@@ -5,12 +5,12 @@
 #define LOCTEXT_NAMESPACE "DMMaterialStageBlendContrastBase"
 
 UDMMaterialStageBlendContrastBase::UDMMaterialStageBlendContrastBase()
-	: UDMMaterialStageBlendContrastBase(LOCTEXT("BlendContrastBase", "Contrast Base"))
+	: UDMMaterialStageBlendContrastBase(LOCTEXT("BlendContrastBase", "Contrast Base"), FText::GetEmpty())
 {
 }
 
-UDMMaterialStageBlendContrastBase::UDMMaterialStageBlendContrastBase(const FText& InName)
-	: UDMMaterialStageBlend(InName)
+UDMMaterialStageBlendContrastBase::UDMMaterialStageBlendContrastBase(const FText& InName, const FText& InDescription)
+	: UDMMaterialStageBlend(InName, InDescription)
 {
 }
 
diff --git a/Engine/Plugins/Experimental/DynamicMaterial/Source/DynamicMaterialEditor/Private/Components/MaterialStageBlends/DMMSBDarken.cpp b/Engine/Plugins/Experimental/DynamicMaterial/Source/DynamicMaterialEditor/Private/Components/MaterialStageBlends/DMMSBDarken.cpp
--- a/Engine/Plugins/Experimental/DynamicMaterial/Source/DynamicMaterialEditor/Private/Components/MaterialStageBlends/DMMSBDarken.cpp
+++ b/Engine/Plugins/Experimental/DynamicMaterial/Source/DynamicMaterialEditor/Private/Components/MaterialStageBlends/DMMSBDarken.cpp
@@ -7,6 +7,7 @@
 UDMMaterialStageBlendDarken::UDMMaterialStageBlendDarken()
 	: UDMMaterialStageBlendFunction(
 		LOCTEXT("BlendDarken", "Darken"),
+		LOCTEXT("BlendDarkenDescription", "Darken chooses the color with the lowest luminance between the base and blend layer."),
 		"DM_Blend_Darken",
 		TEXT("/Script/Engine.MaterialFunction'/DynamicMaterial/MaterialFunctions/Blends/MF_DM_Blend_Darken.MF_DM_Blend_Darken'")
 	)
diff --git a/Engine/Plugins/Experimental/DynamicMaterial/Source/DynamicMaterialEditor/Private/Components/MaterialStageBlends/DMMSBDarkenColor.cpp b/Engine/Plugins/Experimental/DynamicMaterial/Source/DynamicMaterialEditor/Private/Components/MaterialStageBlends/DMMSBDarkenColor.cpp
--- a/Engine/Plugins/Experimental/DynamicMaterial/Source/DynamicMaterialEditor/Private/Components/MaterialStageBlends/DMMSBDarkenColor.cpp
+++ b/Engine/Plugins/Experimental/DynamicMaterial/Source/DynamicMaterialEditor/Private/Components/MaterialStageBlends/DMMSBDarkenColor.cpp
@@ -7,6 +7,7 @@
 UDMMaterialStageBlendDarkenColor::UDMMaterialStageBlendDarkenColor()
 	: UDMMaterialStageBlendFunction(
 		LOCTEXT("BlendDarkenColor", "DarkenColor"),
+		LOCTEXT("BlendDarkenColorDescription", "Darken Color chooses the lowest value in each color channel between the base and blend layer."),
 		"DM_Blend_DarkenColor",
 		TEXT("/Script/Engine.MaterialFunction'/DynamicMaterial/MaterialFunctions/Blends/MF_DM_Blend_DarkenColor.MF_DM_Blend_DarkenColor'")
 	)
diff --git a/Engine/Plugins/Experimental/DynamicMaterial/Source/DynamicMaterialEditor/Private/Components/MaterialStageBlends/DMMSBDifference.cpp b/Engine/Plugins/Experimental/DynamicMaterial/Source/DynamicMaterialEditor/Private/Components/MaterialStageBlends/DMMSBDifference.cpp
--- a/Engine/Plugins/Experimental/DynamicMaterial/Source/DynamicMaterialEditor/Private/Components/MaterialStageBlends/DMMSBDifference.cpp
+++ b/Engine/Plugins/Experimental/DynamicMaterial/Source/DynamicMaterialEditor/Private/Components/MaterialStageBlends/DMMSBDifference.cpp
@@ -7,6 +7,7 @@
 UDMMaterialStageBlendDifference::UDMMaterialStageBlendDifference()
 	: UDMMaterialStageBlendFunction(
 		LOCTEXT("BlendDifference", "Difference"),
+		LOCTEXT("BlendDifferenceDescription", "Difference outputs the difference in color between the base and blend layers."),
 		"DM_Blend_Difference",
 		TEXT("/Script/Engine.MaterialFunction'/DynamicMaterial/MaterialFunctions/Blends/MF_DM_Blend_Difference.MF_DM_Blend_Difference'")
 	)
diff --git a/Engine/Plugins/Experimental/DynamicMaterial/Source/DynamicMaterialEditor/Private/Components/MaterialStageBlends/DMMSBDivide.cpp b/Engine/Plugins/Experimental/DynamicMaterial/Source/DynamicMaterialEditor/Private/Components/MaterialStageBlends/DMMSBDivide.cpp
--- a/Engine/Plugins/Experimental/DynamicMaterial/Source/DynamicMaterialEditor/Private/Components/MaterialStageBlends/DMMSBDivide.cpp
+++ b/Engine/Plugins/Experimental/DynamicMaterial/Source/DynamicMaterialEditor/Private/Components/MaterialStageBlends/DMMSBDivide.cpp
@@ -7,6 +7,7 @@
 UDMMaterialStageBlendDivide::UDMMaterialStageBlendDivide()
 	: UDMMaterialStageBlendFunction(
 		LOCTEXT("BlendDivide", "Divide"),
+		LOCTEXT("BlendDivideDescription", "Divide increases the brightness of the base layer as the blend layer gets darker. Black and white both produce no change."),
 		"DM_Blend_Divide",
 		TEXT("/Script/Engine.MaterialFunction'/DynamicMaterial/MaterialFunctions/Blends/MF_DM_Blend_Divide.MF_DM_Blend_Divide'")
 	)
diff --git a/Engine/Plugins/Experimental/DynamicMaterial/Source/DynamicMaterialEditor/Private/Components/MaterialStageBlends/DMMSBExclusion.cpp b/Engine/Plugins/Experimental/DynamicMaterial/Source/DynamicMaterialEditor/Private/Components/MaterialStageBlends/DMMSBExclusion.cpp
--- a/Engine/Plugins/Experimental/DynamicMaterial/Source/DynamicMaterialEditor/Private/Components/MaterialStageBlends/DMMSBExclusion.cpp
+++ b/Engine/Plugins/Experimental/DynamicMaterial/Source/DynamicMaterialEditor/Private/Components/MaterialStageBlends/DMMSBExclusion.cpp
@@ -7,6 +7,7 @@
 UDMMaterialStageBlendExclusion::UDMMaterialStageBlendExclusion()
 	: UDMMaterialStageBlendFunction(
 		LOCTEXT("BlendExclusion", "Exclusion"),
+		LOCTEXT("BlendExclusionDescription", "Exclusion inverts the base layer colors based on the strength of the blend layer."),
 		"DM_Blend_Exclusion",
 		TEXT("/Script/Engine.MaterialFunction'/DynamicMaterial/MaterialFunctions/Blends/MF_DM_Blend_Exclusion.MF_DM_Blend_Exclusion'")
 	)
diff --git a/Engine/Plugins/Experimental/DynamicMaterial/Source/DynamicMaterialEditor/Private/Components/MaterialStageBlends/DMMSBHardLight.cpp b/Engine/Plugins/Experimental/DynamicMaterial/Source/DynamicMaterialEditor/Private/Components/MaterialStageBlends/DMMSBHardLight.cpp
--- a/Engine/Plugins/Experimental/DynamicMaterial/Source/DynamicMaterialEditor/Private/Components/MaterialStageBlends/DMMSBHardLight.cpp
+++ b/Engine/Plugins/Experimental/DynamicMaterial/Source/DynamicMaterialEditor/Private/Components/MaterialStageBlends/DMMSBHardLight.cpp
@@ -7,6 +7,7 @@
 UDMMaterialStageBlendHardLight::UDMMaterialStageBlendHardLight()
 	: UDMMaterialStageBlendFunction(
 		LOCTEXT("BlendHardLight", "HardLight"),
+		LOCTEXT("BlendHardLightDescription", "Hard Light is a combination of Multiply and Screen using the brightness of the blend layer."),
 		"DM_Blend_HardLight",
 		TEXT("/Script/Engine.MaterialFunction'/DynamicMaterial/MaterialFunctions/Blends/MF_DM_Blend_HardLight.MF_DM_Blend_HardLight'")
 	)
diff --git a/Engine/Plugins/Experimental/DynamicMaterial/Source/DynamicMaterialEditor/Private/Components/MaterialStageBlends/DMMSBHardMix.cpp b/Engine/Plugins/Experimental/DynamicMaterial/Source/DynamicMaterialEditor/Private/Components/MaterialStageBlends/DMMSBHardMix.cpp
--- a/Engine/Plugins/Experimental/DynamicMaterial/Source/DynamicMaterialEditor/Private/Components/MaterialStageBlends/DMMSBHardMix.cpp
+++ b/Engine/Plugins/Experimental/DynamicMaterial/Source/DynamicMaterialEditor/Private/Components/MaterialStageBlends/DMMSBHardMix.cpp
@@ -7,6 +7,7 @@
 UDMMaterialStageBlendHardMix::UDMMaterialStageBlendHardMix()
 	: UDMMaterialStageBlendFunction(
 		LOCTEXT("BlendHardMix", "HardMix"),
+		LOCTEXT("BlendHardMixDescription", "Hard Mix adds the individual color channels together and rounds them, producing one of eight colors: black, red, green, blue, cyan, magenta, yellow, or white."),
 		"DM_Blend_HardMix",
 		TEXT("/Script/Engine.MaterialFunction'/DynamicMaterial/MaterialFunctions/Blends/MF_DM_Blend_HardMix.MF_DM_Blend_HardMix'")
 	)
diff --git a/Engine/Plugins/Experimental/DynamicMaterial/Source/DynamicMaterialEditor/Private/Components/MaterialStageBlends/DMMSBHue.cpp b/Engine/Plugins/Experimental/DynamicMaterial/Source/DynamicMaterialEditor/Private/Components/MaterialStageBlends/DMMSBHue.cpp
--- a/Engine/Plugins/Experimental/DynamicMaterial/Source/DynamicMaterialEditor/Private/Components/MaterialStageBlends/DMMSBHue.cpp
+++ b/Engine/Plugins/Experimental/DynamicMaterial/Source/DynamicMaterialEditor/Private/Components/MaterialStageBlends/DMMSBHue.cpp
@@ -6,7 +6,8 @@
 
 UDMMaterialStageBlendHue::UDMMaterialStageBlendHue()
 	: UDMMaterialStageBlendFunction(
-		LOCTEXT("BlendHue", "Hue"), 
+		LOCTEXT("BlendHue", "Hue"),
+		LOCTEXT("BlendHueDescription", "Hue combines the blend layer's hue with the base layer's luminosity and saturation."),
 		"DM_Blend_Hue", 
 		TEXT("/Script/Engine.MaterialFunction'/DynamicMaterial/MaterialFunctions/Blends/MF_DM_Blend_Hue.MF_DM_Blend_Hue'")
 	)
diff --git a/Engine/Plugins/Experimental/DynamicMaterial/Source/DynamicMaterialEditor/Private/Components/MaterialStageBlends/DMMSBLighten.cpp b/Engine/Plugins/Experimental/DynamicMaterial/Source/DynamicMaterialEditor/Private/Components/MaterialStageBlends/DMMSBLighten.cpp
--- a/Engine/Plugins/Experimental/DynamicMaterial/Source/DynamicMaterialEditor/Private/Components/MaterialStageBlends/DMMSBLighten.cpp
+++ b/Engine/Plugins/Experimental/DynamicMaterial/Source/DynamicMaterialEditor/Private/Components/MaterialStageBlends/DMMSBLighten.cpp
@@ -7,6 +7,7 @@
 UDMMaterialStageBlendLighten::UDMMaterialStageBlendLighten()
 	: UDMMaterialStageBlendFunction(
 		LOCTEXT("BlendLighten", "Lighten"),
+		LOCTEXT("BlendLightenDescription", "Lighten chooses the color with the highest luminance between the base and blend layer."),
 		"DM_Blend_Lighten",
 		TEXT("/Script/Engine.MaterialFunction'/DynamicMaterial/MaterialFunctions/Blends/MF_DM_Blend_Lighten.MF_DM_Blend_Lighten'")
 	)
diff --git a/Engine/Plugins/Experimental/DynamicMaterial/Source/DynamicMaterialEditor/Private/Components/MaterialStageBlends/DMMSBLightenColor.cpp b/Engine/Plugins/Experimental/DynamicMaterial/Source/DynamicMaterialEditor/Private/Components/MaterialStageBlends/DMMSBLightenColor.cpp
--- a/Engine/Plugins/Experimental/DynamicMaterial/Source/DynamicMaterialEditor/Private/Components/MaterialStageBlends/DMMSBLightenColor.cpp
+++ b/Engine/Plugins/Experimental/DynamicMaterial/Source/DynamicMaterialEditor/Private/Components/MaterialStageBlends/DMMSBLightenColor.cpp
@@ -7,6 +7,7 @@
 UDMMaterialStageBlendLightenColor::UDMMaterialStageBlendLightenColor()
 	: UDMMaterialStageBlendFunction(
 		LOCTEXT("BlendLightenColor", "LightenColor"),
+		LOCTEXT("BlendLightenColorDescription", "Lighten Color chooses the highest value in each color channel between the base and blend layer."),
 		"DM_Blend_LightenColor",
 		TEXT("/Script/Engine.MaterialFunction'/DynamicMaterial/MaterialFunctions/Blends/MF_DM_Blend_LightenColor.MF_DM_Blend_LightenColor'")
 	)
diff --git a/Engine/Plugins/Experimental/DynamicMaterial/Source/DynamicMaterialEditor/Private/Components/MaterialStageBlends/DMMSBLinearBurn.cpp b/Engine/Plugins/Experimental/DynamicMaterial/Source/DynamicMaterialEditor/Private/Components/MaterialStageBlends/DMMSBLinearBurn.cpp
--- a/Engine/Plugins/Experimental/DynamicMaterial/Source/DynamicMaterialEditor/Private/Components/MaterialStageBlends/DMMSBLinearBurn.cpp
+++ b/Engine/Plugins/Experimental/DynamicMaterial/Source/DynamicMaterialEditor/Private/Components/MaterialStageBlends/DMMSBLinearBurn.cpp
@@ -7,6 +7,7 @@
 UDMMaterialStageBlendLinearBurn::UDMMaterialStageBlendLinearBurn()
 	: UDMMaterialStageBlendFunction(
 		LOCTEXT("BlendLinearBurn", "LinearBurn"),
+		LOCTEXT("BlendLinearBurnDescription", "Linear Burn decreases the brightness of the base layer based on the value of the blend layer."),
 		"DM_Blend_LinearBurn",
 		TEXT("/Script/Engine.MaterialFunction'/DynamicMaterial/MaterialFunctions/Blends/MF_DM_Blend_LinearBurn.MF_DM_Blend_LinearBurn'")
 	)
diff --git a/Engine/Plugins/Experimental/DynamicMaterial/Source/DynamicMaterialEditor/Private/Components/MaterialStageBlends/DMMSBLinearDodge.cpp b/Engine/Plugins/Experimental/DynamicMaterial/Source/DynamicMaterialEditor/Private/Components/MaterialStageBlends/DMMSBLinearDodge.cpp
--- a/Engine/Plugins/Experimental/DynamicMaterial/Source/DynamicMaterialEditor/Private/Components/MaterialStageBlends/DMMSBLinearDodge.cpp
+++ b/Engine/Plugins/Experimental/DynamicMaterial/Source/DynamicMaterialEditor/Private/Components/MaterialStageBlends/DMMSBLinearDodge.cpp
@@ -7,6 +7,7 @@
 UDMMaterialStageBlendLinearDodge::UDMMaterialStageBlendLinearDodge()
 	: UDMMaterialStageBlendFunction(
 		LOCTEXT("BlendLinearDodge", "LinearDodge"),
+		LOCTEXT("BlendLinearDodgeDescription", "Linear Dodge increases the brightness of the base layer based on the blend layer on a per channel basis."),
 		"DM_Blend_LinearDodge",
 		TEXT("/Script/Engine.MaterialFunction'/DynamicMaterial/MaterialFunctions/Blends/MF_DM_Blend_LinearDodge.MF_DM_Blend_LinearDodge'")
 	)
diff --git a/Engine/Plugins/Experimental/DynamicMaterial/Source/DynamicMaterialEditor/Private/Components/MaterialStageBlends/DMMSBLinearLight.cpp b/Engine/Plugins/Experimental/DynamicMaterial/Source/DynamicMaterialEditor/Private/Components/MaterialStageBlends/DMMSBLinearLight.cpp
--- a/Engine/Plugins/Experimental/DynamicMaterial/Source/DynamicMaterialEditor/Private/Components/MaterialStageBlends/DMMSBLinearLight.cpp
+++ b/Engine/Plugins/Experimental/DynamicMaterial/Source/DynamicMaterialEditor/Private/Components/MaterialStageBlends/DMMSBLinearLight.cpp
@@ -7,6 +7,7 @@
 UDMMaterialStageBlendLinearLight::UDMMaterialStageBlendLinearLight()
 	: UDMMaterialStageBlendFunction(
 		LOCTEXT("BlendLinearLight", "LinearLight"),
+		LOCTEXT("BlendLinearLight", "Linear Light applies Linear Burn below 50% gray and Linear Dodge above 50% gray."),
 		"DM_Blend_LinearLight",
 		TEXT("/Script/Engine.MaterialFunction'/DynamicMaterial/MaterialFunctions/Blends/MF_DM_Blend_LinearLight.MF_DM_Blend_LinearLight'")
 	)
diff --git a/Engine/Plugins/Experimental/DynamicMaterial/Source/DynamicMaterialEditor/Private/Components/MaterialStageBlends/DMMSBLuminosity.cpp b/Engine/Plugins/Experimental/DynamicMaterial/Source/DynamicMaterialEditor/Private/Components/MaterialStageBlends/DMMSBLuminosity.cpp
--- a/Engine/Plugins/Experimental/DynamicMaterial/Source/DynamicMaterialEditor/Private/Components/MaterialStageBlends/DMMSBLuminosity.cpp
+++ b/Engine/Plugins/Experimental/DynamicMaterial/Source/DynamicMaterialEditor/Private/Components/MaterialStageBlends/DMMSBLuminosity.cpp
@@ -6,7 +6,8 @@
 
 UDMMaterialStageBlendLuminosity::UDMMaterialStageBlendLuminosity()
 	: UDMMaterialStageBlendFunction(
-		LOCTEXT("BlendLuminosity", "Luminosity"), 
+		LOCTEXT("BlendLuminosity", "Luminosity"),
+		LOCTEXT("BlendLuminosityDescription", "Luminosity combines the blend layer's luminosity with the base layer's hue and saturation."),
 		"DM_Blend_Luminosity", 
 		TEXT("/Script/Engine.MaterialFunction'/DynamicMaterial/MaterialFunctions/Blends/MF_DM_Blend_Luminosity.MF_DM_Blend_Luminosity'")
 	)
diff --git a/Engine/Plugins/Experimental/DynamicMaterial/Source/DynamicMaterialEditor/Private/Components/MaterialStageBlends/DMMSBMultiply.cpp b/Engine/Plugins/Experimental/DynamicMaterial/Source/DynamicMaterialEditor/Private/Components/MaterialStageBlends/DMMSBMultiply.cpp
--- a/Engine/Plugins/Experimental/DynamicMaterial/Source/DynamicMaterialEditor/Private/Components/MaterialStageBlends/DMMSBMultiply.cpp
+++ b/Engine/Plugins/Experimental/DynamicMaterial/Source/DynamicMaterialEditor/Private/Components/MaterialStageBlends/DMMSBMultiply.cpp
@@ -7,6 +7,7 @@
 UDMMaterialStageBlendMultiply::UDMMaterialStageBlendMultiply()
 	: UDMMaterialStageBlendFunction(
 		LOCTEXT("BlendMultiply", "Multiply"),
+		LOCTEXT("BlendMultiplyDescription", "Multiply multiplies the luminosity of the base and blend layers."),
 		"DM_Blend_Multiply",
 		TEXT("/Script/Engine.MaterialFunction'/DynamicMaterial/MaterialFunctions/Blends/MF_DM_Blend_Multiply.MF_DM_Blend_Multiply'")
 	)
diff --git a/Engine/Plugins/Experimental/DynamicMaterial/Source/DynamicMaterialEditor/Private/Components/MaterialStageBlends/DMMSBNormal.cpp b/Engine/Plugins/Experimental/DynamicMaterial/Source/DynamicMaterialEditor/Private/Components/MaterialStageBlends/DMMSBNormal.cpp
--- a/Engine/Plugins/Experimental/DynamicMaterial/Source/DynamicMaterialEditor/Private/Components/MaterialStageBlends/DMMSBNormal.cpp
+++ b/Engine/Plugins/Experimental/DynamicMaterial/Source/DynamicMaterialEditor/Private/Components/MaterialStageBlends/DMMSBNormal.cpp
@@ -7,6 +7,7 @@
 UDMMaterialStageBlendNormal::UDMMaterialStageBlendNormal()
 	: UDMMaterialStageBlendFunction(
 		LOCTEXT("BlendNormal", "Normal"),
+		LOCTEXT("BlendNormalDescription", "Normal blends replace the pixels below them without any blending unless the layer opacity is below 100%."),
 		"DM_Blend_Normal",
 		TEXT("/Script/Engine.MaterialFunction'/DynamicMaterial/MaterialFunctions/Blends/MF_DM_Blend_Normal.MF_DM_Blend_Normal'")
 	)
diff --git a/Engine/Plugins/Experimental/DynamicMaterial/Source/DynamicMaterialEditor/Private/Components/MaterialStageBlends/DMMSBOverlay.cpp b/Engine/Plugins/Experimental/DynamicMaterial/Source/DynamicMaterialEditor/Private/Components/MaterialStageBlends/DMMSBOverlay.cpp
--- a/Engine/Plugins/Experimental/DynamicMaterial/Source/DynamicMaterialEditor/Private/Components/MaterialStageBlends/DMMSBOverlay.cpp
+++ b/Engine/Plugins/Experimental/DynamicMaterial/Source/DynamicMaterialEditor/Private/Components/MaterialStageBlends/DMMSBOverlay.cpp
@@ -7,6 +7,7 @@
 UDMMaterialStageBlendOverlay::UDMMaterialStageBlendOverlay()
 	: UDMMaterialStageBlendFunction(
 		LOCTEXT("BlendOverlay", "Overlay"),
+		LOCTEXT("BlendOverlayDescription", "Overlay blends the base and blend layer with Multiply (at half strength) below 50% gray and Screen (at half strength) above 50% gray."),
 		"DM_Blend_Overlay",
 		TEXT("/Script/Engine.MaterialFunction'/DynamicMaterial/MaterialFunctions/Blends/MF_DM_Blend_Overlay.MF_DM_Blend_Overlay'")
 	)
diff --git a/Engine/Plugins/Experimental/DynamicMaterial/Source/DynamicMaterialEditor/Private/Components/MaterialStageBlends/DMMSBPinLight.cpp b/Engine/Plugins/Experimental/DynamicMaterial/Source/DynamicMaterialEditor/Private/Components/MaterialStageBlends/DMMSBPinLight.cpp
--- a/Engine/Plugins/Experimental/DynamicMaterial/Source/DynamicMaterialEditor/Private/Components/MaterialStageBlends/DMMSBPinLight.cpp
+++ b/Engine/Plugins/Experimental/DynamicMaterial/Source/DynamicMaterialEditor/Private/Components/MaterialStageBlends/DMMSBPinLight.cpp
@@ -7,6 +7,7 @@
 UDMMaterialStageBlendPinLight::UDMMaterialStageBlendPinLight()
 	: UDMMaterialStageBlendFunction(
 		LOCTEXT("BlendPinLight", "PinLight"),
+		LOCTEXT("BlendPinLightDescription", "Pin Light uses both Darken and Lighten simultaneously."),
 		"DM_Blend_PinLight",
 		TEXT("/Script/Engine.MaterialFunction'/DynamicMaterial/MaterialFunctions/Blends/MF_DM_Blend_PinLight.MF_DM_Blend_PinLight'")
 	)
diff --git a/Engine/Plugins/Experimental/DynamicMaterial/Source/DynamicMaterialEditor/Private/Components/MaterialStageBlends/DMMSBSaturation.cpp b/Engine/Plugins/Experimental/DynamicMaterial/Source/DynamicMaterialEditor/Private/Components/MaterialStageBlends/DMMSBSaturation.cpp
--- a/Engine/Plugins/Experimental/DynamicMaterial/Source/DynamicMaterialEditor/Private/Components/MaterialStageBlends/DMMSBSaturation.cpp
+++ b/Engine/Plugins/Experimental/DynamicMaterial/Source/DynamicMaterialEditor/Private/Components/MaterialStageBlends/DMMSBSaturation.cpp
@@ -6,7 +6,8 @@
 
 UDMMaterialStageBlendSaturation::UDMMaterialStageBlendSaturation()
 	: UDMMaterialStageBlendFunction(
-		LOCTEXT("BlendSaturation", "Saturation"), 
+		LOCTEXT("BlendSaturation", "Saturation"),
+		LOCTEXT("BlendSaturationDescription", "Saturation combines the blend layer's saturation with the base layer's luminosity and hue."),
 		"DM_Blend_Saturation", 
 		TEXT("/Script/Engine.MaterialFunction'/DynamicMaterial/MaterialFunctions/Blends/MF_DM_Blend_Saturation.MF_DM_Blend_Saturation'")
 	)
diff --git a/Engine/Plugins/Experimental/DynamicMaterial/Source/DynamicMaterialEditor/Private/Components/MaterialStageBlends/DMMSBScreen.cpp b/Engine/Plugins/Experimental/DynamicMaterial/Source/DynamicMaterialEditor/Private/Components/MaterialStageBlends/DMMSBScreen.cpp
--- a/Engine/Plugins/Experimental/DynamicMaterial/Source/DynamicMaterialEditor/Private/Components/MaterialStageBlends/DMMSBScreen.cpp
+++ b/Engine/Plugins/Experimental/DynamicMaterial/Source/DynamicMaterialEditor/Private/Components/MaterialStageBlends/DMMSBScreen.cpp
@@ -7,6 +7,7 @@
 UDMMaterialStageBlendScreen::UDMMaterialStageBlendScreen()
 	: UDMMaterialStageBlendFunction(
 		LOCTEXT("BlendScreen", "Screen"),
+		LOCTEXT("BlendScreenDescription", "Screen increases the brightness of the base layer based on the blend layer."),
 		"DM_Blend_Screen",
 		TEXT("/Script/Engine.MaterialFunction'/DynamicMaterial/MaterialFunctions/Blends/MF_DM_Blend_Screen.MF_DM_Blend_Screen'")
 	)
diff --git a/Engine/Plugins/Experimental/DynamicMaterial/Source/DynamicMaterialEditor/Private/Components/MaterialStageBlends/DMMSBSoftLight.cpp b/Engine/Plugins/Experimental/DynamicMaterial/Source/DynamicMaterialEditor/Private/Components/MaterialStageBlends/DMMSBSoftLight.cpp
--- a/Engine/Plugins/Experimental/DynamicMaterial/Source/DynamicMaterialEditor/Private/Components/MaterialStageBlends/DMMSBSoftLight.cpp
+++ b/Engine/Plugins/Experimental/DynamicMaterial/Source/DynamicMaterialEditor/Private/Components/MaterialStageBlends/DMMSBSoftLight.cpp
@@ -7,6 +7,7 @@
 UDMMaterialStageBlendSoftLight::UDMMaterialStageBlendSoftLight()
 	: UDMMaterialStageBlendFunction(
 		LOCTEXT("BlendSoftLight", "SoftLight"),
+		LOCTEXT("BlendSoftLightDescription", "Soft Light is similar to Overlay, but much more subtle."),
 		"DM_Blend_SoftLight",
 		TEXT("/Script/Engine.MaterialFunction'/DynamicMaterial/MaterialFunctions/Blends/MF_DM_Blend_SoftLight.MF_DM_Blend_SoftLight'")
 	)
diff --git a/Engine/Plugins/Experimental/DynamicMaterial/Source/DynamicMaterialEditor/Private/Components/MaterialStageBlends/DMMSBSubtract.cpp b/Engine/Plugins/Experimental/DynamicMaterial/Source/DynamicMaterialEditor/Private/Components/MaterialStageBlends/DMMSBSubtract.cpp
--- a/Engine/Plugins/Experimental/DynamicMaterial/Source/DynamicMaterialEditor/Private/Components/MaterialStageBlends/DMMSBSubtract.cpp
+++ b/Engine/Plugins/Experimental/DynamicMaterial/Source/DynamicMaterialEditor/Private/Components/MaterialStageBlends/DMMSBSubtract.cpp
@@ -7,6 +7,7 @@
 UDMMaterialStageBlendSubtract::UDMMaterialStageBlendSubtract()
 	: UDMMaterialStageBlendFunction(
 		LOCTEXT("BlendSubtract", "Subtract"),
+		LOCTEXT("BlendSubtractDescription", "Subtract subtracts the blend layer's color from the base layer. As the blend layer gets brighter, the base layer gets darker."),
 		"DM_Blend_Subtract",
 		TEXT("/Script/Engine.MaterialFunction'/DynamicMaterial/MaterialFunctions/Blends/MF_DM_Blend_Subtract.MF_DM_Blend_Subtract'")
 	)
diff --git a/Engine/Plugins/Experimental/DynamicMaterial/Source/DynamicMaterialEditor/Private/Components/MaterialStageBlends/DMMSBVividLight.cpp b/Engine/Plugins/Experimental/DynamicMaterial/Source/DynamicMaterialEditor/Private/Components/MaterialStageBlends/DMMSBVividLight.cpp
--- a/Engine/Plugins/Experimental/DynamicMaterial/Source/DynamicMaterialEditor/Private/Components/MaterialStageBlends/DMMSBVividLight.cpp
+++ b/Engine/Plugins/Experimental/DynamicMaterial/Source/DynamicMaterialEditor/Private/Components/MaterialStageBlends/DMMSBVividLight.cpp
@@ -7,6 +7,7 @@
 UDMMaterialStageBlendVividLight::UDMMaterialStageBlendVividLight()
 	: UDMMaterialStageBlendFunction(
 		LOCTEXT("BlendVividLight", "VividLight"),
+		LOCTEXT("BlendVividLightDescription", "Vivid Light darkens colors below 50% gray and lightens colors above 50% gray."),
 		"DM_Blend_VividLight",
 		TEXT("/Script/Engine.MaterialFunction'/DynamicMaterial/MaterialFunctions/Blends/MF_DM_Blend_VividLight.MF_DM_Blend_VividLight'")
 	)
diff --git a/Engine/Plugins/Experimental/DynamicMaterial/Source/DynamicMaterialEditor/Private/UI/Widgets/Editor/SlotEditor/SDMMaterialLayerBlendMode.cpp b/Engine/Plugins/Experimental/DynamicMaterial/Source/DynamicMaterialEditor/Private/UI/Widgets/Editor/SlotEditor/SDMMaterialLayerBlendMode.cpp
--- a/Engine/Plugins/Experimental/DynamicMaterial/Source/DynamicMaterialEditor/Private/UI/Widgets/Editor/SlotEditor/SDMMaterialLayerBlendMode.cpp
+++ b/Engine/Plugins/Experimental/DynamicMaterial/Source/DynamicMaterialEditor/Private/UI/Widgets/Editor/SlotEditor/SDMMaterialLayerBlendMode.cpp
@@ -74,21 +74,20 @@ namespace UE::DynamicMaterialEditor::Private
 				LOCTEXT("BlendDarken", "Darken Blends"),
 				{
 					UDMMaterialStageBlendDarken::StaticClass(),
+					UDMMaterialStageBlendDarkenColor::StaticClass(),
 					UDMMaterialStageBlendMultiply::StaticClass(),
 					UDMMaterialStageBlendColorBurn::StaticClass(),
-					UDMMaterialStageBlendLinearBurn::StaticClass(),
-					UDMMaterialStageBlendDarkenColor::StaticClass()
+					UDMMaterialStageBlendLinearBurn::StaticClass()
 				}
 			},
 			{
 				LOCTEXT("BlendLighten", "Lighten Blends"),
 				{
-					UDMMaterialStageBlendAdd::StaticClass(),
 					UDMMaterialStageBlendLighten::StaticClass(),
+					UDMMaterialStageBlendLightenColor::StaticClass(),
 					UDMMaterialStageBlendScreen::StaticClass(),
 					UDMMaterialStageBlendColorDodge::StaticClass(),
-					UDMMaterialStageBlendLinearDodge::StaticClass(),
-					UDMMaterialStageBlendLightenColor::StaticClass()
+					UDMMaterialStageBlendLinearDodge::StaticClass()
 				}
 			},
 			{
@@ -371,12 +370,27 @@ void SDMMaterialLayerBlendMode::MakeSourceBlendMenu(UToolMenu* InToolMenu)
 
 	auto AddClassToSection = [&SourceBlendWidget](FToolMenuSection& InSection, const FText& InName, UClass* InBlendClass)
 	{
-		const FText ToolTipFormat = LOCTEXT("ChangeSourceBlendTooltip", "Change the source of this stage to a {0} Material Blend.");
+		if (!InBlendClass)
+		{
+			return;
+		}
+
+		if (InBlendClass->GetClassFlags() & EClassFlags::CLASS_Deprecated)
+		{
+			return;
+		}
+
+		UDMMaterialStageBlend* BlendCDO = InBlendClass->GetDefaultObject<UDMMaterialStageBlend>();
+
+		if (!BlendCDO)
+		{
+			return;
+		}
 
 		InSection.AddMenuEntry(
 			NAME_None,
 			InName,
-			FText::Format(ToolTipFormat, InName),
+			BlendCDO->GetBlendDescription(),
 			FSlateIcon(),
 			FUIAction(
 				FExecuteAction::CreateSP(SourceBlendWidget.ToSharedRef(), &SDMMaterialLayerBlendMode::OnBlendModeSelected, InBlendClass),
diff --git a/Engine/Plugins/Experimental/DynamicMaterial/Source/DynamicMaterialEditor/Public/Components/DMMaterialStageBlend.h b/Engine/Plugins/Experimental/DynamicMaterial/Source/DynamicMaterialEditor/Public/Components/DMMaterialStageBlend.h
--- a/Engine/Plugins/Experimental/DynamicMaterial/Source/DynamicMaterialEditor/Public/Components/DMMaterialStageBlend.h
+++ b/Engine/Plugins/Experimental/DynamicMaterial/Source/DynamicMaterialEditor/Public/Components/DMMaterialStageBlend.h
@@ -42,6 +42,9 @@ public:
 	UFUNCTION(BlueprintCallable, Category = "Material Designer")
 	DYNAMICMATERIALEDITOR_API void SetBaseChannelOverride(EAvaColorChannel InMaskChannel);
 
+	UFUNCTION(BlueprintPure, Category = "Material Designer")
+	const FText& GetBlendDescription() const;
+
 	//~ Begin UDMMaterialStageThroughput
 	DYNAMICMATERIALEDITOR_API virtual bool CanInputAcceptType(int32 InputIndex, EDMValueType ValueType) const override;
 	DYNAMICMATERIALEDITOR_API virtual void AddDefaultInput(int32 InInputIndex) const override;
@@ -86,7 +89,11 @@ protected:
 		meta = (NotKeyframeable, ToolTip = "Changes the output channel of the base input."))
 	mutable EAvaColorChannel BaseChannelOverride;
 
-	DYNAMICMATERIALEDITOR_API UDMMaterialStageBlend(const FText& InName);
+	UPROPERTY(EditDefaultsOnly, BlueprintReadOnly, Getter, BlueprintGetter = GetBlendDescription,
+		Category = "Material Designer")
+	FText BlendDescription;
+
+	DYNAMICMATERIALEDITOR_API UDMMaterialStageBlend(const FText& InName, const FText& InDescription);
 
 	/* Returns true if there are any outputs on the base input that have more than 1 channel. */
 	bool CanUseBaseChannelOverride() const;
diff --git a/Engine/Plugins/Experimental/DynamicMaterial/Source/DynamicMaterialEditor/Public/Components/DMMaterialStageBlendFunction.h b/Engine/Plugins/Experimental/DynamicMaterial/Source/DynamicMaterialEditor/Public/Components/DMMaterialStageBlendFunction.h
--- a/Engine/Plugins/Experimental/DynamicMaterial/Source/DynamicMaterialEditor/Public/Components/DMMaterialStageBlendFunction.h
+++ b/Engine/Plugins/Experimental/DynamicMaterial/Source/DynamicMaterialEditor/Public/Components/DMMaterialStageBlendFunction.h
@@ -31,8 +31,8 @@ protected:
 	UPROPERTY()
 	TObjectPtr<UMaterialFunctionInterface> MaterialFunction;
 
-	DYNAMICMATERIALEDITOR_API UDMMaterialStageBlendFunction(const FText& InName, UMaterialFunctionInterface* InMaterialFunction);
+	DYNAMICMATERIALEDITOR_API UDMMaterialStageBlendFunction(const FText& InName, const FText& InDescription, UMaterialFunctionInterface* InMaterialFunction);
 
 	/** Loads the function asset. */
-	DYNAMICMATERIALEDITOR_API UDMMaterialStageBlendFunction(const FText& InName, const FName& InFunctionName, const FString& InFunctionPath);
+	DYNAMICMATERIALEDITOR_API UDMMaterialStageBlendFunction(const FText& InName, const FText& InDescription, const FName& InFunctionName, const FString& InFunctionPath);
 };
diff --git a/Engine/Plugins/Experimental/DynamicMaterial/Source/DynamicMaterialEditor/Public/Components/MaterialStageBlends/DMMSBAdd.h b/Engine/Plugins/Experimental/DynamicMaterial/Source/DynamicMaterialEditor/Public/Components/MaterialStageBlends/DMMSBAdd.h
--- a/Engine/Plugins/Experimental/DynamicMaterial/Source/DynamicMaterialEditor/Public/Components/MaterialStageBlends/DMMSBAdd.h
+++ b/Engine/Plugins/Experimental/DynamicMaterial/Source/DynamicMaterialEditor/Public/Components/MaterialStageBlends/DMMSBAdd.h
@@ -5,6 +5,7 @@
 #include "Components/DMMaterialStageBlendFunction.h"
 #include "DMMSBAdd.generated.h"
 
+// The same as linear dodge
 UCLASS(MinimalAPI, BlueprintType, ClassGroup = "Material Designer")
 class UDMMaterialStageBlendAdd : public UDMMaterialStageBlendFunction
 {
diff --git a/Engine/Plugins/Experimental/DynamicMaterial/Source/DynamicMaterialEditor/Public/Components/MaterialStageBlends/DMMSBContrastBase.h b/Engine/Plugins/Experimental/DynamicMaterial/Source/DynamicMaterialEditor/Public/Components/MaterialStageBlends/DMMSBContrastBase.h
--- a/Engine/Plugins/Experimental/DynamicMaterial/Source/DynamicMaterialEditor/Public/Components/MaterialStageBlends/DMMSBContrastBase.h
+++ b/Engine/Plugins/Experimental/DynamicMaterial/Source/DynamicMaterialEditor/Public/Components/MaterialStageBlends/DMMSBContrastBase.h
@@ -14,7 +14,7 @@ class UDMMaterialStageBlendContrastBase : public UDMMaterialStageBlend
 	GENERATED_BODY()
 
 public:
-	DYNAMICMATERIALEDITOR_API UDMMaterialStageBlendContrastBase(const FText& InName);
+	DYNAMICMATERIALEDITOR_API UDMMaterialStageBlendContrastBase(const FText& InName, const FText& InDescription);
 
 private:
 	UDMMaterialStageBlendContrastBase();
