Material Designer: Blend mode layer connections fixed
- An issue preventing layers from connecting their outputs to each other correctly has been fixed.

diff --git a/Engine/Plugins/Experimental/DynamicMaterial/Source/DynamicMaterialEditor/Private/Components/DMMaterialStageThroughput.cpp b/Engine/Plugins/Experimental/DynamicMaterial/Source/DynamicMaterialEditor/Private/Components/DMMaterialStageThroughput.cpp
--- a/Engine/Plugins/Experimental/DynamicMaterial/Source/DynamicMaterialEditor/Private/Components/DMMaterialStageThroughput.cpp
+++ b/Engine/Plugins/Experimental/DynamicMaterial/Source/DynamicMaterialEditor/Private/Components/DMMaterialStageThroughput.cpp
@@ -722,6 +722,7 @@ int32 UDMMaterialStageThroughput::ResolveInputChannel(const TSharedRef<FDMMateri
 			const TArray<UMaterialExpression*>& PreviousLayerExpressions = InBuildState->GetLayerExpressions(PreviousLayer);
 			check(!PreviousLayerExpressions.IsEmpty());
 			LastExpression = PreviousLayerExpressions.Last();
+			OutExpressions.Add(LastExpression);
 
 			UDMMaterialStage* PreviousStage = PreviousLayer->GetLastValidStage(EDMMaterialLayerStage::All);
 			check(PreviousStage);
