Horde issue 779357

diff --git a/Engine/Plugins/Experimental/DynamicMaterial/Source/DynamicMaterialEditor/Private/Components/MaterialStageBlends/DMMSBLinearLight.cpp b/Engine/Plugins/Experimental/DynamicMaterial/Source/DynamicMaterialEditor/Private/Components/MaterialStageBlends/DMMSBLinearLight.cpp
--- a/Engine/Plugins/Experimental/DynamicMaterial/Source/DynamicMaterialEditor/Private/Components/MaterialStageBlends/DMMSBLinearLight.cpp
+++ b/Engine/Plugins/Experimental/DynamicMaterial/Source/DynamicMaterialEditor/Private/Components/MaterialStageBlends/DMMSBLinearLight.cpp
@@ -7,7 +7,7 @@
 UDMMaterialStageBlendLinearLight::UDMMaterialStageBlendLinearLight()
 	: UDMMaterialStageBlendFunction(
 		LOCTEXT("BlendLinearLight", "LinearLight"),
-		LOCTEXT("BlendLinearLight", "Linear Light applies Linear Burn below 50% gray and Linear Dodge above 50% gray."),
+		LOCTEXT("BlendLinearLightDescription", "Linear Light applies Linear Burn below 50% gray and Linear Dodge above 50% gray."),
 		"DM_Blend_LinearLight",
 		TEXT("/Script/Engine.MaterialFunction'/DynamicMaterial/MaterialFunctions/Blends/MF_DM_Blend_LinearLight.MF_DM_Blend_LinearLight'")
 	)
