Material Designer: Updated menus
- Every menu now has an extension point.
- Video renamed to Media.
- Video moved to base menu instead of Advanced menu.

diff --git a/Engine/Plugins/Experimental/DynamicMaterial/Source/DynamicMaterialEditor/Private/UI/Menus/DMMaterialSlotLayerMenus.cpp b/Engine/Plugins/Experimental/DynamicMaterial/Source/DynamicMaterialEditor/Private/UI/Menus/DMMaterialSlotLayerMenus.cpp
--- a/Engine/Plugins/Experimental/DynamicMaterial/Source/DynamicMaterialEditor/Private/UI/Menus/DMMaterialSlotLayerMenus.cpp
+++ b/Engine/Plugins/Experimental/DynamicMaterial/Source/DynamicMaterialEditor/Private/UI/Menus/DMMaterialSlotLayerMenus.cpp
@@ -30,7 +30,9 @@
 #include "ScopedTransaction.h"
 #include "Styling/SlateIconFinder.h"
 #include "Templates/SharedPointer.h"
+#include "ToolMenu.h"
 #include "ToolMenus.h"
+#include "ToolMenuSection.h"
 #include "UI/Menus/DMMenuContext.h"
 #include "UI/Widgets/Editor/SDMMaterialSlotEditor.h"
 #include "UI/Widgets/SDMMaterialEditor.h"
@@ -41,12 +43,55 @@
 namespace UE::DynamicMaterialEditor::Private
 {
 	static const FName SlotLayerMenuName = TEXT("MaterialDesigner.MaterialSlot.Layer");
+	static const FName SlotLayerAddMenuName = TEXT("MaterialDesigner.MaterialSlot.AddLayer");
 	static const FName SlotLayerAddSectionName = TEXT("AddLayer");
 	static const FName SlotLayerModifySectionName = TEXT("ModifyLayer");
 	static const FName GlobalValuesSectionName = TEXT("GlobalValues");
 }
 
 TSharedRef<SWidget> FDMMaterialSlotLayerMenus::GenerateSlotLayerMenu(const TSharedPtr<SDMMaterialSlotEditor>& InSlotWidget, UDMMaterialLayerObject* InLayer)
+{
+	return InLayer
+		? GenerateSlotLayerMenu_Layer(InSlotWidget, InLayer)
+		: GenerateSlotLayerMenu_AddLayer(InSlotWidget);
+}
+
+TSharedRef<SWidget> FDMMaterialSlotLayerMenus::GenerateSlotLayerMenu_AddLayer(const TSharedPtr<SDMMaterialSlotEditor>& InSlotWidget)
+{
+	using namespace UE::DynamicMaterialEditor::Private;
+
+	if (!UToolMenus::Get()->IsMenuRegistered(SlotLayerAddMenuName))
+	{
+		UToolMenu* NewToolMenu = UDMMenuContext::GenerateContextMenuDefault(SlotLayerAddMenuName);
+
+		if (!NewToolMenu)
+		{
+			return SNullWidget::NullWidget;
+		}
+
+		FToolMenuSection& AddLayerSection = NewToolMenu->FindOrAddSection(SlotLayerAddSectionName, LOCTEXT("AddLayer", "Add Layer"));
+		AddLayerSection.AddDynamicEntry(SlotLayerAddSectionName, FNewToolMenuSectionDelegate::CreateStatic(&FDMMaterialSlotLayerMenus::AddAddLayerSection));
+		AddLayerSection.InsertPosition.Position = EToolMenuInsertType::First;
+
+		if constexpr (UE::DynamicMaterialEditor::bGlobalValuesEnabled)
+		{
+			FToolMenuSection& GlobalValueSection = NewToolMenu->FindOrAddSection(GlobalValuesSectionName, LOCTEXT("GlobalValue", "Add Global Value Layer"));
+			GlobalValueSection.AddDynamicEntry(GlobalValuesSectionName, FNewToolMenuSectionDelegate::CreateStatic(&FDMMaterialSlotLayerMenus::AddGlobalValueSection));
+		}
+	}
+
+	FToolMenuContext MenuContext(UDMMenuContext::CreateLayer(InSlotWidget->GetEditorWidget(), nullptr));
+
+	if (TSharedPtr<SDMMaterialEditor> EditorWidget = InSlotWidget->GetEditorWidget())
+	{
+		MenuContext.AppendCommandList(EditorWidget->GetCommandList());
+	}
+
+	return UToolMenus::Get()->GenerateWidget(SlotLayerAddMenuName, MenuContext);
+}
+
+TSharedRef<SWidget> FDMMaterialSlotLayerMenus::GenerateSlotLayerMenu_Layer(const TSharedPtr<SDMMaterialSlotEditor>& InSlotWidget,
+	UDMMaterialLayerObject* InLayer)
 {
 	using namespace UE::DynamicMaterialEditor::Private;
 
@@ -59,14 +104,19 @@ TSharedRef<SWidget> FDMMaterialSlotLayerMenus::GenerateSlotLayerMenu(const TShar
 			return SNullWidget::NullWidget;
 		}
 
-		NewToolMenu->AddDynamicSection(TEXT("AddLayer"), FNewToolMenuDelegate::CreateStatic(&FDMMaterialSlotLayerMenus::AddAddLayerSection));
+		FToolMenuSection& AddLayerSection = NewToolMenu->FindOrAddSection(SlotLayerAddSectionName, LOCTEXT("AddLayer", "Add Layer"));
+		AddLayerSection.AddDynamicEntry(SlotLayerAddSectionName, FNewToolMenuSectionDelegate::CreateStatic(&FDMMaterialSlotLayerMenus::AddAddLayerSection));
+		AddLayerSection.InsertPosition.Position = EToolMenuInsertType::First;
 
 		if constexpr (UE::DynamicMaterialEditor::bGlobalValuesEnabled)
 		{
-			NewToolMenu->AddDynamicSection(NAME_None, FNewToolMenuDelegate::CreateStatic(&FDMMaterialSlotLayerMenus::AddGlobalValueSection));
+			FToolMenuSection& GlobalValueSection = NewToolMenu->FindOrAddSection(GlobalValuesSectionName, LOCTEXT("GlobalValue", "Add Global Value Layer"));
+			GlobalValueSection.AddDynamicEntry(GlobalValuesSectionName, FNewToolMenuSectionDelegate::CreateStatic(&FDMMaterialSlotLayerMenus::AddGlobalValueSection));
 		}
 
-		NewToolMenu->AddDynamicSection(NAME_None, FNewToolMenuDelegate::CreateStatic(&FDMMaterialSlotLayerMenus::AddLayerModifySection));
+		FToolMenuSection& ModifyLayerSection = NewToolMenu->FindOrAddSection(SlotLayerModifySectionName, LOCTEXT("ModifyLayer", "Modify Layer"));
+		ModifyLayerSection.AddDynamicEntry(SlotLayerModifySectionName, FNewToolMenuSectionDelegate::CreateStatic(&FDMMaterialSlotLayerMenus::AddLayerModifySection));
+		ModifyLayerSection.InsertPosition.Position = EToolMenuInsertType::Last;
 	}
 
 	FToolMenuContext MenuContext(UDMMenuContext::CreateLayer(InSlotWidget->GetEditorWidget(), InLayer));
@@ -79,16 +129,11 @@ TSharedRef<SWidget> FDMMaterialSlotLayerMenus::GenerateSlotLayerMenu(const TShar
 	return UToolMenus::Get()->GenerateWidget(SlotLayerMenuName, MenuContext);
 }
 
-void FDMMaterialSlotLayerMenus::AddAddLayerSection(UToolMenu* InMenu)
+void FDMMaterialSlotLayerMenus::AddAddLayerSection(FToolMenuSection& InSection)
 {
 	using namespace UE::DynamicMaterialEditor::Private;
 
-	if (!IsValid(InMenu) || InMenu->ContainsSection(SlotLayerAddSectionName))
-	{
-		return;
-	}
-
-	UDMMenuContext* MenuContext = InMenu->FindContext<UDMMenuContext>();
+	UDMMenuContext* MenuContext = InSection.FindContext<UDMMenuContext>();
 
 	if (!MenuContext)
 	{
@@ -123,24 +168,39 @@ void FDMMaterialSlotLayerMenus::AddAddLayerSection(UToolMenu* InMenu)
 		return;
 	}
 
-	FToolMenuSection& NewSection = InMenu->AddSection(SlotLayerAddSectionName, LOCTEXT("AddLayer", "Add Layer"));
-
-	NewSection.AddMenuEntry(
-		NAME_None,
-		LOCTEXT("AddTextureSample", "Texture"),
-		LOCTEXT("AddTextureSampleTooltip", "Add a Material Stage based on a Texture."),
-		GetDefault<UDMMaterialStageExpressionTextureSample>()->GetComponentIcon(),
+	InSection.AddMenuEntry(
+		TEXT("Noise"),
+		LOCTEXT("AddNoise", "Noise"),
+		LOCTEXT("AddNoiseTooltip", "Add a new Material Layer with a noise pattern."),
+		GetDefault<UDMMaterialStageExpressionWorldPositionNoise>()->GetComponentIcon(),
 		FUIAction(FExecuteAction::CreateWeakLambda(
 			Slot,
-			[Slot, ExpressionClass = TSubclassOf<UDMMaterialStageExpression>(UDMMaterialStageExpressionTextureSample::StaticClass())]
+			[Slot]
 			{
-				UDMMaterialSlotFunctionLibrary::AddNewLayer_Expression(Slot, ExpressionClass);
+				UDMMaterialSlotFunctionLibrary::AddNewLayer_Expression(Slot, TSubclassOf<UDMMaterialStageExpression>(UDMMaterialStageExpressionWorldPositionNoise::StaticClass()));
 			}
 		))
 	);
 
-	NewSection.AddMenuEntry(
-		NAME_None,
+	if (ModelEditorOnlyData->GetDomain() == EMaterialDomain::MD_PostProcess)
+	{
+		InSection.AddMenuEntry(
+			TEXT("PostProcess"),
+			LOCTEXT("AddSceneTexture", "Post Process"),
+			LOCTEXT("AddSceneTextureTooltip", "Add a new Material Layer that represents the Scene Texture for a post process material."),
+			GetDefault<UDMMaterialStageExpressionSceneTexture>()->GetComponentIcon(),
+			FUIAction(FExecuteAction::CreateWeakLambda(
+				Slot,
+				[Slot]
+				{
+					UDMMaterialSlotFunctionLibrary::AddNewLayer_SceneTexture(Slot);
+				}
+			))
+		);
+	}
+
+	InSection.AddMenuEntry(
+		TEXT("SolidColor"),
 		LOCTEXT("AddColor", "Solid Color"),
 		LOCTEXT("AddColorTooltip", "Add a new Material Layer with a solid RGB color."),
 		GetDefault<UDMMaterialValueFloat3RGB>()->GetComponentIcon(),
@@ -153,57 +213,40 @@ void FDMMaterialSlotLayerMenus::AddAddLayerSection(UToolMenu* InMenu)
 		))
 	);
 
-	NewSection.AddMenuEntry(
-		NAME_None,
-		LOCTEXT("AddEdgeColor", "Texture Edge Color"),
-		LOCTEXT("AddEdgeColorTooltip", "Add a new Material Layer with a solid color based on the edge color on a texture."),
-		GetDefault<UDMMaterialStageExpressionTextureSampleEdgeColor>()->GetComponentIcon(),
+	InSection.AddMenuEntry(
+		TEXT("Texture"),
+		LOCTEXT("AddTextureSample", "Texture"),
+		LOCTEXT("AddTextureSampleTooltip", "Add a Material Stage based on a Texture."),
+		GetDefault<UDMMaterialStageExpressionTextureSample>()->GetComponentIcon(),
 		FUIAction(FExecuteAction::CreateWeakLambda(
 			Slot,
-			[Slot]
+			[Slot, ExpressionClass = TSubclassOf<UDMMaterialStageExpression>(UDMMaterialStageExpressionTextureSample::StaticClass())]
 			{
-				UDMMaterialSlotFunctionLibrary::AddNewLayer_Expression(Slot, TSubclassOf<UDMMaterialStageExpression>(UDMMaterialStageExpressionTextureSampleEdgeColor::StaticClass()));
+				UDMMaterialSlotFunctionLibrary::AddNewLayer_Expression(Slot, ExpressionClass);
 			}
 		))
 	);
 
-	NewSection.AddMenuEntry(
-		NAME_None,
-		LOCTEXT("AddNoise", "Noise"),
-		LOCTEXT("AddNoiseTooltip", "Add a new Material Layer with a noise pattern."),
-		GetDefault<UDMMaterialStageExpressionWorldPositionNoise>()->GetComponentIcon(),
+	InSection.AddMenuEntry(
+		TEXT("TextureEdgeColor"),
+		LOCTEXT("AddEdgeColor", "Texture Edge Color"),
+		LOCTEXT("AddEdgeColorTooltip", "Add a new Material Layer with a solid color based on the edge color on a texture."),
+		GetDefault<UDMMaterialStageExpressionTextureSampleEdgeColor>()->GetComponentIcon(),
 		FUIAction(FExecuteAction::CreateWeakLambda(
 			Slot,
 			[Slot]
 			{
-				UDMMaterialSlotFunctionLibrary::AddNewLayer_Expression(Slot, TSubclassOf<UDMMaterialStageExpression>(UDMMaterialStageExpressionWorldPositionNoise::StaticClass()));
+				UDMMaterialSlotFunctionLibrary::AddNewLayer_Expression(Slot, TSubclassOf<UDMMaterialStageExpression>(UDMMaterialStageExpressionTextureSampleEdgeColor::StaticClass()));
 			}
 		))
 	);
 
-	if (ModelEditorOnlyData->GetDomain() == EMaterialDomain::MD_PostProcess)
-	{
-		NewSection.AddMenuEntry(
-			NAME_None,
-			LOCTEXT("AddSceneTexture", "Post Process"),
-			LOCTEXT("AddSceneTextureTooltip", "Add a new Material Layer that represents the Scene Texture for a post process material."),
-			GetDefault<UDMMaterialStageExpressionSceneTexture>()->GetComponentIcon(),
-			FUIAction(FExecuteAction::CreateWeakLambda(
-				Slot,
-				[Slot]
-				{
-					UDMMaterialSlotFunctionLibrary::AddNewLayer_SceneTexture(Slot);
-				}
-			))
-		);
-	}
-
 	const TArray<TStrongObjectPtr<UClass>>& Gradients = UDMMaterialStageGradient::GetAvailableGradients();
 
 	if (!Gradients.IsEmpty())
 	{
-		NewSection.AddSubMenu(
-			NAME_None,
+		InSection.AddSubMenu(
+			TEXT("GradientMenu"),
 			LOCTEXT("AddGradientStage", "Gradient"),
 			LOCTEXT("AddGradientStageTooltip", "Add a Material Stage based on a Material Gradient."),
 			FNewToolMenuDelegate::CreateStatic(&AddLayerMenu_Gradients)
@@ -212,27 +255,22 @@ void FDMMaterialSlotLayerMenus::AddAddLayerSection(UToolMenu* InMenu)
 
 	if constexpr (UE::DynamicMaterialEditor::bGlobalValuesEnabled)
 	{
-		AddGlobalValueSection(InMenu);
+		AddGlobalValueSection(InSection);
 	}
 
-	NewSection.AddSubMenu(
-		TEXT("Advanced"),
+	InSection.AddSubMenu(
+		TEXT("AdvancedMenu"),
 		LOCTEXT("AddAdvancedStage", "Advanced"),
 		LOCTEXT("AddAdvancedStageTooltip", "Add an advanced Material Stage."),
 		FNewToolMenuDelegate::CreateStatic(&AddLayerMenu_Advanced)
 	);
 }
 
-void FDMMaterialSlotLayerMenus::AddLayerModifySection(UToolMenu* InMenu)
+void FDMMaterialSlotLayerMenus::AddLayerModifySection(FToolMenuSection& InSection)
 {
 	using namespace UE::DynamicMaterialEditor::Private;
 
-	if (!IsValid(InMenu) || InMenu->ContainsSection(SlotLayerModifySectionName))
-	{
-		return;
-	}
-
-	const UDMMenuContext* const MenuContext = InMenu->FindContext<UDMMenuContext>();
+	const UDMMenuContext* const MenuContext = InSection.FindContext<UDMMenuContext>();
 	if (!MenuContext)
 	{
 		return;
@@ -250,15 +288,14 @@ void FDMMaterialSlotLayerMenus::AddLayerModifySection(UToolMenu* InMenu)
 		return;
 	}
 
-	FToolMenuSection& NewSection = InMenu->AddSection(SlotLayerModifySectionName, LOCTEXT("LayerActions", "Layer Actions"));
-
 	if (Slot->CanRemoveLayer(Layer))
 	{
 		FSlateIcon ToggleLayerIcon = Layer->IsEnabled()
 			? FSlateIcon(FAppStyle::Get().GetStyleSetName(), TEXT("Kismet.VariableList.HideForInstance"))
 			: FSlateIcon(FAppStyle::Get().GetStyleSetName(), TEXT("Kismet.VariableList.ExposeForInstance"));
 
-		NewSection.AddMenuEntry(NAME_None,
+		InSection.AddMenuEntry(
+			TEXT("ToggleLayer"),
 			LOCTEXT("ToggleLayer", "Toggle Layer"),
 			LOCTEXT("ToggleLayerTooltip", "Toggle the Layer.\n\nAlt+Left Click"),
 			ToggleLayerIcon,
@@ -278,30 +315,25 @@ void FDMMaterialSlotLayerMenus::AddLayerModifySection(UToolMenu* InMenu)
 		);
 	}
 
-	NewSection.AddMenuEntry(
+	InSection.AddMenuEntry(
 		FDynamicMaterialEditorCommands::Get().InsertDefaultLayerAbove,
 		TAttribute<FText>(),
 		TAttribute<FText>(),
 		FSlateIconFinder::FindIcon("EditableComboBox.Add")
 	);
 
-	NewSection.AddMenuEntry(FGenericCommands::Get().Copy);
-	NewSection.AddMenuEntry(FGenericCommands::Get().Cut);
-	NewSection.AddMenuEntry(FGenericCommands::Get().Paste);
-	NewSection.AddMenuEntry(FGenericCommands::Get().Duplicate);
-	NewSection.AddMenuEntry(FGenericCommands::Get().Delete);
+	InSection.AddMenuEntry(FGenericCommands::Get().Copy);
+	InSection.AddMenuEntry(FGenericCommands::Get().Cut);
+	InSection.AddMenuEntry(FGenericCommands::Get().Paste);
+	InSection.AddMenuEntry(FGenericCommands::Get().Duplicate);
+	InSection.AddMenuEntry(FGenericCommands::Get().Delete);
 }
 
-void FDMMaterialSlotLayerMenus::AddGlobalValueSection(UToolMenu* InMenu)
+void FDMMaterialSlotLayerMenus::AddGlobalValueSection(FToolMenuSection& InSection)
 {
 	using namespace UE::DynamicMaterialEditor::Private;
 
-	if (!IsValid(InMenu))
-	{
-		return;
-	}
-
-	UDMMenuContext* MenuContext = InMenu->FindContext<UDMMenuContext>();
+	UDMMenuContext* MenuContext = InSection.FindContext<UDMMenuContext>();
 
 	if (!MenuContext)
 	{
@@ -322,10 +354,8 @@ void FDMMaterialSlotLayerMenus::AddGlobalValueSection(UToolMenu* InMenu)
 		return;
 	}
 
-	FToolMenuSection& NewSection = InMenu->AddSection(GlobalValuesSectionName, LOCTEXT("GlobalValues", "Add Global Value"));
-
-	NewSection.AddSubMenu(
-		NAME_None,
+	InSection.AddSubMenu(
+		TEXT("GlobalValueMenu"),
 		LOCTEXT("AddValueStage", "Global Value"),
 		LOCTEXT("AddValueStageTooltip", "Add a Material Stage based on a Material Value defined above."),
 		FNewToolMenuDelegate::CreateLambda([](UToolMenu* InMenu)
@@ -356,9 +386,9 @@ void FDMMaterialSlotLayerMenus::AddGlobalValueSection(UToolMenu* InMenu)
 				for (UDMMaterialValue* Value : Values)
 				{
 					InMenu->AddMenuEntry(
-						NAME_None,
+						TEXT("Value"),
 						FToolMenuEntry::InitMenuEntry(
-							NAME_None,
+							Value->GetFName(),
 							Value->GetDescription(),
 							LOCTEXT("AddValueStageSpecificTooltip", "Add a Material Stage based on this Material Value."),
 							Value->GetComponentIcon(),
@@ -374,8 +404,8 @@ void FDMMaterialSlotLayerMenus::AddGlobalValueSection(UToolMenu* InMenu)
 			})
 	);
 
-	NewSection.AddSubMenu(
-		NAME_None,
+	InSection.AddSubMenu(
+		TEXT("NewGlobalValue"),
 		LOCTEXT("AddNewValueStage", "New Global Value"),
 		LOCTEXT("AddNewValueStageTooltip", "Add a new global Material Value as use it as a Material Stage."),
 		FNewToolMenuDelegate::CreateLambda([](UToolMenu* InMenu)
@@ -409,9 +439,9 @@ void FDMMaterialSlotLayerMenus::AddGlobalValueSection(UToolMenu* InMenu)
 					const FSlateIcon ValueIcon = UDMValueDefinitionLibrary::GetValueIcon(ValueType);
 
 					InMenu->AddMenuEntry(
-						NAME_None,
+						TEXT("NewGlobalValue"),
 						FToolMenuEntry::InitMenuEntry(
-							NAME_None,
+							*Name.ToString(),
 							Name,
 							FormattedTooltip,
 							ValueIcon,
@@ -440,9 +470,9 @@ void FDMMaterialSlotLayerMenus::AddSlotMenuEntry(const TSharedPtr<SDMMaterialSlo
 	}
 
 	InMenu->AddMenuEntry(
-		NAME_None,
+		TEXT("Slot"),
 		FToolMenuEntry::InitMenuEntry(
-			NAME_None,
+			TargetSlot->GetFName(),
 			InName,
 			LOCTEXT("AddSlotStageSpecificTooltip", "Add a Material Stage based on this Material Slot."),
 			FSlateIcon(),
@@ -591,9 +621,9 @@ void FDMMaterialSlotLayerMenus::AddLayerInputsMenu_Slots(UToolMenu* InMenu)
 		else
 		{
 			InMenu->AddMenuEntry(
-				NAME_None,
+				TEXT("Slot"),
 				FToolMenuEntry::InitSubMenu(
-					NAME_None,
+					CurrentSlot->GetFName(),
 					LOCTEXT("AddSlotStage2", "Add Slot Output"),
 					LOCTEXT("AddSlotStageTooltip2", "Add a Material Stage based on the output of another Material Slot."),
 					FNewToolMenuDelegate::CreateStatic(&AddLayerInputsMenu_Slot_Properties, CurrentSlot))
@@ -632,6 +662,11 @@ void FDMMaterialSlotLayerMenus::AddLayerMenu_Gradients(UToolMenu* InMenu)
 
 	const TArray<TStrongObjectPtr<UClass>>& Gradients = UDMMaterialStageGradient::GetAvailableGradients();
 
+	FToolMenuSection& NewSection = InMenu->AddSection(
+		TEXT("Gradient"),
+		LOCTEXT("Gradients", "Gradients")
+	);
+
 	for (const TStrongObjectPtr<UClass>& Gradient : Gradients)
 	{
 		UDMMaterialStageGradient* GradientCDO = Cast<UDMMaterialStageGradient>(Gradient->GetDefaultObject());
@@ -643,39 +678,33 @@ void FDMMaterialSlotLayerMenus::AddLayerMenu_Gradients(UToolMenu* InMenu)
 
 		const FText MenuName = GradientCDO->GetDescription();
 
-		InMenu->AddMenuEntry(
-			NAME_None,
-			FToolMenuEntry::InitMenuEntry(
-				NAME_None,
-				MenuName,
-				LOCTEXT("ChangeGradientSourceTooltip", "Change the source of this stage to a Material Gradient."),
-				GradientCDO->GetComponentIcon(),
-				FUIAction(FExecuteAction::CreateWeakLambda(
-					Slot,
-					[Slot, GradientClass = TSubclassOf<UDMMaterialStageGradient>(Gradient.Get())]
-					{
-						UDMMaterialSlotFunctionLibrary::AddNewLayer_Gradient(Slot, GradientClass);
-					}
-				))
-			)
-		);
-	}
-
-	InMenu->AddMenuEntry(
-		NAME_None,
-		FToolMenuEntry::InitMenuEntry(
-			NAME_None,
-			LOCTEXT("AddColorAtlas", "Color Atlas"),
-			LOCTEXT("AddColorAtlasTooltip", "Add a new Material Layer with a Color Atlas."),
-			GetDefault<UDMMaterialValueColorAtlas>()->GetComponentIcon(),
+		NewSection.AddMenuEntry(
+			GradientCDO->GetFName(),
+			MenuName,
+			LOCTEXT("ChangeGradientSourceTooltip", "Change the source of this stage to a Material Gradient."),
+			GradientCDO->GetComponentIcon(),
 			FUIAction(FExecuteAction::CreateWeakLambda(
 				Slot,
-				[Slot]
+				[Slot, GradientClass = TSubclassOf<UDMMaterialStageGradient>(Gradient.Get())]
 				{
-					UDMMaterialSlotFunctionLibrary::AddNewLayer_NewLocalValue(Slot, EDMValueType::VT_ColorAtlas);
+					UDMMaterialSlotFunctionLibrary::AddNewLayer_Gradient(Slot, GradientClass);
 				}
 			))
-		)
+		);
+	}
+
+	NewSection.AddMenuEntry(
+		TEXT("ColorAtlas"),
+		LOCTEXT("AddColorAtlas", "Color Atlas"),
+		LOCTEXT("AddColorAtlasTooltip", "Add a new Material Layer with a Color Atlas."),
+		GetDefault<UDMMaterialValueColorAtlas>()->GetComponentIcon(),
+		FUIAction(FExecuteAction::CreateWeakLambda(
+			Slot,
+			[Slot]
+			{
+				UDMMaterialSlotFunctionLibrary::AddNewLayer_NewLocalValue(Slot, EDMValueType::VT_ColorAtlas);
+			}
+		))
 	);
 }
 
@@ -724,12 +753,12 @@ void FDMMaterialSlotLayerMenus::AddLayerMenu_Advanced(UToolMenu* InMenu)
 	}
 
 	FToolMenuSection& NewSection = InMenu->AddSection(
-		NAME_None,
+		TEXT("Advanced"),
 		LOCTEXT("Advanced", "Advanced")
 	);
 
 	NewSection.AddMenuEntry(
-		NAME_None,
+		TEXT("Text"),
 		LOCTEXT("AddText", "Text"),
 		LOCTEXT("AddTextTooltip", "Add a Material Stage based on a Text Renderer."),
 		GetDefault<UDMRenderTargetTextRenderer>()->GetComponentIcon(),
@@ -743,7 +772,7 @@ void FDMMaterialSlotLayerMenus::AddLayerMenu_Advanced(UToolMenu* InMenu)
 	);
 
 	NewSection.AddMenuEntry(
-		NAME_None,
+		TEXT("Widget"),
 		LOCTEXT("AddWidget", "Widget"),
 		LOCTEXT("AddWidgetTooltip", "Add a Material Stage based on a Widget Renderer."),
 		GetDefault<UDMRenderTargetUMGWidgetRenderer>()->GetComponentIcon(),
@@ -757,7 +786,7 @@ void FDMMaterialSlotLayerMenus::AddLayerMenu_Advanced(UToolMenu* InMenu)
 	);
 
 	NewSection.AddMenuEntry(
-		NAME_None,
+		TEXT("MaterialFunction"),
 		LOCTEXT("AddMaterialFunction", "Material Function"),
 		LOCTEXT("AddMaterialFunctionTooltip", "Add a new Material Layer based on a Material Function."),
 		GetDefault<UDMMaterialStageFunction>()->GetComponentIcon(),
@@ -801,7 +830,7 @@ void FDMMaterialSlotLayerMenus::AddLayerMenu_Advanced(UToolMenu* InMenu)
 		if (bHasValidSlot)
 		{
 			NewSection.AddSubMenu(
-				NAME_None,
+				TEXT("SlotOutput"),
 				LOCTEXT("AddSlotStage", "Slot Output"),
 				LOCTEXT("AddSlotStageTooltip", "Add a Material Stage based on the output of another Material Slot."),
 				FNewToolMenuDelegate::CreateStatic(&AddLayerInputsMenu_Slots)
diff --git a/Engine/Plugins/Experimental/DynamicMaterial/Source/DynamicMaterialEditor/Private/UI/Menus/DMMaterialSlotLayerMenus.h b/Engine/Plugins/Experimental/DynamicMaterial/Source/DynamicMaterialEditor/Private/UI/Menus/DMMaterialSlotLayerMenus.h
--- a/Engine/Plugins/Experimental/DynamicMaterial/Source/DynamicMaterialEditor/Private/UI/Menus/DMMaterialSlotLayerMenus.h
+++ b/Engine/Plugins/Experimental/DynamicMaterial/Source/DynamicMaterialEditor/Private/UI/Menus/DMMaterialSlotLayerMenus.h
@@ -11,18 +11,23 @@ class UDMMaterialLayerObject;
 class UDMMaterialSlot;
 class UToolMenu;
 enum class EDMMaterialPropertyType : uint8;
+struct FToolMenuSection;
 
 class FDMMaterialSlotLayerMenus final
 {
 public:
 	static TSharedRef<SWidget> GenerateSlotLayerMenu(const TSharedPtr<SDMMaterialSlotEditor>& InSlotWidget, UDMMaterialLayerObject* InLayer);
 
-	static void AddAddLayerSection(UToolMenu* InMenu);
+	static void AddAddLayerSection(FToolMenuSection& InSection);
 
 private:
-	static void AddLayerModifySection(UToolMenu* InMenu);
+	static TSharedRef<SWidget> GenerateSlotLayerMenu_AddLayer(const TSharedPtr<SDMMaterialSlotEditor>& InSlotWidget);
 
-	static void AddGlobalValueSection(UToolMenu* InMenu);
+	static TSharedRef<SWidget> GenerateSlotLayerMenu_Layer(const TSharedPtr<SDMMaterialSlotEditor>& InSlotWidget, UDMMaterialLayerObject* InLayer);
+
+	static void AddLayerModifySection(FToolMenuSection& InSection);
+
+	static void AddGlobalValueSection(FToolMenuSection& InSection);
 
 	static void AddSlotMenuEntry(const TSharedPtr<SDMMaterialSlotEditor> InSlotWidget, UToolMenu* InMenu, const FText& InName, 
 		UDMMaterialSlot* InSourceSlot, EDMMaterialPropertyType InMaterialProperty);
diff --git a/Engine/Plugins/Experimental/DynamicMaterial/Source/DynamicMaterialEditor/Private/UI/Menus/DMMaterialSlotMenus.cpp b/Engine/Plugins/Experimental/DynamicMaterial/Source/DynamicMaterialEditor/Private/UI/Menus/DMMaterialSlotMenus.cpp
--- a/Engine/Plugins/Experimental/DynamicMaterial/Source/DynamicMaterialEditor/Private/UI/Menus/DMMaterialSlotMenus.cpp
+++ b/Engine/Plugins/Experimental/DynamicMaterial/Source/DynamicMaterialEditor/Private/UI/Menus/DMMaterialSlotMenus.cpp
@@ -29,7 +29,9 @@ TSharedRef<SWidget> FDMMaterialSlotMenus::MakeAddLayerButtonMenu(const TSharedPt
 			return SNullWidget::NullWidget;
 		}
 
-		NewToolMenu->AddDynamicSection(NAME_None, FNewToolMenuDelegate::CreateStatic(&FDMMaterialSlotLayerMenus::AddAddLayerSection));
+		FToolMenuSection& AddLayerSection = NewToolMenu->FindOrAddSection(TEXT("AddLayer"), LOCTEXT("AddLayer", "Add Layer"));
+
+		AddLayerSection.AddDynamicEntry(TEXT("AddLayer"), FNewToolMenuSectionDelegate::CreateStatic(&FDMMaterialSlotLayerMenus::AddAddLayerSection));
 	}
 
 	FToolMenuContext MenuContext(UDMMenuContext::CreateEditor(InSlotWidget->GetEditorWidget()));
diff --git a/Engine/Plugins/Experimental/DynamicMaterial/Source/DynamicMaterialEditor/Private/UI/Menus/DMMaterialStageMenus.cpp b/Engine/Plugins/Experimental/DynamicMaterial/Source/DynamicMaterialEditor/Private/UI/Menus/DMMaterialStageMenus.cpp
--- a/Engine/Plugins/Experimental/DynamicMaterial/Source/DynamicMaterialEditor/Private/UI/Menus/DMMaterialStageMenus.cpp
+++ b/Engine/Plugins/Experimental/DynamicMaterial/Source/DynamicMaterialEditor/Private/UI/Menus/DMMaterialStageMenus.cpp
@@ -21,8 +21,8 @@
 namespace UE::DynamicMaterialEditor::Private
 {
 	const FLazyName StageSettingsMenuName = TEXT("MaterialDesigner.MaterialStage");
-	const FLazyName StageMenuToggleName = TEXT("StageToggle");
-	const FLazyName StageSourceMenuName = TEXT("MaterialDesigner.MaterialStageSource");
+	const FLazyName StageSettingsSectionName = TEXT("ChangeStageSettings");
+	const FLazyName StageSourceSectionName = TEXT("ChangeStageSource");
 }
 
 TSharedRef<SWidget> FDMMaterialStageMenus::GenerateStageMenu(const TSharedPtr<SDMMaterialSlotEditor>& InSlotWidget,
@@ -41,10 +41,11 @@ TSharedRef<SWidget> FDMMaterialStageMenus::GenerateStageMenu(const TSharedPtr<SD
 			return SNullWidget::NullWidget;
 		}
 
-		//FToolMenuSection& NewSection = NewToolMenu->FindOrAddSection("ChangeMaterialStage");
-		FToolMenuSection& NewSection = NewToolMenu->FindOrAddSection(NAME_None);
-		NewSection.AddDynamicEntry(StageMenuToggleName, FNewToolMenuSectionDelegate::CreateStatic(&FDMMaterialStageMenus::AddStageSettingsSection));
-		NewSection.AddDynamicEntry(StageSourceMenuName, FNewToolMenuSectionDelegate::CreateStatic(&FDMMaterialStageMenus::AddStageSourceSection));
+		FToolMenuSection& ChangeSection = NewToolMenu->FindOrAddSection(StageSettingsSectionName, LOCTEXT("ModifyStage", "Modify Stage"));
+		ChangeSection.AddDynamicEntry(StageSettingsSectionName, FNewToolMenuSectionDelegate::CreateStatic(&FDMMaterialStageMenus::AddStageSettingsSection));
+
+		FToolMenuSection& SourceSection = NewToolMenu->FindOrAddSection(StageSourceSectionName, LOCTEXT("ChangeStageSource", "Change Stage Source"));
+		SourceSection.AddDynamicEntry(StageSourceSectionName, FNewToolMenuSectionDelegate::CreateStatic(&FDMMaterialStageMenus::AddStageSourceSection));
 	}
 
 	FToolMenuContext MenuContext(UDMMenuContext::CreateStage(InSlotWidget->GetEditorWidget(), InStageWidget));
@@ -85,13 +86,14 @@ void FDMMaterialStageMenus::AddStageSettingsSection(FToolMenuSection& InSection)
 	// Only if we can remove a layer can we toggle the base stage.
 	const bool bAllowRemoveLayer = Slot->CanRemoveLayer(Layer);
 
-	InSection.InitSection(StageMenuToggleName, LOCTEXT("MaterialStageMenu", "Stage Actions"), FToolMenuInsert());
+	InSection.InitSection(StageSettingsSectionName, LOCTEXT("MaterialStageMenu", "Stage Actions"), FToolMenuInsert());
 
 	if (bAllowRemoveLayer)
 	{
 		if (StageType == EDMMaterialLayerStage::Base)
 		{
-			InSection.AddMenuEntry(NAME_None,
+			InSection.AddMenuEntry(
+				TEXT("ToggleBase"),
 				LOCTEXT("ToggleLayerBase", "Toggle Base"),
 				LOCTEXT("ToggleLayerBaseTooltip", "Toggle the Layer Base.\n\nAlt+Shift+Left Click"),
 				FSlateIcon(FDynamicMaterialEditorStyle::Get().GetStyleSetName(), TEXT("Icons.Stage.Enabled")),
@@ -114,7 +116,8 @@ void FDMMaterialStageMenus::AddStageSettingsSection(FToolMenuSection& InSection)
 
 	if (StageType == EDMMaterialLayerStage::Mask)
 	{
-		InSection.AddMenuEntry(NAME_None,
+		InSection.AddMenuEntry(
+			TEXT("ToggleMask"),
 			LOCTEXT("ToggleLayerMask", "Toggle Mask"),
 			LOCTEXT("ToggleLayerMaskTooltip", "Toggle the Layer Mask.\n\nAlt+Shift+Left Click"),
 			FSlateIcon(FDynamicMaterialEditorStyle::Get().GetStyleSetName(), TEXT("Icons.Stage.Enabled")),
@@ -139,7 +142,7 @@ void FDMMaterialStageMenus::AddStageSourceSection(FToolMenuSection& InSection)
 {
 	using namespace UE::DynamicMaterialEditor::Private;
 
-	InSection.InitSection(StageSourceMenuName, LOCTEXT("MaterialStageSource", "Change Stage Source"), FToolMenuInsert());
+	InSection.InitSection(StageSourceSectionName, LOCTEXT("MaterialStageSource", "Change Stage Source"), FToolMenuInsert());
 
 	FDMMaterialStageSourceMenus::CreateChangeMaterialStageSource(InSection);
 }
diff --git a/Engine/Plugins/Experimental/DynamicMaterial/Source/DynamicMaterialEditor/Private/UI/Menus/DMMaterialStageSourceMenus.cpp b/Engine/Plugins/Experimental/DynamicMaterial/Source/DynamicMaterialEditor/Private/UI/Menus/DMMaterialStageSourceMenus.cpp
--- a/Engine/Plugins/Experimental/DynamicMaterial/Source/DynamicMaterialEditor/Private/UI/Menus/DMMaterialStageSourceMenus.cpp
+++ b/Engine/Plugins/Experimental/DynamicMaterial/Source/DynamicMaterialEditor/Private/UI/Menus/DMMaterialStageSourceMenus.cpp
@@ -64,8 +64,8 @@ TSharedRef<SWidget> FDMMaterialStageSourceMenus::MakeChangeSourceMenu(const TSha
 			return SNullWidget::NullWidget;
 		}
 
-		FToolMenuSection& NewSection = NewToolMenu->FindOrAddSection("ChangeMaterialStage", LOCTEXT("ChangeStageSource", "Change Stage Source"));
-		NewSection.AddDynamicEntry(NAME_None, FNewToolMenuSectionDelegate::CreateStatic(&CreateChangeMaterialStageSource));
+		FToolMenuSection& NewSection = NewToolMenu->FindOrAddSection(TEXT("ChangeStageSource"), LOCTEXT("ChangeStageSource", "Change Stage Source"));
+		NewSection.AddDynamicEntry(TEXT("ChangeStageSource"), FNewToolMenuSectionDelegate::CreateStatic(&CreateChangeMaterialStageSource));
 	}
 
 	FToolMenuContext MenuContext(UDMMenuContext::CreateStage(InSlotWidget->GetEditorWidget(), InStageWidget));
@@ -706,6 +706,11 @@ void FDMMaterialStageSourceMenus::GenerateChangeSourceMenu_Gradients(UToolMenu*
 		return;
 	}
 
+	FToolMenuSection& NewSection = InMenu->AddSection(
+		TEXT("Gradient"),
+		LOCTEXT("Gradients", "Gradients")
+	);
+
 	for (TStrongObjectPtr<UClass> GradientClass : Gradients)
 	{
 		UDMMaterialStageGradient* GradientCDO = Cast<UDMMaterialStageGradient>(GradientClass->GetDefaultObject());
@@ -714,75 +719,69 @@ void FDMMaterialStageSourceMenus::GenerateChangeSourceMenu_Gradients(UToolMenu*
 		{
 			const FText MenuName = GradientCDO->GetDescription();
 
-			InMenu->AddMenuEntry(
-				NAME_None,
-				FToolMenuEntry::InitMenuEntry(
-					NAME_None,
-					MenuName,
-					LOCTEXT("ChangeSourceGradientTooltip", "Change the source of this stage to a Material Gradient."),
-					GradientCDO->GetComponentIcon(),
-					FUIAction(FExecuteAction::CreateWeakLambda(
-						MenuContext,
-						[MenuContext, GradientClass]()
+			NewSection.AddMenuEntry(
+				GradientCDO->GetFName(),
+				MenuName,
+				LOCTEXT("ChangeSourceGradientTooltip", "Change the source of this stage to a Material Gradient."),
+				GradientCDO->GetComponentIcon(),
+				FUIAction(FExecuteAction::CreateWeakLambda(
+					MenuContext,
+					[MenuContext, GradientClass]()
+					{
+						UDMMaterialStage* const Stage = MenuContext->GetStage();
+						if (!Stage)
 						{
-							UDMMaterialStage* const Stage = MenuContext->GetStage();
-							if (!Stage)
-							{
-								return;
-							}
+							return;
+						}
 
-							UDMMaterialStageSource* const StageSource = Stage->GetSource();
-							if (!StageSource)
-							{
-								return;
-							}
+						UDMMaterialStageSource* const StageSource = Stage->GetSource();
+						if (!StageSource)
+						{
+							return;
+						}
 
-							if (StageSource->IsA<UDMMaterialStageBlend>())
-							{
-								FScopedTransaction Transaction(LOCTEXT("SetStageInputBase", "Set Material Designer Base Source"));
-								Stage->Modify();
+						if (StageSource->IsA<UDMMaterialStageBlend>())
+						{
+							FScopedTransaction Transaction(LOCTEXT("SetStageInputBase", "Set Material Designer Base Source"));
+							Stage->Modify();
 
-								UDMMaterialStageInputGradient::ChangeStageInput_Gradient(
-									Stage,
-									GradientClass.Get(), UDMMaterialStageBlend::InputB,
-									FDMMaterialStageConnectorChannel::WHOLE_CHANNEL,
-									FDMMaterialStageConnectorChannel::WHOLE_CHANNEL
-								);
-							}
-							else if (StageSource->IsA<UDMMaterialStageThroughputLayerBlend>())
-							{
-								FScopedTransaction Transaction(LOCTEXT("SetStageInputMask", "Set Material Designer Mask Source"));
-								Stage->Modify();
-								UDMMaterialStageInputGradient::ChangeStageInput_Gradient(
-									Stage,
-									GradientClass.Get(), 
-									UDMMaterialStageThroughputLayerBlend::InputMaskSource,
-									FDMMaterialStageConnectorChannel::WHOLE_CHANNEL,
-									FDMMaterialStageConnectorChannel::WHOLE_CHANNEL
-								);
-							}
-							else
-							{
-								ensureMsgf(false, TEXT("Invalid stage type (%s)"), *StageSource->GetClass()->GetName());
-							}
-						})
-					)
+							UDMMaterialStageInputGradient::ChangeStageInput_Gradient(
+								Stage,
+								GradientClass.Get(), UDMMaterialStageBlend::InputB,
+								FDMMaterialStageConnectorChannel::WHOLE_CHANNEL,
+								FDMMaterialStageConnectorChannel::WHOLE_CHANNEL
+							);
+						}
+						else if (StageSource->IsA<UDMMaterialStageThroughputLayerBlend>())
+						{
+							FScopedTransaction Transaction(LOCTEXT("SetStageInputMask", "Set Material Designer Mask Source"));
+							Stage->Modify();
+							UDMMaterialStageInputGradient::ChangeStageInput_Gradient(
+								Stage,
+								GradientClass.Get(), 
+								UDMMaterialStageThroughputLayerBlend::InputMaskSource,
+								FDMMaterialStageConnectorChannel::WHOLE_CHANNEL,
+								FDMMaterialStageConnectorChannel::WHOLE_CHANNEL
+							);
+						}
+						else
+						{
+							ensureMsgf(false, TEXT("Invalid stage type (%s)"), *StageSource->GetClass()->GetName());
+						}
+					})
 			));
 		}
 	}
 
-	InMenu->AddMenuEntry(
-		NAME_None,
-		FToolMenuEntry::InitMenuEntry(
-			NAME_None,
-			LOCTEXT("ChangeSourceColorAtlas", "Color Atlas"),
-			LOCTEXT("ChangeSourceColorAtlasTooltip", "Change the source of this stage to a Color Atlas."),
-			GetDefault<UDMMaterialValueColorAtlas>()->GetComponentIcon(),
-			FUIAction(
-				FExecuteAction::CreateStatic(
-					&ChangeSourceToColorAtlasFromContext,
-					MenuContext
-				)
+	NewSection.AddMenuEntry(
+		TEXT("ColorAtlas"),
+		LOCTEXT("ChangeSourceColorAtlas", "Color Atlas"),
+		LOCTEXT("ChangeSourceColorAtlasTooltip", "Change the source of this stage to a Color Atlas."),
+		GetDefault<UDMMaterialValueColorAtlas>()->GetComponentIcon(),
+		FUIAction(
+			FExecuteAction::CreateStatic(
+				&ChangeSourceToColorAtlasFromContext,
+				MenuContext
 			)
 		)
 	);
@@ -1465,18 +1464,37 @@ void FDMMaterialStageSourceMenus::CreateChangeMaterialStageSource(FToolMenuSecti
 
 	const TArray<TStrongObjectPtr<UClass>>& Gradients = UDMMaterialStageGradient::GetAvailableGradients();
 
-	InSection.AddMenuEntry("TextureSample",
-		LOCTEXT("TextureSample", "Texture"),
-		LOCTEXT("ChangeSourceTextureSampleTooltip", "Change the source of this stage to a texture."),
-		GetDefault<UDMMaterialStageExpressionTextureSample>()->GetComponentIcon(),
+	InSection.AddMenuEntry("Noise",
+		LOCTEXT("ChangeSourceNoise", "Noise"),
+		LOCTEXT("ChangeSourceNoiseTooltip", "Change the source of this stage to a Noise Renderer."),
+		GetDefault<UDMMaterialStageExpressionWorldPositionNoise>()->GetComponentIcon(),
 		FUIAction(
 			FExecuteAction::CreateStatic(
-				&ChangeSourceToTextureSampleFromContext,
+				&ChangeSourceToNoiseFromContext,
 				MenuContext
 			)
 		)
 	);
 
+	if (ModelEditorOnlyData->GetDomain() == EMaterialDomain::MD_PostProcess)
+	{
+		InSection.AddMenuEntry("SceneTexture",
+			LOCTEXT("AddSceneTexture", "Post Process"),
+			LOCTEXT("ChangeSourceSceneTextureTooltip", "Change the source of this stage to Scene Texture in post process materials."),
+			GetDefault<UDMMaterialStageExpressionSceneTexture>()->GetComponentIcon(),
+			FUIAction(
+				FExecuteAction::CreateStatic(
+					&ChangeSourceToSolidColorRGBFromContext,
+					MenuContext
+				),
+				FCanExecuteAction::CreateStatic(
+					&CanChangeSourceToSceneTextureFromContext,
+					MenuContext
+				)
+			)
+		);
+	}
+
 	InSection.AddMenuEntry("SolidColor",
 		LOCTEXT("ChangeSourceColorRGB", "Solid Color"),
 		LOCTEXT("ChangeSourceColorRGBTooltip", "Change the source of this stage to a Solid Color."),
@@ -1489,49 +1507,30 @@ void FDMMaterialStageSourceMenus::CreateChangeMaterialStageSource(FToolMenuSecti
 		)
 	);
 
-	InSection.AddMenuEntry("TextureSample_EdgeColor",
-		LOCTEXT("AddTextureSampleEgdeColor", "Texture Edge Color"),
-		LOCTEXT("ChangeSourceTextureSampleEdgeColorTooltip", "Change the source of this stage to the edge color of a texture."),
-		GetDefault<UDMMaterialStageExpressionTextureSampleEdgeColor>()->GetComponentIcon(),
+	InSection.AddMenuEntry("TextureSample",
+		LOCTEXT("TextureSample", "Texture"),
+		LOCTEXT("ChangeSourceTextureSampleTooltip", "Change the source of this stage to a texture."),
+		GetDefault<UDMMaterialStageExpressionTextureSample>()->GetComponentIcon(),
 		FUIAction(
 			FExecuteAction::CreateStatic(
-				&ChangeSourceToTextureSampleEdgeColorFromContext,
+				&ChangeSourceToTextureSampleFromContext,
 				MenuContext
 			)
 		)
 	);
 
-	InSection.AddMenuEntry("Noise",
-		LOCTEXT("ChangeSourceNoise", "Noise"),
-		LOCTEXT("ChangeSourceNoiseTooltip", "Change the source of this stage to a Noise Renderer."),
-		GetDefault<UDMMaterialStageExpressionWorldPositionNoise>()->GetComponentIcon(),
+	InSection.AddMenuEntry("TextureSample_EdgeColor",
+		LOCTEXT("AddTextureSampleEgdeColor", "Texture Edge Color"),
+		LOCTEXT("ChangeSourceTextureSampleEdgeColorTooltip", "Change the source of this stage to the edge color of a texture."),
+		GetDefault<UDMMaterialStageExpressionTextureSampleEdgeColor>()->GetComponentIcon(),
 		FUIAction(
 			FExecuteAction::CreateStatic(
-				&ChangeSourceToNoiseFromContext,
+				&ChangeSourceToTextureSampleEdgeColorFromContext,
 				MenuContext
 			)
 		)
 	);
 
-	if (ModelEditorOnlyData->GetDomain() == EMaterialDomain::MD_PostProcess)
-	{
-		InSection.AddMenuEntry("SceneTexture",
-			LOCTEXT("AddSceneTexture", "Post Process"),
-			LOCTEXT("ChangeSourceSceneTextureTooltip", "Change the source of this stage to Scene Texture in post process materials."),
-			GetDefault<UDMMaterialStageExpressionSceneTexture>()->GetComponentIcon(),
-			FUIAction(
-				FExecuteAction::CreateStatic(
-					&ChangeSourceToSolidColorRGBFromContext,
-					MenuContext
-				),
-				FCanExecuteAction::CreateStatic(
-					&CanChangeSourceToSceneTextureFromContext,
-					MenuContext
-				)
-			)
-		);
-	}
-
 	if constexpr (UE::DynamicMaterialEditor::bGlobalValuesEnabled)
 	{
 		const TArray<UDMMaterialValue*>& Values = MaterialModel->GetValues();
diff --git a/Engine/Plugins/Experimental/DynamicMaterialMediaStreamBridge/Source/DynamicMaterialMediaStreamBridgeEditor/Private/DMMediaStreamStageSourceMenuExtender.cpp b/Engine/Plugins/Experimental/DynamicMaterialMediaStreamBridge/Source/DynamicMaterialMediaStreamBridgeEditor/Private/DMMediaStreamStageSourceMenuExtender.cpp
--- a/Engine/Plugins/Experimental/DynamicMaterialMediaStreamBridge/Source/DynamicMaterialMediaStreamBridgeEditor/Private/DMMediaStreamStageSourceMenuExtender.cpp
+++ b/Engine/Plugins/Experimental/DynamicMaterialMediaStreamBridge/Source/DynamicMaterialMediaStreamBridgeEditor/Private/DMMediaStreamStageSourceMenuExtender.cpp
@@ -15,8 +15,6 @@
 #include "ScopedTransaction.h"
 #include "ToolMenu.h"
 #include "ToolMenus.h"
-#include "Components/DMMaterialValueDynamic.h"
-#include "Model/DynamicMaterialModelEditorOnlyData.h"
 #include "UI/Menus/DMMenuContext.h"
 
 #define LOCTEXT_NAMESPACE "DMMEdiaStreamStageSourceMenuExtender"
@@ -35,48 +33,61 @@ void FDMMediaStreamStageSourceMenuExtender::Integrate()
 		return;
 	}
 
-	if (UToolMenu* AdvancedMenu = ToolMenus->ExtendMenu(TEXT("MaterialDesigner.MaterialStage.ChangeSource.Advanced")))
+	if (UToolMenu* AdvancedMenu = ToolMenus->ExtendMenu(TEXT("MaterialDesigner.MaterialStage.ChangeSource")))
 	{
-		AdvancedMenu->AddDynamicSection(
-			NAME_None,
-			FNewToolMenuDelegate::CreateRaw(this, &FDMMediaStreamStageSourceMenuExtender::ExtendMenu_ChangeSource)
+		FToolMenuSection& NewSection = AdvancedMenu->FindOrAddSection(TEXT("ChangeStageSource"), LOCTEXT("ChangeStageSource", "Change Stage Source"));
+
+		NewSection.AddDynamicEntry(
+			TEXT("MediaStream"),
+			FNewToolMenuSectionDelegate::CreateRaw(this, &FDMMediaStreamStageSourceMenuExtender::ExtendMenu_ChangeSource)
+		);
+	}
+
+	if (UToolMenu* AdvancedMenu = ToolMenus->ExtendMenu(TEXT("MaterialDesigner.MaterialStage")))
+	{
+		FToolMenuSection& NewSection = AdvancedMenu->FindOrAddSection(TEXT("ChangeStageSource"), LOCTEXT("ChangeStageSource", "Change Stage Source"));
+
+		NewSection.AddDynamicEntry(
+			TEXT("MediaStream"),
+			FNewToolMenuSectionDelegate::CreateRaw(this, &FDMMediaStreamStageSourceMenuExtender::ExtendMenu_ChangeSource)
 		);
 	}
 
-	if (UToolMenu* AdvancedMenu = ToolMenus->ExtendMenu(TEXT("MaterialDesigner.MaterialStage.Advanced")))
+	if (UToolMenu* AdvancedMenu = ToolMenus->ExtendMenu(TEXT("MaterialDesigner.MaterialSlot.Layer")))
 	{
-		AdvancedMenu->AddDynamicSection(
-			NAME_None,
-			FNewToolMenuDelegate::CreateRaw(this, &FDMMediaStreamStageSourceMenuExtender::ExtendMenu_ChangeSource)
+		FToolMenuSection& NewSection = AdvancedMenu->FindOrAddSection(TEXT("AddLayer"), LOCTEXT("AddLayer", "Add Layer"));
+
+		NewSection.AddDynamicEntry(
+			TEXT("MediaStream"),
+			FNewToolMenuSectionDelegate::CreateRaw(this, &FDMMediaStreamStageSourceMenuExtender::ExtendMenu_AddLayer)
 		);
 	}
 
-	if (UToolMenu* AdvancedMenu = ToolMenus->ExtendMenu(TEXT("MaterialDesigner.MaterialSlot.Layer.Advanced")))
+	if (UToolMenu* AdvancedMenu = ToolMenus->ExtendMenu(TEXT("MaterialDesigner.MaterialSlot.AddLayer")))
 	{
-		AdvancedMenu->AddDynamicSection(
-			NAME_None,
-			FNewToolMenuDelegate::CreateRaw(this, &FDMMediaStreamStageSourceMenuExtender::ExtendMenu_AddLayer)
+		FToolMenuSection& NewSection = AdvancedMenu->FindOrAddSection(TEXT("AddLayer"), LOCTEXT("AddLayer", "Add Layer"));
+
+		NewSection.AddDynamicEntry(
+			TEXT("MediaStream"),
+			FNewToolMenuSectionDelegate::CreateRaw(this, &FDMMediaStreamStageSourceMenuExtender::ExtendMenu_AddLayer)
 		);
 	}
 
 	bIntegrated = true;
 }
 
-void FDMMediaStreamStageSourceMenuExtender::ExtendMenu_ChangeSource(UToolMenu* InMenu)
+void FDMMediaStreamStageSourceMenuExtender::ExtendMenu_ChangeSource(FToolMenuSection& InSection)
 {
-	UDMMenuContext* MenuContext = InMenu->FindContext<UDMMenuContext>();
+	UDMMenuContext* MenuContext = InSection.FindContext<UDMMenuContext>();
 
 	if (!ensure(MenuContext))
 	{
 		return;
 	}
 
-	FToolMenuSection& AdvancedSection = InMenu->FindOrAddSection(NAME_None);
-	AdvancedSection.Context.AddObject(MenuContext);
-
-	AdvancedSection.AddMenuEntry("MediaStream",
-		LOCTEXT("ChangeSourceMediaStream", "Video"),
-		LOCTEXT("ChangeSourceMediaStreamTooltip", "Change the source of this stage to a Video."),
+	FToolMenuEntry& NewEntry = InSection.AddMenuEntry("MediaStream",
+		LOCTEXT("ChangeSourceMediaStream", "Media"),
+		LOCTEXT("ChangeSourceMediaStreamTooltip", "Change the source of this stage to a Media Stream."),
 		GetDefault<UDMMaterialValueMediaStream>()->GetComponentIcon(),
 		FUIAction(
 			FExecuteAction::CreateRaw(this,
@@ -85,23 +96,22 @@ void FDMMediaStreamStageSourceMenuExtender::ExtendMenu_ChangeSource(UToolMenu* I
 			)
 		)
 	);
+
+	NewEntry.InsertPosition = FToolMenuInsert(TEXT("Noise"), EToolMenuInsertType::Before);
 }
 
-void FDMMediaStreamStageSourceMenuExtender::ExtendMenu_AddLayer(UToolMenu* InMenu)
+void FDMMediaStreamStageSourceMenuExtender::ExtendMenu_AddLayer(FToolMenuSection& InSection)
 {
-	UDMMenuContext* MenuContext = InMenu->FindContext<UDMMenuContext>();
+	UDMMenuContext* MenuContext = InSection.FindContext<UDMMenuContext>();
 
 	if (!ensure(MenuContext))
 	{
 		return;
 	}
 
-	FToolMenuSection& AdvancedSection = InMenu->FindOrAddSection(NAME_None);
-	AdvancedSection.Context.AddObject(MenuContext);
-
-	AdvancedSection.AddMenuEntry("MediaStream",
-		LOCTEXT("AddLayerMediaStream", "Video"),
-		LOCTEXT("AddLayerMediaStreamTooltip", "Add a new layer based on a Video."),
+	FToolMenuEntry& NewEntry = InSection.AddMenuEntry("MediaStream",
+		LOCTEXT("AddLayerMediaStream", "Media"),
+		LOCTEXT("AddLayerMediaStreamTooltip", "Add a new layer based on a Media Stream."),
 		GetDefault<UDMMaterialValueMediaStream>()->GetComponentIcon(),
 		FUIAction(
 			FExecuteAction::CreateRaw(this,
@@ -110,6 +120,8 @@ void FDMMediaStreamStageSourceMenuExtender::ExtendMenu_AddLayer(UToolMenu* InMen
 			)
 		)
 	);
+
+	NewEntry.InsertPosition = FToolMenuInsert(TEXT("Noise"), EToolMenuInsertType::Before);
 }
 
 void FDMMediaStreamStageSourceMenuExtender::ChangeSourceToMediaStreamFromContext(UDMMenuContext* InMenuContext)
diff --git a/Engine/Plugins/Experimental/DynamicMaterialMediaStreamBridge/Source/DynamicMaterialMediaStreamBridgeEditor/Private/DMMediaStreamStageSourceMenuExtender.h b/Engine/Plugins/Experimental/DynamicMaterialMediaStreamBridge/Source/DynamicMaterialMediaStreamBridgeEditor/Private/DMMediaStreamStageSourceMenuExtender.h
--- a/Engine/Plugins/Experimental/DynamicMaterialMediaStreamBridge/Source/DynamicMaterialMediaStreamBridgeEditor/Private/DMMediaStreamStageSourceMenuExtender.h
+++ b/Engine/Plugins/Experimental/DynamicMaterialMediaStreamBridge/Source/DynamicMaterialMediaStreamBridgeEditor/Private/DMMediaStreamStageSourceMenuExtender.h
@@ -3,7 +3,6 @@
 #pragma once
 
 class UDMMenuContext;
-class UToolMenu;
 struct FToolMenuSection;
 
 class FDMMediaStreamStageSourceMenuExtender
@@ -16,9 +15,9 @@ public:
 private:
 	bool bIntegrated = false;
 
-	void ExtendMenu_ChangeSource(UToolMenu* InMenu);
+	void ExtendMenu_ChangeSource(FToolMenuSection& InSection);
 
-	void ExtendMenu_AddLayer(UToolMenu* InMenu);
+	void ExtendMenu_AddLayer(FToolMenuSection& InSection);
 
 	void ChangeSourceToMediaStreamFromContext(UDMMenuContext* InMenuContext);
 
