Material Designer: Fixed actor copy
- Copying actors was not copying script interface. Changed to UObject.

diff --git a/Engine/Plugins/Experimental/DynamicMaterial/Source/DynamicMaterial/Private/Model/DynamicMaterialModel.cpp b/Engine/Plugins/Experimental/DynamicMaterial/Source/DynamicMaterial/Private/Model/DynamicMaterialModel.cpp
--- a/Engine/Plugins/Experimental/DynamicMaterial/Source/DynamicMaterial/Private/Model/DynamicMaterialModel.cpp
+++ b/Engine/Plugins/Experimental/DynamicMaterial/Source/DynamicMaterial/Private/Model/DynamicMaterialModel.cpp
@@ -86,7 +86,7 @@ UDynamicMaterialModel::UDynamicMaterialModel()
 	DynamicMaterialInstance = nullptr;
 
 #if WITH_EDITOR
-	EditorOnlyDataSI.SetObject(nullptr);
+	EditorOnlyDataSI = nullptr;
 #endif
 
 	auto AddFloatParameter = [this](TObjectPtr<UDMMaterialValue>& InValuePtr, FName InPropertyName, FName InParameterName, float InDefaultValue = 1.f, bool bInSetRange = false)
@@ -364,14 +364,17 @@ UDMMaterialValue* UDynamicMaterialModel::GetValueByName(FName InName) const
 }
 
 #if WITH_EDITOR
-IDynamicMaterialModelEditorOnlyDataInterface* UDynamicMaterialModel::GetEditorOnlyData() const
+TScriptInterface<IDynamicMaterialModelEditorOnlyDataInterface> UDynamicMaterialModel::BP_GetEditorOnlyData() const
 {
-	if (IsValid(EditorOnlyDataSI.GetObject()))
-	{
-		return EditorOnlyDataSI.GetInterface();
-	}
+	TScriptInterface<IDynamicMaterialModelEditorOnlyDataInterface> Interface;
+	Interface.SetObject(EditorOnlyDataSI);
 
-	return nullptr;
+	return Interface;
+}
+
+IDynamicMaterialModelEditorOnlyDataInterface* UDynamicMaterialModel::GetEditorOnlyData() const
+{
+	return Cast<IDynamicMaterialModelEditorOnlyDataInterface>(EditorOnlyDataSI);
 }
 
 UDMMaterialValue* UDynamicMaterialModel::AddValue(TSubclassOf<UDMMaterialValue> InValueClass)
diff --git a/Engine/Plugins/Experimental/DynamicMaterial/Source/DynamicMaterial/Public/Model/DynamicMaterialModel.h b/Engine/Plugins/Experimental/DynamicMaterial/Source/DynamicMaterial/Public/Model/DynamicMaterialModel.h
--- a/Engine/Plugins/Experimental/DynamicMaterial/Source/DynamicMaterial/Public/Model/DynamicMaterialModel.h
+++ b/Engine/Plugins/Experimental/DynamicMaterial/Source/DynamicMaterial/Public/Model/DynamicMaterialModel.h
@@ -146,7 +146,7 @@ public:
 
 	/** Returns the editor only data for this model. */
 	UFUNCTION(BlueprintPure, Category = "Material Designer", Meta = (DisplayName = "Get Editor Only Data"))
-	TScriptInterface<IDynamicMaterialModelEditorOnlyDataInterface> BP_GetEditorOnlyData() const { return EditorOnlyDataSI; }
+	TScriptInterface<IDynamicMaterialModelEditorOnlyDataInterface> BP_GetEditorOnlyData() const;
 
 	/** Returns the editor only data for this model. */
 	DYNAMICMATERIAL_API IDynamicMaterialModelEditorOnlyDataInterface* GetEditorOnlyData() const;
@@ -246,8 +246,8 @@ protected:
 
 #if WITH_EDITORONLY_DATA
 	/** Object holding the editor-only data used by this model. */
-	UPROPERTY(VisibleInstanceOnly, BlueprintReadOnly, Category = "Material Designer")
-	TScriptInterface<IDynamicMaterialModelEditorOnlyDataInterface> EditorOnlyDataSI;
+	UPROPERTY(VisibleInstanceOnly, BlueprintReadOnly, Instanced, Category = "Material Designer")
+	TObjectPtr<UObject> EditorOnlyDataSI;
 #endif
 
 	UPROPERTY(BlueprintReadOnly, Instanced, Category = "Material Designer")
diff --git a/Engine/Plugins/Experimental/DynamicMaterial/Source/DynamicMaterialEditor/Private/Model/DynamicMateriaModelFactory.cpp b/Engine/Plugins/Experimental/DynamicMaterial/Source/DynamicMaterialEditor/Private/Model/DynamicMateriaModelFactory.cpp
--- a/Engine/Plugins/Experimental/DynamicMaterial/Source/DynamicMaterialEditor/Private/Model/DynamicMateriaModelFactory.cpp
+++ b/Engine/Plugins/Experimental/DynamicMaterial/Source/DynamicMaterialEditor/Private/Model/DynamicMateriaModelFactory.cpp
@@ -40,7 +40,7 @@ UObject* UDynamicMaterialModelFactory::FactoryCreateNew(UClass* InClass, UObject
 		RF_Transactional
 	);
 
-	NewModel->EditorOnlyDataSI = TScriptInterface<IDynamicMaterialModelEditorOnlyDataInterface>(ModelEditorOnlyData);
+	NewModel->EditorOnlyDataSI = ModelEditorOnlyData;
 	ModelEditorOnlyData->MaterialModel = NewModel;
 
 	const FDMInitializationGuard InitGuard;
