Material Designer: Apply button
- Materials edited in the designer are now preview materials which are applied to the source asset when clicking apply or the save buttons.

diff --git a/Engine/Plugins/Experimental/DynamicMaterial/Source/DynamicMaterialEditor/Private/UI/Menus/DMMenuContext.cpp b/Engine/Plugins/Experimental/DynamicMaterial/Source/DynamicMaterialEditor/Private/UI/Menus/DMMenuContext.cpp
--- a/Engine/Plugins/Experimental/DynamicMaterial/Source/DynamicMaterialEditor/Private/UI/Menus/DMMenuContext.cpp
+++ b/Engine/Plugins/Experimental/DynamicMaterial/Source/DynamicMaterialEditor/Private/UI/Menus/DMMenuContext.cpp
@@ -116,7 +116,7 @@ UDynamicMaterialModelBase* UDMMenuContext::GetModelBase() const
 {
 	if (TSharedPtr<SDMMaterialEditor> EditorWidget = EditorWidgetWeak.Pin())
 	{
-		return EditorWidget->GetMaterialModelBase();
+		return EditorWidget->GetPreviewMaterialModelBase();
 	}
 
 	return nullptr;
@@ -126,7 +126,7 @@ UDynamicMaterialModel* UDMMenuContext::GetModel() const
 {
 	if (TSharedPtr<SDMMaterialEditor> EditorWidget = EditorWidgetWeak.Pin())
 	{
-		return EditorWidget->GetMaterialModel();
+		return EditorWidget->GetPreviewMaterialModel();
 	}
 
 	return nullptr;
diff --git a/Engine/Plugins/Experimental/DynamicMaterial/Source/DynamicMaterialEditor/Private/UI/PropertyGenerators/DMComponentPropertyRowGenerator.cpp b/Engine/Plugins/Experimental/DynamicMaterial/Source/DynamicMaterialEditor/Private/UI/PropertyGenerators/DMComponentPropertyRowGenerator.cpp
--- a/Engine/Plugins/Experimental/DynamicMaterial/Source/DynamicMaterialEditor/Private/UI/PropertyGenerators/DMComponentPropertyRowGenerator.cpp
+++ b/Engine/Plugins/Experimental/DynamicMaterial/Source/DynamicMaterialEditor/Private/UI/PropertyGenerators/DMComponentPropertyRowGenerator.cpp
@@ -112,7 +112,7 @@ bool FDMComponentPropertyRowGenerator::IsDynamic(const TSharedRef<SDMMaterialCom
 {
 	if (TSharedPtr<SDMMaterialEditor> EditorWidget = InComponentEditorWidget->GetEditorWidget())
 	{
-		if (UDynamicMaterialModelBase* MaterialModelBase = EditorWidget->GetMaterialModelBase())
+		if (UDynamicMaterialModelBase* MaterialModelBase = EditorWidget->GetPreviewMaterialModelBase())
 		{
 			return !!Cast<UDynamicMaterialModelDynamic>(MaterialModelBase);
 		}
diff --git a/Engine/Plugins/Experimental/DynamicMaterial/Source/DynamicMaterialEditor/Private/UI/PropertyGenerators/DMMaterialValuePropertyRowGenerator.cpp b/Engine/Plugins/Experimental/DynamicMaterial/Source/DynamicMaterialEditor/Private/UI/PropertyGenerators/DMMaterialValuePropertyRowGenerator.cpp
--- a/Engine/Plugins/Experimental/DynamicMaterial/Source/DynamicMaterialEditor/Private/UI/PropertyGenerators/DMMaterialValuePropertyRowGenerator.cpp
+++ b/Engine/Plugins/Experimental/DynamicMaterial/Source/DynamicMaterialEditor/Private/UI/PropertyGenerators/DMMaterialValuePropertyRowGenerator.cpp
@@ -49,7 +49,7 @@ void FDMMaterialValuePropertyRowGenerator::AddComponentProperties(const TSharedR
 
 	if (TSharedPtr<SDMMaterialEditor> EditorWidget = InComponentEditorWidget->GetEditorWidget())
 	{
-		if (UDynamicMaterialModelBase* MaterialModelBase = EditorWidget->GetMaterialModelBase())
+		if (UDynamicMaterialModelBase* MaterialModelBase = EditorWidget->GetPreviewMaterialModelBase())
 		{
 			if (UDynamicMaterialModelDynamic* MaterialModelDynamic = Cast<UDynamicMaterialModelDynamic>(MaterialModelBase))
 			{
diff --git a/Engine/Plugins/Experimental/DynamicMaterial/Source/DynamicMaterialEditor/Private/UI/PropertyGenerators/DMTextureUVPropertyRowGenerator.cpp b/Engine/Plugins/Experimental/DynamicMaterial/Source/DynamicMaterialEditor/Private/UI/PropertyGenerators/DMTextureUVPropertyRowGenerator.cpp
--- a/Engine/Plugins/Experimental/DynamicMaterial/Source/DynamicMaterialEditor/Private/UI/PropertyGenerators/DMTextureUVPropertyRowGenerator.cpp
+++ b/Engine/Plugins/Experimental/DynamicMaterial/Source/DynamicMaterialEditor/Private/UI/PropertyGenerators/DMTextureUVPropertyRowGenerator.cpp
@@ -62,7 +62,7 @@ void FDMTextureUVPropertyRowGenerator::AddComponentProperties(const TSharedRef<S
 
 	if (TSharedPtr<SDMMaterialEditor> EditorWidget = InComponentEditorWidget->GetEditorWidget())
 	{
-		if (UDynamicMaterialModelBase* MaterialModelBase = EditorWidget->GetMaterialModelBase())
+		if (UDynamicMaterialModelBase* MaterialModelBase = EditorWidget->GetPreviewMaterialModelBase())
 		{
 			if (UDynamicMaterialModelDynamic* MaterialModelDynamic = Cast<UDynamicMaterialModelDynamic>(MaterialModelBase))
 			{
diff --git a/Engine/Plugins/Experimental/DynamicMaterial/Source/DynamicMaterialEditor/Private/UI/Widgets/Editor/SDMMaterialProperties.cpp b/Engine/Plugins/Experimental/DynamicMaterial/Source/DynamicMaterialEditor/Private/UI/Widgets/Editor/SDMMaterialProperties.cpp
--- a/Engine/Plugins/Experimental/DynamicMaterial/Source/DynamicMaterialEditor/Private/UI/Widgets/Editor/SDMMaterialProperties.cpp
+++ b/Engine/Plugins/Experimental/DynamicMaterial/Source/DynamicMaterialEditor/Private/UI/Widgets/Editor/SDMMaterialProperties.cpp
@@ -87,7 +87,7 @@ TSharedRef<SWidget> SDMMaterialProperties::CreateSlot_Content()
 	{
 		return SNullWidget::NullWidget;
 	}
-	UDynamicMaterialModelBase* MaterialModelBase = EditorWidget->GetMaterialModelBase();
+	UDynamicMaterialModelBase* MaterialModelBase = EditorWidget->GetPreviewMaterialModelBase();
 
 	if (!MaterialModelBase)
 	{
@@ -245,7 +245,7 @@ void SDMMaterialProperties::AddProperty(const TSharedRef<ICustomDetailsView>& In
 		return;
 	}
 
-	UDynamicMaterialModel* MaterialModel = EditorWidget->GetMaterialModel();
+	UDynamicMaterialModel* MaterialModel = EditorWidget->GetPreviewMaterialModel();
 
 	if (!MaterialModel)
 	{
@@ -275,7 +275,7 @@ TSharedRef<SWidget> SDMMaterialProperties::CreatePropertyRow(UDMMaterialProperty
 {
 	// There are all ensured to be valid by the caller of this method
 	TSharedPtr<SDMMaterialEditor> EditorWidget = EditorWidgetWeak.Pin();
-	UDynamicMaterialModelBase* MaterialModelBase = EditorWidget->GetMaterialModelBase();
+	UDynamicMaterialModelBase* MaterialModelBase = EditorWidget->GetPreviewMaterialModelBase();
 	UDynamicMaterialModelEditorOnlyData* EditorOnlyData = UDynamicMaterialModelEditorOnlyData::Get(MaterialModelBase);
 	UDynamicMaterialEditorSettings* Settings = UDynamicMaterialEditorSettings::Get();
 	const EDMMaterialPropertyType MaterialProperty = InProperty->GetMaterialProperty();
@@ -422,7 +422,7 @@ bool SDMMaterialProperties::GetPropertyEnabledEnabled(EDMMaterialPropertyType In
 		return false;
 	}
 
-	UDynamicMaterialModel* MaterialModel = EditorWidget->GetMaterialModel();
+	UDynamicMaterialModel* MaterialModel = EditorWidget->GetPreviewMaterialModel();
 
 	if (!MaterialModel)
 	{
@@ -455,7 +455,7 @@ ECheckBoxState SDMMaterialProperties::GetPropertyEnabledState(EDMMaterialPropert
 		return ECheckBoxState::Unchecked;
 	}
 
-	UDynamicMaterialModel* MaterialModel = EditorWidget->GetMaterialModel();
+	UDynamicMaterialModel* MaterialModel = EditorWidget->GetPreviewMaterialModel();
 
 	if (!MaterialModel)
 	{
@@ -490,7 +490,7 @@ void SDMMaterialProperties::OnPropertyEnabledStateChanged(ECheckBoxState InState
 		return;
 	}
 
-	UDynamicMaterialModel* MaterialModel = EditorWidget->GetMaterialModel();
+	UDynamicMaterialModel* MaterialModel = EditorWidget->GetPreviewMaterialModel();
 
 	if (!MaterialModel)
 	{
@@ -541,7 +541,7 @@ FReply SDMMaterialProperties::OnPropertyClicked(const FGeometry& InGeometry, con
 		return FReply::Handled();
 	}
 
-	UDynamicMaterialModel* MaterialModel = EditorWidget->GetMaterialModel();
+	UDynamicMaterialModel* MaterialModel = EditorWidget->GetPreviewMaterialModel();
 
 	if (!MaterialModel)
 	{
@@ -581,7 +581,7 @@ TSharedRef<SWidget> SDMMaterialProperties::CreateGlobalSlider(UDMMaterialPropert
 		return SNullWidget::NullWidget;
 	}
 
-	UDynamicMaterialModelBase* MaterialModelBase = EditorWidget->GetMaterialModelBase();
+	UDynamicMaterialModelBase* MaterialModelBase = EditorWidget->GetPreviewMaterialModelBase();
 
 	if (!MaterialModelBase)
 	{
@@ -717,7 +717,7 @@ void SDMMaterialProperties::OnExpansionStateChanged(const TSharedRef<ICustomDeta
 		return;
 	}
 
-	UDynamicMaterialModelBase* MaterialModelBase = EditorWidget->GetMaterialModelBase();
+	UDynamicMaterialModelBase* MaterialModelBase = EditorWidget->GetPreviewMaterialModelBase();
 
 	if (!MaterialModelBase)
 	{
@@ -767,7 +767,7 @@ bool SDMMaterialProperties::OnAssetDraggedOver(TArrayView<FAssetData> InAssets,
 		return false;
 	}
 
-	UDynamicMaterialModel* MaterialModel = EditorWidget->GetMaterialModel();
+	UDynamicMaterialModel* MaterialModel = EditorWidget->GetPreviewMaterialModel();
 
 	if (!MaterialModel)
 	{
@@ -844,7 +844,7 @@ void SDMMaterialProperties::HandleDrop_Texture(UTexture* InTexture, EDMMaterialP
 		return;
 	}
 
-	UDynamicMaterialModel* MaterialModel = EditorWidget->GetMaterialModel();
+	UDynamicMaterialModel* MaterialModel = EditorWidget->GetPreviewMaterialModel();
 
 	if (!MaterialModel)
 	{
diff --git a/Engine/Plugins/Experimental/DynamicMaterial/Source/DynamicMaterialEditor/Private/UI/Widgets/Editor/SDMMaterialPropertySelector.cpp b/Engine/Plugins/Experimental/DynamicMaterial/Source/DynamicMaterialEditor/Private/UI/Widgets/Editor/SDMMaterialPropertySelector.cpp
--- a/Engine/Plugins/Experimental/DynamicMaterial/Source/DynamicMaterialEditor/Private/UI/Widgets/Editor/SDMMaterialPropertySelector.cpp
+++ b/Engine/Plugins/Experimental/DynamicMaterial/Source/DynamicMaterialEditor/Private/UI/Widgets/Editor/SDMMaterialPropertySelector.cpp
@@ -108,7 +108,7 @@ UDynamicMaterialModelEditorOnlyData* SDMMaterialPropertySelector::GetEditorOnlyD
 		return nullptr;
 	}
 
-	UDynamicMaterialModel* MaterialModel = EditorWidget->GetMaterialModel();
+	UDynamicMaterialModel* MaterialModel = EditorWidget->GetPreviewMaterialModel();
 
 	if (!IsValid(MaterialModel))
 	{
diff --git a/Engine/Plugins/Experimental/DynamicMaterial/Source/DynamicMaterialEditor/Private/UI/Widgets/Editor/SDMMaterialSlotEditor.cpp b/Engine/Plugins/Experimental/DynamicMaterial/Source/DynamicMaterialEditor/Private/UI/Widgets/Editor/SDMMaterialSlotEditor.cpp
--- a/Engine/Plugins/Experimental/DynamicMaterial/Source/DynamicMaterialEditor/Private/UI/Widgets/Editor/SDMMaterialSlotEditor.cpp
+++ b/Engine/Plugins/Experimental/DynamicMaterial/Source/DynamicMaterialEditor/Private/UI/Widgets/Editor/SDMMaterialSlotEditor.cpp
@@ -81,7 +81,7 @@ void SDMMaterialSlotEditor::Construct(const FArguments& InArgs, const TSharedRef
 
 	SetCanTick(false);
 
-	bIsDynamic = !Cast<UDynamicMaterialModel>(InEditorWidget->GetMaterialModelBase());
+	bIsDynamic = !Cast<UDynamicMaterialModel>(InEditorWidget->GetPreviewMaterialModelBase());
 
 	ContentSlot = TDMWidgetSlot<SWidget>(SharedThis(this), 0, SNullWidget::NullWidget);
 
@@ -678,7 +678,7 @@ TSharedRef<SWidget> SDMMaterialSlotEditor::CreateSlot_LayerOpacity()
 						{
 							if (TSharedPtr<SDMMaterialEditor> EditorWidget = EditorWidgetWeak.Pin())
 							{
-								if (UDynamicMaterialModelBase* MaterialModelBase = EditorWidget->GetMaterialModelBase())
+								if (UDynamicMaterialModelBase* MaterialModelBase = EditorWidget->GetPreviewMaterialModelBase())
 								{
 									World = MaterialModelBase->GetWorld();
 								}
@@ -963,7 +963,7 @@ bool SDMMaterialSlotEditor::OnAreAssetsAcceptableForDrop(TArrayView<FAssetData>
 		return false;
 	}
 
-	UDynamicMaterialModelBase* MaterialModelBase = EditorWidget->GetMaterialModelBase();
+	UDynamicMaterialModelBase* MaterialModelBase = EditorWidget->GetPreviewMaterialModelBase();
 
 	if (MaterialModelBase->IsA<UDynamicMaterialModelDynamic>())
 	{
@@ -1023,7 +1023,7 @@ void SDMMaterialSlotEditor::OnAssetsDropped(const FDragDropEvent& InDragDropEven
 		return;
 	}
 
-	UDynamicMaterialModelBase* MaterialModelBase = EditorWidget->GetMaterialModelBase();
+	UDynamicMaterialModelBase* MaterialModelBase = EditorWidget->GetPreviewMaterialModelBase();
 
 	if (MaterialModelBase->IsA<UDynamicMaterialModelDynamic>())
 	{
diff --git a/Engine/Plugins/Experimental/DynamicMaterial/Source/DynamicMaterialEditor/Private/UI/Widgets/Editor/SDMObjectEditorWidgetBase.cpp b/Engine/Plugins/Experimental/DynamicMaterial/Source/DynamicMaterialEditor/Private/UI/Widgets/Editor/SDMObjectEditorWidgetBase.cpp
--- a/Engine/Plugins/Experimental/DynamicMaterial/Source/DynamicMaterialEditor/Private/UI/Widgets/Editor/SDMObjectEditorWidgetBase.cpp
+++ b/Engine/Plugins/Experimental/DynamicMaterial/Source/DynamicMaterialEditor/Private/UI/Widgets/Editor/SDMObjectEditorWidgetBase.cpp
@@ -42,7 +42,7 @@ void SDMObjectEditorWidgetBase::Construct(const FArguments& InArgs, const TShare
 
 	if (!WorldContext)
 	{
-		WorldContext = InEditorWidget->GetMaterialModelBase();
+		WorldContext = InEditorWidget->GetPreviewMaterialModelBase();
 	}
 
 	if (WorldContext)
diff --git a/Engine/Plugins/Experimental/DynamicMaterial/Source/DynamicMaterialEditor/Private/UI/Widgets/Editor/SDMStatusBar.cpp b/Engine/Plugins/Experimental/DynamicMaterial/Source/DynamicMaterialEditor/Private/UI/Widgets/Editor/SDMStatusBar.cpp
--- a/Engine/Plugins/Experimental/DynamicMaterial/Source/DynamicMaterialEditor/Private/UI/Widgets/Editor/SDMStatusBar.cpp
+++ b/Engine/Plugins/Experimental/DynamicMaterial/Source/DynamicMaterialEditor/Private/UI/Widgets/Editor/SDMStatusBar.cpp
@@ -26,7 +26,7 @@ SDMStatusBar::~SDMStatusBar()
 		return;
 	}
 
-	UDynamicMaterialModel* MaterialModel = EditorWidget->GetMaterialModel();
+	UDynamicMaterialModel* MaterialModel = EditorWidget->GetPreviewMaterialModel();
 
 	if (!MaterialModel)
 	{
@@ -51,7 +51,7 @@ void SDMStatusBar::Construct(const FArguments& InArgs, const TSharedRef<SDMMater
 
 	SetCanTick(false);
 
-	UDynamicMaterialModel* MaterialModel = InEditorWidget->GetMaterialModel();
+	UDynamicMaterialModel* MaterialModel = InEditorWidget->GetPreviewMaterialModel();
 
 	if (!MaterialModel)
 	{
diff --git a/Engine/Plugins/Experimental/DynamicMaterial/Source/DynamicMaterialEditor/Private/UI/Widgets/Editor/SDMToolBar.cpp b/Engine/Plugins/Experimental/DynamicMaterial/Source/DynamicMaterialEditor/Private/UI/Widgets/Editor/SDMToolBar.cpp
--- a/Engine/Plugins/Experimental/DynamicMaterial/Source/DynamicMaterialEditor/Private/UI/Widgets/Editor/SDMToolBar.cpp
+++ b/Engine/Plugins/Experimental/DynamicMaterial/Source/DynamicMaterialEditor/Private/UI/Widgets/Editor/SDMToolBar.cpp
@@ -16,7 +16,6 @@
 #include "Model/DynamicMaterialModel.h"
 #include "Model/DynamicMaterialModelBase.h"
 #include "Model/DynamicMaterialModelDynamic.h"
-#include "PackageTools.h"
 #include "Selection.h"
 #include "Styling/StyleColors.h"
 #include "UI/Menus/DMToolBarMenus.h"
@@ -68,7 +67,7 @@ void SDMToolBar::Construct(const FArguments& InArgs, const TSharedRef<SDMMateria
 		];
 
 	SetActorPropertySelected(InActor);
-	SetButtonVisibilities();
+	UpdateButtonVisibilities();
 }
 
 TSharedRef<SWidget> SDMToolBar::CreateToolBarEntries()
@@ -110,7 +109,7 @@ TSharedRef<SWidget> SDMToolBar::CreateToolBarEntries()
 						.Image(this, &SDMToolBar::GetSaveIcon)
 						.DesiredSizeOverride(LargeIconToolBarButtonSize)
 					]
-				]		
+				]
 				+ SHorizontalBox::Slot()
 				.AutoWidth()
 				.Padding(0.f)
@@ -129,6 +128,22 @@ TSharedRef<SWidget> SDMToolBar::CreateToolBarEntries()
 				]
 				+ SHorizontalBox::Slot()
 				.AutoWidth()
+				.Padding(0.f)
+				.VAlign(EVerticalAlignment::VAlign_Center)
+				[
+					SAssignNew(ApplyButtonWidget, SButton)
+					.ContentPadding(LargeIconToolBarButtonContentPadding)
+					.ButtonStyle(FDynamicMaterialEditorStyle::Get(), "HoverHintOnly")
+					.ToolTipText(LOCTEXT("MaterialDesignerSaveTooltip", "Apply the preview Material Designer Material back to the original asset.."))
+					.OnClicked(this, &SDMToolBar::OnApplyClicked)
+					[
+						SNew(STextBlock)
+						.TextStyle(FDynamicMaterialEditorStyle::Get(), "RegularFont")
+						.Text(LOCTEXT("Apply", "Apply"))
+					]
+				]	
+				+ SHorizontalBox::Slot()
+				.AutoWidth()
 				.VAlign(VAlign_Center)
 				.Padding(0.f)
 				[
@@ -325,7 +340,7 @@ void SDMToolBar::SetActorPropertySelected(AActor* InActor)
 		ActorRowWidget->SetVisibility(EVisibility::Visible);
 
 		TArray<FDMObjectMaterialProperty> ActorProperties = UDMMaterialInstanceFunctionLibrary::GetActorMaterialProperties(InActor);
-		UDynamicMaterialModelBase* MaterialModelBase = GetMaterialModelBase();
+		UDynamicMaterialModelBase* MaterialModelBase = GetOriginalMaterialModelBase();
 
 		const int32 ActorPropertyCount = ActorProperties.Num();
 		ActorMaterialProperties.Empty(ActorPropertyCount);
@@ -352,12 +367,12 @@ void SDMToolBar::SetActorPropertySelected(AActor* InActor)
 	PropertySelectorContainer->SetContent(CreateSlotsComboBoxWidget());
 }
 
-void SDMToolBar::SetButtonVisibilities()
+void SDMToolBar::UpdateButtonVisibilities()
 {
 	bool bIsAsset = false;
 	bool bIsDynamic = false;
 
-	UDynamicMaterialModelBase* MaterialModelBase = GetMaterialModelBase();
+	UDynamicMaterialModelBase* MaterialModelBase = GetOriginalMaterialModelBase();
 
 	if (IsValid(MaterialModelBase))
 	{
@@ -441,7 +456,7 @@ TSharedRef<SWidget> SDMToolBar::CreateToolBarButton(TAttribute<const FSlateBrush
 
 TSharedRef<SWidget> SDMToolBar::CreateSlotsComboBoxWidget()
 {
-	UDynamicMaterialModelBase* MaterialModelBase = GetMaterialModelBase();
+	UDynamicMaterialModelBase* MaterialModelBase = GetOriginalMaterialModelBase();
 
 	if (!MaterialActorWeak.IsValid() || !IsValid(MaterialModelBase))
 	{
@@ -527,11 +542,11 @@ void SDMToolBar::OnMaterialSlotChanged(TSharedPtr<FDMObjectMaterialProperty> InS
 	}
 }
 
-UDynamicMaterialModelBase* SDMToolBar::GetMaterialModelBase() const
+UDynamicMaterialModelBase* SDMToolBar::GetOriginalMaterialModelBase() const
 {
 	if (TSharedPtr<SDMMaterialEditor> EditorWidget = GetEditorWidget())
 	{
-		return EditorWidget->GetMaterialModelBase();
+		return EditorWidget->GetOriginalMaterialModelBase();
 	}
 
 	return nullptr;
@@ -591,7 +606,7 @@ FReply SDMToolBar::OnExportMaterialInstanceButtonClicked()
 		return FReply::Handled();
 	}
 
-	UDynamicMaterialModelBase* MaterialModelBase = GetMaterialModelBase();
+	UDynamicMaterialModelBase* MaterialModelBase = GetOriginalMaterialModelBase();
 
 	if (!MaterialModelBase)
 	{
@@ -612,7 +627,7 @@ FReply SDMToolBar::OnExportMaterialInstanceButtonClicked()
 
 FReply SDMToolBar::OnBrowseClicked()
 {
-	UDynamicMaterialModelBase* MaterialModelBase = GetMaterialModelBase();
+	UDynamicMaterialModelBase* MaterialModelBase = GetOriginalMaterialModelBase();
 
 	if (!IsValid(MaterialModelBase))
 	{
@@ -652,7 +667,7 @@ FReply SDMToolBar::OnUseClicked()
 		return FReply::Handled();
 	}
 
-	UDynamicMaterialModelBase* CurrentModelBase = GetMaterialModelBase();
+	UDynamicMaterialModelBase* CurrentModelBase = GetOriginalMaterialModelBase();
 	UDMWorldSubsystem* DMSubsystem = nullptr;
 
 	if (CurrentModelBase)
@@ -727,23 +742,6 @@ FReply SDMToolBar::OnUseClicked()
 	return FReply::Handled();
 }
 
-UPackage* SDMToolBar::GetSaveablePackage(UObject* InObject)
-{
-	if (!IsValid(InObject))
-	{
-		return nullptr;
-	}
-
-	UPackage* Package = InObject->GetPackage();
-
-	if (!Package || Package->HasAllFlags(RF_Transient))
-	{
-		return nullptr;
-	}
-
-	return Package;
-}
-
 FText SDMToolBar::GetActorName() const
 {
 	if (const AActor* const SlotActor = GetMaterialActor())
@@ -761,7 +759,7 @@ TSharedPtr<SDMMaterialEditor> SDMToolBar::GetEditorWidget() const
 
 FText SDMToolBar::GetAssetName() const
 {
-	if (UDynamicMaterialModelBase* MaterialModelBase = GetMaterialModelBase())
+	if (UDynamicMaterialModelBase* MaterialModelBase = GetOriginalMaterialModelBase())
 	{
 		if (UDynamicMaterialInstance* MaterialInstance = MaterialModelBase->GetDynamicMaterialInstance())
 		{
@@ -781,7 +779,7 @@ FText SDMToolBar::GetAssetName() const
 
 FText SDMToolBar::GetAssetToolTip() const
 {
-	if (UDynamicMaterialModelBase* MaterialModelBase = GetMaterialModelBase())
+	if (UDynamicMaterialModelBase* MaterialModelBase = GetOriginalMaterialModelBase())
 	{
 		if (UDynamicMaterialInstance* MaterialInstance = MaterialModelBase->GetDynamicMaterialInstance())
 		{
@@ -801,7 +799,7 @@ FText SDMToolBar::GetAssetToolTip() const
 
 const FSlateBrush* SDMToolBar::GetSaveIcon() const
 {
-	if (UPackage* Package = GetSaveablePackage(GetMaterialModelBase()))
+	if (UPackage* Package = SDMMaterialEditor::GetSaveablePackage(GetOriginalMaterialModelBase()))
 	{
 		if (Package->IsDirty())
 		{
@@ -814,14 +812,19 @@ const FSlateBrush* SDMToolBar::GetSaveIcon() const
 
 FReply SDMToolBar::OnSaveClicked()
 {
-	if (UDynamicMaterialModelBase* MaterialModelBase = GetMaterialModelBase())
+	if (TSharedPtr<SDMMaterialEditor> EditorWidget = GetEditorWidget())
 	{
-		if (UPackage* Package = GetSaveablePackage(MaterialModelBase))
-		{
-			TArray<UObject*> AssetsToSave;
-			AssetsToSave.Add(MaterialModelBase);
-			UPackageTools::SavePackagesForObjects(AssetsToSave);
-		}
+		EditorWidget->SaveOriginal();
+	}
+
+	return FReply::Handled();
+}
+
+FReply SDMToolBar::OnApplyClicked()
+{
+	if (TSharedPtr<SDMMaterialEditor> EditorWidget = GetEditorWidget())
+	{
+		EditorWidget->ApplyToOriginal();
 	}
 
 	return FReply::Handled();
@@ -833,7 +836,7 @@ FReply SDMToolBar::OnOpenParentClicked()
 	{
 		if (TSharedPtr<SDMMaterialDesigner> DesignerWidget = EditorWidget->GetDesignerWidget())
 		{
-			if (UDynamicMaterialModelDynamic* DynamicMaterialModel = Cast<UDynamicMaterialModelDynamic>(EditorWidget->GetMaterialModelBase()))
+			if (UDynamicMaterialModelDynamic* DynamicMaterialModel = Cast<UDynamicMaterialModelDynamic>(GetOriginalMaterialModelBase()))
 			{
 				if (UDynamicMaterialModel* ParentModel = DynamicMaterialModel->ResolveMaterialModel())
 				{
@@ -848,7 +851,7 @@ FReply SDMToolBar::OnOpenParentClicked()
 
 FReply SDMToolBar::OnConvertToEditableClicked()
 {
-	UDynamicMaterialModelDynamic* CurrentModelDynamic = Cast<UDynamicMaterialModelDynamic>(GetMaterialModelBase());
+	UDynamicMaterialModelDynamic* CurrentModelDynamic = Cast<UDynamicMaterialModelDynamic>(GetOriginalMaterialModelBase());
 
 	if (!CurrentModelDynamic)
 	{
@@ -948,7 +951,7 @@ FReply SDMToolBar::OnConvertToEditableClicked()
 
 FReply SDMToolBar::OnConvertToInstanceClicked()
 {
-	UDynamicMaterialModel* CurrentModel = Cast<UDynamicMaterialModel>(GetMaterialModelBase());
+	UDynamicMaterialModel* CurrentModel = Cast<UDynamicMaterialModel>(GetOriginalMaterialModelBase());
 
 	if (!CurrentModel)
 	{
diff --git a/Engine/Plugins/Experimental/DynamicMaterial/Source/DynamicMaterialEditor/Private/UI/Widgets/Editor/SDMToolBar.h b/Engine/Plugins/Experimental/DynamicMaterial/Source/DynamicMaterialEditor/Private/UI/Widgets/Editor/SDMToolBar.h
--- a/Engine/Plugins/Experimental/DynamicMaterial/Source/DynamicMaterialEditor/Private/UI/Widgets/Editor/SDMToolBar.h
+++ b/Engine/Plugins/Experimental/DynamicMaterial/Source/DynamicMaterialEditor/Private/UI/Widgets/Editor/SDMToolBar.h
@@ -47,8 +47,6 @@ public:
 	TSharedPtr<SDMMaterialEditor> GetEditorWidget() const;
 
 protected:
-	static UPackage* GetSaveablePackage(UObject* InObject);
-
 	TWeakPtr<SDMMaterialEditor> EditorWidgetWeak;
 	TWeakObjectPtr<AActor> MaterialActorWeak;
 
@@ -57,6 +55,7 @@ protected:
 
 	TSharedPtr<SBox> PropertySelectorContainer;
 	TSharedPtr<SWidget> SaveButtonWidget;
+	TSharedPtr<SWidget> ApplyButtonWidget;
 	TSharedPtr<SWidget> ActorRowWidget;
 	TSharedPtr<SWidget> AssetRowWidget;
 	TSharedPtr<STextBlock> ActorNameWidget;
@@ -66,13 +65,13 @@ protected:
 	TSharedPtr<SWidget> ConvertToEditableButton;
 	TSharedPtr<SWidget> ConvertToInstanceButton;
 
-	UDynamicMaterialModelBase* GetMaterialModelBase() const;
+	UDynamicMaterialModelBase* GetOriginalMaterialModelBase() const;
 
 	TSharedRef<SWidget> CreateToolBarEntries();
 
 	void SetActorPropertySelected(AActor* InActor);
 
-	void SetButtonVisibilities();
+	void UpdateButtonVisibilities();
 
 	TSharedRef<SWidget> CreateToolBarButton(TAttribute<const FSlateBrush*> InImageBrush, const TAttribute<FText>& InTooltipText, FOnClicked InOnClicked);
 
@@ -101,6 +100,8 @@ protected:
 
 	FReply OnSaveClicked();
 
+	FReply OnApplyClicked();
+
 	FReply OnOpenParentClicked();
 
 	FReply OnConvertToEditableClicked();
diff --git a/Engine/Plugins/Experimental/DynamicMaterial/Source/DynamicMaterialEditor/Private/UI/Widgets/Editor/SlotEditor/SDMMaterialSlotLayerEffectItem.cpp b/Engine/Plugins/Experimental/DynamicMaterial/Source/DynamicMaterialEditor/Private/UI/Widgets/Editor/SlotEditor/SDMMaterialSlotLayerEffectItem.cpp
--- a/Engine/Plugins/Experimental/DynamicMaterial/Source/DynamicMaterialEditor/Private/UI/Widgets/Editor/SlotEditor/SDMMaterialSlotLayerEffectItem.cpp
+++ b/Engine/Plugins/Experimental/DynamicMaterial/Source/DynamicMaterialEditor/Private/UI/Widgets/Editor/SlotEditor/SDMMaterialSlotLayerEffectItem.cpp
@@ -175,7 +175,7 @@ bool SDMMaterialSlotLayerEffectItem::CanModifyMaterialModel() const
 				{
 					if (TSharedPtr<SDMMaterialEditor> EditorWidget = SlotEditorWidget->GetEditorWidget())
 					{
-						if (UDynamicMaterialModelBase* MaterialModelBase = EditorWidget->GetMaterialModelBase())
+						if (UDynamicMaterialModelBase* MaterialModelBase = EditorWidget->GetPreviewMaterialModelBase())
 						{
 							return MaterialModelBase->IsA<UDynamicMaterialModel>();
 						}
diff --git a/Engine/Plugins/Experimental/DynamicMaterial/Source/DynamicMaterialEditor/Private/UI/Widgets/SDMMaterialDesigner.cpp b/Engine/Plugins/Experimental/DynamicMaterial/Source/DynamicMaterialEditor/Private/UI/Widgets/SDMMaterialDesigner.cpp
--- a/Engine/Plugins/Experimental/DynamicMaterial/Source/DynamicMaterialEditor/Private/UI/Widgets/SDMMaterialDesigner.cpp
+++ b/Engine/Plugins/Experimental/DynamicMaterial/Source/DynamicMaterialEditor/Private/UI/Widgets/SDMMaterialDesigner.cpp
@@ -155,7 +155,7 @@ UDynamicMaterialModelBase* SDMMaterialDesigner::GetMaterialModelBase() const
 
 	if (Content->GetWidgetClass().GetWidgetType() == SDMMaterialEditor::StaticWidgetClass().GetWidgetType())
 	{
-		return StaticCastSharedPtr<SDMMaterialEditor>(Content)->GetMaterialModelBase();
+		return StaticCastSharedPtr<SDMMaterialEditor>(Content)->GetPreviewMaterialModelBase();
 	}
 
 	if (Content->GetWidgetClass().GetWidgetType() == SDMMaterialWizard::StaticWidgetClass().GetWidgetType())
@@ -502,7 +502,7 @@ void SDMMaterialDesigner::OnLayoutChanged()
 	// Already assured
 	TSharedPtr<SDMMaterialEditor> CurrentEditor = StaticCastSharedRef<SDMMaterialEditor>(Content.ToSharedRef());
 
-	UDynamicMaterialModelBase* MaterialModelBase = CurrentEditor->GetMaterialModelBase();
+	UDynamicMaterialModelBase* MaterialModelBase = CurrentEditor->GetPreviewMaterialModelBase();
 	const FDMObjectMaterialProperty* MaterialObjectProperty = CurrentEditor->GetMaterialObjectProperty();
 	EDMMaterialEditorMode EditorMode = CurrentEditor->GetEditMode();
 	EDMMaterialPropertyType SelectedProperty = CurrentEditor->GetSelectedPropertyType();
diff --git a/Engine/Plugins/Experimental/DynamicMaterial/Source/DynamicMaterialEditor/Private/UI/Widgets/SDMMaterialEditor.cpp b/Engine/Plugins/Experimental/DynamicMaterial/Source/DynamicMaterialEditor/Private/UI/Widgets/SDMMaterialEditor.cpp
--- a/Engine/Plugins/Experimental/DynamicMaterial/Source/DynamicMaterialEditor/Private/UI/Widgets/SDMMaterialEditor.cpp
+++ b/Engine/Plugins/Experimental/DynamicMaterial/Source/DynamicMaterialEditor/Private/UI/Widgets/SDMMaterialEditor.cpp
@@ -16,6 +16,8 @@
 #include "Framework/Application/IInputProcessor.h"
 #include "Framework/Application/SlateApplication.h"
 #include "Framework/Commands/GenericCommands.h"
+#include "Material/DynamicMaterialInstance.h"
+#include "Material/DynamicMaterialInstanceFactory.h"
 #include "Materials/Material.h"
 #include "Materials/MaterialFunctionInterface.h"
 #include "Misc/MessageDialog.h"
@@ -23,7 +25,10 @@
 #include "Model/DynamicMaterialModelBase.h"
 #include "Model/DynamicMaterialModelDynamic.h"
 #include "Model/DynamicMaterialModelEditorOnlyData.h"
+#include "PackageTools.h"
+#include "Selection.h"
 #include "Styling/SlateIconFinder.h"
+#include "UObject/Linker.h"
 #include "UI/Utils/DMEditorSelectionContext.h"
 #include "UI/Utils/DMPreviewMaterialManager.h"
 #include "UI/Widgets/Editor/SDMMaterialComponentEditor.h"
@@ -196,18 +201,51 @@ TSharedPtr<SDMMaterialDesigner> SDMMaterialEditor::GetDesignerWidget() const
 	return DesignerWidgetWeak.Pin();
 }
 
-UDynamicMaterialModelBase* SDMMaterialEditor::GetMaterialModelBase() const
+UDynamicMaterialModelBase* SDMMaterialEditor::GetOriginalMaterialModelBase() const
 {
-	return MaterialModelBaseWeak.Get();
+	return OriginalMaterialModelBaseWeak.Get();
+}
+
+UDynamicMaterialModelBase* SDMMaterialEditor::GetPreviewMaterialModelBase() const
+{
+	return PreviewMaterialModelBase.Get();
 }
 
 void SDMMaterialEditor::SetMaterialModelBase(UDynamicMaterialModelBase* InMaterialModelBase)
 {
-	MaterialModelBaseWeak = InMaterialModelBase;
+	OriginalMaterialModelBaseWeak = InMaterialModelBase;
 
-	if (UDynamicMaterialModelDynamic* MaterialModelDynamic = Cast<UDynamicMaterialModelDynamic>(InMaterialModelBase))
+	if (!InMaterialModelBase)
 	{
-		MaterialModelDynamic->EnsureComponents();
+		PreviewMaterialModelBase = nullptr;
+	}
+	else
+	{
+		PreviewMaterialModelBase = Cast<UDynamicMaterialModelBase>(StaticDuplicateObject(InMaterialModelBase, GetTransientPackage()));
+		PreviewMaterialModelBase->SetDynamicMaterialInstance(nullptr);
+
+		if (UDynamicMaterialModelDynamic* MaterialModelDynamic = Cast<UDynamicMaterialModelDynamic>(PreviewMaterialModelBase))
+		{
+			MaterialModelDynamic->EnsureComponents();
+		}
+
+		if (UDynamicMaterialModelEditorOnlyData* EditorOnlyData = UDynamicMaterialModelEditorOnlyData::Get(PreviewMaterialModelBase.Get()))
+		{
+			EditorOnlyData->RequestMaterialBuild(/* Immediate */ true);
+		}
+
+		UDynamicMaterialInstanceFactory* Factory = NewObject<UDynamicMaterialInstanceFactory>();
+
+		UDynamicMaterialInstance* NewInstance = Cast<UDynamicMaterialInstance>(Factory->FactoryCreateNew(
+			UDynamicMaterialInstance::StaticClass(),
+			GetTransientPackage(),
+			NAME_None,
+			RF_Transactional,
+			PreviewMaterialModelBase,
+			nullptr
+		));
+
+		PreviewMaterialInstance = NewInstance;
 	}
 
 	EditGlobalSettings();
@@ -217,9 +255,9 @@ void SDMMaterialEditor::SetMaterialModelBase(UDynamicMaterialModelBase* InMateri
 	BindEditorOnlyDataUpdate(InMaterialModelBase);
 }
 
-UDynamicMaterialModel* SDMMaterialEditor::GetMaterialModel() const
+UDynamicMaterialModel* SDMMaterialEditor::GetPreviewMaterialModel() const
 {
-	if (UDynamicMaterialModelBase* MaterialModelBase = MaterialModelBaseWeak.Get())
+	if (UDynamicMaterialModelBase* MaterialModelBase = GetPreviewMaterialModelBase())
 	{
 		return MaterialModelBase->ResolveMaterialModel();
 	}
@@ -229,7 +267,7 @@ UDynamicMaterialModel* SDMMaterialEditor::GetMaterialModel() const
 
 bool SDMMaterialEditor::IsDynamicModel() const
 {
-	return !!Cast<UDynamicMaterialModelDynamic>(MaterialModelBaseWeak.Get());
+	return !!Cast<UDynamicMaterialModelDynamic>(GetPreviewMaterialModelBase());
 }
 
 const FDMObjectMaterialProperty* SDMMaterialEditor::GetMaterialObjectProperty() const
@@ -323,7 +361,7 @@ void SDMMaterialEditor::SelectProperty(EDMMaterialPropertyType InProperty, bool
 
 	PageHistoryAdd({EDMMaterialEditorMode::EditSlot, InProperty});
 
-	if (UDynamicMaterialModelEditorOnlyData* EditorOnlyData = UDynamicMaterialModelEditorOnlyData::Get(MaterialModelBaseWeak))
+	if (UDynamicMaterialModelEditorOnlyData* EditorOnlyData = UDynamicMaterialModelEditorOnlyData::Get(GetPreviewMaterialModelBase()))
 	{
 		UDMMaterialSlot* CurrentSlot = SelectionContext.Slot.Get();
 
@@ -438,7 +476,7 @@ void SDMMaterialEditor::OnEffectSelected(const TSharedRef<SDMMaterialSlotLayerEf
 
 void SDMMaterialEditor::OpenMaterialPreviewTab()
 {
-	UDynamicMaterialModelBase* MaterialModelBase = GetMaterialModelBase();
+	UDynamicMaterialModelBase* MaterialModelBase = GetPreviewMaterialModelBase();
 
 	if (!MaterialModelBase)
 	{
@@ -504,7 +542,7 @@ void SDMMaterialEditor::CloseMaterialPreviewTab()
 
 TSharedPtr<IToolTip> SDMMaterialEditor::GetMaterialPreviewToolTip()
 {
-	UDynamicMaterialModelBase* MaterialModelBase = GetMaterialModelBase();
+	UDynamicMaterialModelBase* MaterialModelBase = GetPreviewMaterialModelBase();
 
 	if (!MaterialModelBase)
 	{
@@ -569,7 +607,7 @@ void SDMMaterialEditor::Validate()
 		return;
 	}
 
-	UDynamicMaterialModelBase* MaterialModelBase = GetMaterialModelBase();
+	UDynamicMaterialModelBase* MaterialModelBase = GetPreviewMaterialModelBase();
 
 	if (!IsValid(MaterialModelBase))
 	{
@@ -671,6 +709,41 @@ void SDMMaterialEditor::PostRedo(bool bInSuccess)
 	OnUndo();
 }
 
+void SDMMaterialEditor::AddReferencedObjects(FReferenceCollector& Collector)
+{
+	if (PreviewMaterialInstance)
+	{
+		Collector.AddReferencedObject(PreviewMaterialInstance);
+	}
+
+	if (PreviewMaterialModelBase)
+	{
+		Collector.AddReferencedObject(PreviewMaterialModelBase);
+	}
+}
+
+FString SDMMaterialEditor::GetReferencerName() const
+{
+	return TEXT("SDMMaterialEditor");
+}
+
+UPackage* SDMMaterialEditor::GetSaveablePackage(UObject* InObject)
+{
+	if (!IsValid(InObject))
+	{
+		return nullptr;
+	}
+
+	UPackage* Package = InObject->GetPackage();
+
+	if (!Package || Package->HasAllFlags(RF_Transient))
+	{
+		return nullptr;
+	}
+
+	return Package;
+}
+
 void SDMMaterialEditor::BindCommands(SDMMaterialSlotEditor* InSlotEditor)
 {
 	const FGenericCommands& GenericCommands = FGenericCommands::Get();
@@ -756,7 +829,7 @@ void SDMMaterialEditor::BindCommands(SDMMaterialSlotEditor* InSlotEditor)
 
 bool SDMMaterialEditor::IsPropertyValidForModel(EDMMaterialPropertyType InProperty) const
 {
-	UDynamicMaterialModelEditorOnlyData* EditorOnlyData = UDynamicMaterialModelEditorOnlyData::Get(MaterialModelBaseWeak);
+	UDynamicMaterialModelEditorOnlyData* EditorOnlyData = UDynamicMaterialModelEditorOnlyData::Get(GetPreviewMaterialModelBase());
 
 	if (!EditorOnlyData)
 	{
@@ -913,6 +986,57 @@ void SDMMaterialEditor::PageHistoryAdd(const FDMMaterialEditorPage& InPage)
 	SelectionContext.PageHistoryCount = NewPageIndex + 1;
 }
 
+void SDMMaterialEditor::ApplyToOriginal()
+{
+	UDynamicMaterialModelBase* Copy = GetPreviewMaterialModelBase();
+
+	if (!Copy)
+	{
+		return;
+	}
+
+	UDynamicMaterialModelBase* Original = GetOriginalMaterialModelBase();
+
+	if (!Original)
+	{
+		return;
+	}
+
+	if (Original->IsSelected())
+	{
+		GEditor->GetSelectedObjects()->Deselect(Original);
+	}
+
+	ResetLoaders(Original->GetPackage());
+
+	UDynamicMaterialInstance* Outer = Cast<UDynamicMaterialInstance>(Original->GetOuter());
+
+	Original = Cast<UDynamicMaterialModelBase>(StaticDuplicateObject(
+		Copy, 
+		Original->GetOuter(),
+		Original->GetFName(),
+		RF_AllFlags, 
+		Original->GetClass()
+	));
+
+	Original->SetDynamicMaterialInstance(Outer);
+
+	if (Outer)
+	{
+		Outer->SetMaterialModel(Original);
+	}
+
+	if (UDynamicMaterialModelEditorOnlyData* EditorOnlyData = UDynamicMaterialModelEditorOnlyData::Get(Original))
+	{
+		EditorOnlyData->RequestMaterialBuild(/* Immediate */ true);
+	}
+
+	if (Outer)
+	{
+		Outer->InitializeMIDPublic();
+	}
+}
+
 bool SDMMaterialEditor::SetActivePage(const FDMMaterialEditorPage& InPage)
 {
 	switch (InPage.EditorMode)
@@ -936,6 +1060,21 @@ bool SDMMaterialEditor::SetActivePage(const FDMMaterialEditorPage& InPage)
 	}
 }
 
+void SDMMaterialEditor::SaveOriginal()
+{
+	if (UDynamicMaterialModelBase* MaterialModelBase = GetOriginalMaterialModelBase())
+	{
+		ApplyToOriginal();
+
+		if (!!GetSaveablePackage(MaterialModelBase))
+		{
+			TArray<UObject*> AssetsToSave;
+			AssetsToSave.Add(MaterialModelBase);
+			UPackageTools::SavePackagesForObjects(AssetsToSave);
+		}
+	}
+}
+
 void SDMMaterialEditor::HandleDrop_CreateTextureSet(const TArray<FAssetData>& InTextureAssets)
 {
 	if (InTextureAssets.Num() < 2)
@@ -965,7 +1104,7 @@ void SDMMaterialEditor::HandleDrop_TextureSet(UDMTextureSet* InTextureSet)
 		return;
 	}
 
-	UDynamicMaterialModel* MaterialModel = GetMaterialModel();
+	UDynamicMaterialModel* MaterialModel = GetPreviewMaterialModel();
 
 	if (!MaterialModel)
 	{
@@ -1104,7 +1243,7 @@ TSharedRef<SDMToolBar> SDMMaterialEditor::CreateSlot_ToolBar()
 
 TSharedRef<SDMMaterialGlobalSettingsEditor> SDMMaterialEditor::CreateSlot_GlobalSettingsEditor()
 {
-	return SNew(SDMMaterialGlobalSettingsEditor, SharedThis(this), GetMaterialModelBase());
+	return SNew(SDMMaterialGlobalSettingsEditor, SharedThis(this), GetPreviewMaterialModelBase());
 }
 
 TSharedRef<SDMMaterialProperties> SDMMaterialEditor::CreateSlot_MaterialProperties()
@@ -1117,7 +1256,7 @@ TSharedRef<SWidget> SDMMaterialEditor::CreateSlot_Preview()
 	return SNew(SOverlay)
 		+ SOverlay::Slot()
 		[
-			SNew(SDMMaterialPreview, SharedThis(this), GetMaterialModelBase())
+			SNew(SDMMaterialPreview, SharedThis(this), GetPreviewMaterialModelBase())
 		]
 		+ SOverlay::Slot()
 		.HAlign(EHorizontalAlignment::HAlign_Left)
@@ -1140,7 +1279,7 @@ TSharedRef<SDMMaterialPropertySelector> SDMMaterialEditor::CreateSlot_PropertySe
 
 	if (SelectionContext.EditorMode == EDMMaterialEditorMode::EditSlot && SelectionContext.Property == EDMMaterialPropertyType::None)
 	{
-		if (UDynamicMaterialModel* MaterialModel = GetMaterialModel())
+		if (UDynamicMaterialModel* MaterialModel = GetPreviewMaterialModel())
 		{
 			if (UDynamicMaterialModelEditorOnlyData* EditorOnlyData = UDynamicMaterialModelEditorOnlyData::Get(MaterialModel))
 			{
@@ -1194,12 +1333,12 @@ TSharedRef<SDMMaterialComponentEditor> SDMMaterialEditor::CreateSlot_ComponentEd
 
 TSharedRef<SDMStatusBar> SDMMaterialEditor::CreateSlot_StatusBar()
 {
-	return SNew(SDMStatusBar, SharedThis(this), GetMaterialModelBase());
+	return SNew(SDMStatusBar, SharedThis(this), GetPreviewMaterialModelBase());
 }
 
 void SDMMaterialEditor::OnUndo()
 {
-	UDynamicMaterialModelBase* MaterialModel = GetMaterialModelBase();
+	UDynamicMaterialModelBase* MaterialModel = GetPreviewMaterialModelBase();
 
 	if (!IsValid(MaterialModel))
 	{
@@ -1233,7 +1372,7 @@ void SDMMaterialEditor::ApplySelectionContext()
 
 	if (SelectionContext.EditorMode == EDMMaterialEditorMode::EditSlot)
 	{
-		if (UDynamicMaterialModelEditorOnlyData* EditorOnlyData = UDynamicMaterialModelEditorOnlyData::Get(MaterialModelBaseWeak))
+		if (UDynamicMaterialModelEditorOnlyData* EditorOnlyData = UDynamicMaterialModelEditorOnlyData::Get(GetPreviewMaterialModelBase()))
 		{
 			if (UDMMaterialSlot* Slot = SelectionContext.Slot.Get())
 			{
diff --git a/Engine/Plugins/Experimental/DynamicMaterial/Source/DynamicMaterialEditor/Private/UI/Widgets/SDMMaterialEditor.h b/Engine/Plugins/Experimental/DynamicMaterial/Source/DynamicMaterialEditor/Private/UI/Widgets/SDMMaterialEditor.h
--- a/Engine/Plugins/Experimental/DynamicMaterial/Source/DynamicMaterialEditor/Private/UI/Widgets/SDMMaterialEditor.h
+++ b/Engine/Plugins/Experimental/DynamicMaterial/Source/DynamicMaterialEditor/Private/UI/Widgets/SDMMaterialEditor.h
@@ -51,7 +51,7 @@ namespace UE::DynamicMaterialEditor::Private
 	inline const TCHAR* EditorLightBackground = TEXT("Brushes.Header");
 }
 
-class SDMMaterialEditor : public SCompoundWidget, public FSelfRegisteringEditorUndoClient
+class SDMMaterialEditor : public SCompoundWidget, public FSelfRegisteringEditorUndoClient, public FGCObject
 {
 	SLATE_DECLARE_WIDGET(SDMMaterialEditor, SCompoundWidget)
 
@@ -64,6 +64,8 @@ class SDMMaterialEditor : public SCompoundWidget, public FSelfRegisteringEditorU
 	SLATE_END_ARGS()
 
 public:
+	static UPackage* GetSaveablePackage(UObject* InObject);
+
 	SDMMaterialEditor();
 
 	virtual ~SDMMaterialEditor() override;
@@ -73,9 +75,11 @@ public:
 	TSharedPtr<SDMMaterialDesigner> GetDesignerWidget() const;
 
 	/** Material Model / Actor */
-	UDynamicMaterialModelBase* GetMaterialModelBase() const;
+	UDynamicMaterialModelBase* GetOriginalMaterialModelBase() const;
+
+	UDynamicMaterialModelBase* GetPreviewMaterialModelBase() const;
 
-	UDynamicMaterialModel* GetMaterialModel() const;
+	UDynamicMaterialModel* GetPreviewMaterialModel() const;
 
 	bool IsDynamicModel() const;
 
@@ -139,6 +143,10 @@ public:
 
 	bool SetActivePage(const FDMMaterialEditorPage& InPage);
 
+	void SaveOriginal();
+
+	void ApplyToOriginal();
+
 	void HandleDrop_CreateTextureSet(const TArray<FAssetData>& InTextureAssets);
 
 	void HandleDrop_TextureSet(UDMTextureSet* InTextureSet);
@@ -154,6 +162,11 @@ public:
 	virtual void PostRedo(bool bInSuccess) override;
 	//~ End FUndoClient
 
+	//~ Begin FGCObject
+	virtual void AddReferencedObjects(FReferenceCollector& Collector) override;
+	virtual FString GetReferencerName() const override;
+	//~ End FGCObject
+
 protected:
 	TWeakPtr<SDMMaterialDesigner> DesignerWidgetWeak;
 
@@ -169,7 +182,9 @@ protected:
 	TDMWidgetSlot<SDMMaterialComponentEditor> ComponentEditorSlot;
 	TDMWidgetSlot<SDMStatusBar> StatusBarSlot;
 
-	TWeakObjectPtr<UDynamicMaterialModelBase> MaterialModelBaseWeak;
+	TWeakObjectPtr<UDynamicMaterialModelBase> OriginalMaterialModelBaseWeak;
+	TObjectPtr<UDynamicMaterialInstance> PreviewMaterialInstance;
+	TObjectPtr<UDynamicMaterialModelBase> PreviewMaterialModelBase;
 	TOptional<FDMObjectMaterialProperty> ObjectMaterialPropertyOpt;
 
 	TSharedRef<FUICommandList> CommandList;
diff --git a/Engine/Plugins/Experimental/DynamicMaterial/Source/DynamicMaterialEditor/Private/UI/Widgets/Visualizers/SDMMaterialComponentPreview.cpp b/Engine/Plugins/Experimental/DynamicMaterial/Source/DynamicMaterialEditor/Private/UI/Widgets/Visualizers/SDMMaterialComponentPreview.cpp
--- a/Engine/Plugins/Experimental/DynamicMaterial/Source/DynamicMaterialEditor/Private/UI/Widgets/Visualizers/SDMMaterialComponentPreview.cpp
+++ b/Engine/Plugins/Experimental/DynamicMaterial/Source/DynamicMaterialEditor/Private/UI/Widgets/Visualizers/SDMMaterialComponentPreview.cpp
@@ -55,7 +55,7 @@ void SDMMaterialComponentPreview::Construct(const FArguments& InArgs, const TSha
 
 	PreviewMaterialBaseWeak = InEditorWidget->GetPreviewMaterialManager()->CreatePreviewMaterial(InComponent);
 	PreviewMaterialDynamicWeak = InEditorWidget->GetPreviewMaterialManager()->CreatePreviewMaterialDynamic(PreviewMaterialBaseWeak.Get());
-	MaterialModelBaseWeak = InEditorWidget->GetMaterialModelBase();
+	MaterialModelBaseWeak = InEditorWidget->GetPreviewMaterialModelBase();
 
 	SetCanTick(true);
 
@@ -211,7 +211,7 @@ void SDMMaterialComponentPreview::RecreateMaterial()
 		EditorWidget->GetPreviewMaterialManager()->FreePreviewMaterialDynamic(PreviewMaterialBase);
 		PreviewMaterialDynamicWeak = EditorWidget->GetPreviewMaterialManager()->CreatePreviewMaterialDynamic(PreviewMaterialBase);
 
-		UDynamicMaterialModelBase* MaterialModelBase = EditorWidget->GetMaterialModelBase();
+		UDynamicMaterialModelBase* MaterialModelBase = EditorWidget->GetPreviewMaterialModelBase();
 
 		if (UDynamicMaterialModel* MaterialModel = Cast<UDynamicMaterialModel>(MaterialModelBase))
 		{
@@ -231,7 +231,7 @@ void SDMMaterialComponentPreview::RecreateMaterial()
 		EditorWidget->GetPreviewMaterialManager()->FreePreviewMaterialDynamic(PreviewMaterialBase);
 		PreviewMaterialDynamicWeak = EditorWidget->GetPreviewMaterialManager()->CreatePreviewMaterialDynamic(PreviewMaterialBase);
 
-		UDynamicMaterialModelBase* MaterialModelBase = EditorWidget->GetMaterialModelBase();
+		UDynamicMaterialModelBase* MaterialModelBase = EditorWidget->GetPreviewMaterialModelBase();
 
 		if (UDynamicMaterialModel* MaterialModel = Cast<UDynamicMaterialModel>(MaterialModelBase))
 		{
