Minor fix to LOC macro duplicated ids

diff --git a/Engine/Plugins/Editor/UMGWidgetPreview/Source/UMGWidgetPreview/Private/Widgets/SWidgetPreview.cpp b/Engine/Plugins/Editor/UMGWidgetPreview/Source/UMGWidgetPreview/Private/Widgets/SWidgetPreview.cpp
--- a/Engine/Plugins/Editor/UMGWidgetPreview/Source/UMGWidgetPreview/Private/Widgets/SWidgetPreview.cpp
+++ b/Engine/Plugins/Editor/UMGWidgetPreview/Source/UMGWidgetPreview/Private/Widgets/SWidgetPreview.cpp
@@ -265,7 +265,7 @@ namespace UE::UMGWidgetPreview::Private
 
 				if (bSizeOverridden)
 				{
-					return FText::Format(LOCTEXT("PreviewSizeTextCustom", "Preview Size: Overriden ({0}, {1})"),
+					return FText::Format(LOCTEXT("PreviewSizeTextOverriden", "Preview Size: Overriden ({0}, {1})"),
 						FText::AsNumber(CreatedWidgetPreviewCustomSize.X, &FNumberFormattingOptions::DefaultNoGrouping()),
 						FText::AsNumber(CreatedWidgetPreviewCustomSize.Y, &FNumberFormattingOptions::DefaultNoGrouping()));
 				}
diff --git a/Engine/Plugins/Experimental/DynamicMaterial/Source/DynamicMaterialEditor/Private/UI/Widgets/Editor/SDMToolBar.cpp b/Engine/Plugins/Experimental/DynamicMaterial/Source/DynamicMaterialEditor/Private/UI/Widgets/Editor/SDMToolBar.cpp
--- a/Engine/Plugins/Experimental/DynamicMaterial/Source/DynamicMaterialEditor/Private/UI/Widgets/Editor/SDMToolBar.cpp
+++ b/Engine/Plugins/Experimental/DynamicMaterial/Source/DynamicMaterialEditor/Private/UI/Widgets/Editor/SDMToolBar.cpp
@@ -134,7 +134,7 @@ TSharedRef<SWidget> SDMToolBar::CreateToolBarEntries()
 					SAssignNew(ApplyButtonWidget, SButton)
 					.ContentPadding(LargeIconToolBarButtonContentPadding)
 					.ButtonStyle(FDynamicMaterialEditorStyle::Get(), "HoverHintOnly")
-					.ToolTipText(LOCTEXT("MaterialDesignerSaveTooltip", "Apply the preview Material Designer Material back to the original asset.."))
+					.ToolTipText(LOCTEXT("MaterialDesignerApplyTooltip", "Apply the preview Material Designer Material back to the original asset.."))
 					.OnClicked(this, &SDMToolBar::OnApplyClicked)
 					[
 						SNew(STextBlock)
