[Backout] - CL40244961
#rnx
Original CL Desc
-----------------------------------------------------------------
Material Designer: Apply button
- Materials edited in the designer are now preview materials which are applied to the source asset when clicking apply or the save buttons.

diff --git a/Engine/Plugins/Experimental/DynamicMaterial/Source/DynamicMaterialEditor/Private/UI/Menus/DMMenuContext.cpp b/Engine/Plugins/Experimental/DynamicMaterial/Source/DynamicMaterialEditor/Private/UI/Menus/DMMenuContext.cpp
--- a/Engine/Plugins/Experimental/DynamicMaterial/Source/DynamicMaterialEditor/Private/UI/Menus/DMMenuContext.cpp
+++ b/Engine/Plugins/Experimental/DynamicMaterial/Source/DynamicMaterialEditor/Private/UI/Menus/DMMenuContext.cpp
@@ -116,7 +116,7 @@ UDynamicMaterialModelBase* UDMMenuContext::GetModelBase() const
 {
 	if (TSharedPtr<SDMMaterialEditor> EditorWidget = EditorWidgetWeak.Pin())
 	{
-		return EditorWidget->GetPreviewMaterialModelBase();
+		return EditorWidget->GetMaterialModelBase();
 	}
 
 	return nullptr;
@@ -126,7 +126,7 @@ UDynamicMaterialModel* UDMMenuContext::GetModel() const
 {
 	if (TSharedPtr<SDMMaterialEditor> EditorWidget = EditorWidgetWeak.Pin())
 	{
-		return EditorWidget->GetPreviewMaterialModel();
+		return EditorWidget->GetMaterialModel();
 	}
 
 	return nullptr;
diff --git a/Engine/Plugins/Experimental/DynamicMaterial/Source/DynamicMaterialEditor/Private/UI/PropertyGenerators/DMComponentPropertyRowGenerator.cpp b/Engine/Plugins/Experimental/DynamicMaterial/Source/DynamicMaterialEditor/Private/UI/PropertyGenerators/DMComponentPropertyRowGenerator.cpp
--- a/Engine/Plugins/Experimental/DynamicMaterial/Source/DynamicMaterialEditor/Private/UI/PropertyGenerators/DMComponentPropertyRowGenerator.cpp
+++ b/Engine/Plugins/Experimental/DynamicMaterial/Source/DynamicMaterialEditor/Private/UI/PropertyGenerators/DMComponentPropertyRowGenerator.cpp
@@ -112,7 +112,7 @@ bool FDMComponentPropertyRowGenerator::IsDynamic(const TSharedRef<SDMMaterialCom
 {
 	if (TSharedPtr<SDMMaterialEditor> EditorWidget = InComponentEditorWidget->GetEditorWidget())
 	{
-		if (UDynamicMaterialModelBase* MaterialModelBase = EditorWidget->GetPreviewMaterialModelBase())
+		if (UDynamicMaterialModelBase* MaterialModelBase = EditorWidget->GetMaterialModelBase())
 		{
 			return !!Cast<UDynamicMaterialModelDynamic>(MaterialModelBase);
 		}
diff --git a/Engine/Plugins/Experimental/DynamicMaterial/Source/DynamicMaterialEditor/Private/UI/PropertyGenerators/DMMaterialValuePropertyRowGenerator.cpp b/Engine/Plugins/Experimental/DynamicMaterial/Source/DynamicMaterialEditor/Private/UI/PropertyGenerators/DMMaterialValuePropertyRowGenerator.cpp
--- a/Engine/Plugins/Experimental/DynamicMaterial/Source/DynamicMaterialEditor/Private/UI/PropertyGenerators/DMMaterialValuePropertyRowGenerator.cpp
+++ b/Engine/Plugins/Experimental/DynamicMaterial/Source/DynamicMaterialEditor/Private/UI/PropertyGenerators/DMMaterialValuePropertyRowGenerator.cpp
@@ -49,7 +49,7 @@ void FDMMaterialValuePropertyRowGenerator::AddComponentProperties(const TSharedR
 
 	if (TSharedPtr<SDMMaterialEditor> EditorWidget = InComponentEditorWidget->GetEditorWidget())
 	{
-		if (UDynamicMaterialModelBase* MaterialModelBase = EditorWidget->GetPreviewMaterialModelBase())
+		if (UDynamicMaterialModelBase* MaterialModelBase = EditorWidget->GetMaterialModelBase())
 		{
 			if (UDynamicMaterialModelDynamic* MaterialModelDynamic = Cast<UDynamicMaterialModelDynamic>(MaterialModelBase))
 			{
diff --git a/Engine/Plugins/Experimental/DynamicMaterial/Source/DynamicMaterialEditor/Private/UI/PropertyGenerators/DMTextureUVPropertyRowGenerator.cpp b/Engine/Plugins/Experimental/DynamicMaterial/Source/DynamicMaterialEditor/Private/UI/PropertyGenerators/DMTextureUVPropertyRowGenerator.cpp
--- a/Engine/Plugins/Experimental/DynamicMaterial/Source/DynamicMaterialEditor/Private/UI/PropertyGenerators/DMTextureUVPropertyRowGenerator.cpp
+++ b/Engine/Plugins/Experimental/DynamicMaterial/Source/DynamicMaterialEditor/Private/UI/PropertyGenerators/DMTextureUVPropertyRowGenerator.cpp
@@ -62,7 +62,7 @@ void FDMTextureUVPropertyRowGenerator::AddComponentProperties(const TSharedRef<S
 
 	if (TSharedPtr<SDMMaterialEditor> EditorWidget = InComponentEditorWidget->GetEditorWidget())
 	{
-		if (UDynamicMaterialModelBase* MaterialModelBase = EditorWidget->GetPreviewMaterialModelBase())
+		if (UDynamicMaterialModelBase* MaterialModelBase = EditorWidget->GetMaterialModelBase())
 		{
 			if (UDynamicMaterialModelDynamic* MaterialModelDynamic = Cast<UDynamicMaterialModelDynamic>(MaterialModelBase))
 			{
diff --git a/Engine/Plugins/Experimental/DynamicMaterial/Source/DynamicMaterialEditor/Private/UI/Widgets/Editor/SDMMaterialProperties.cpp b/Engine/Plugins/Experimental/DynamicMaterial/Source/DynamicMaterialEditor/Private/UI/Widgets/Editor/SDMMaterialProperties.cpp
--- a/Engine/Plugins/Experimental/DynamicMaterial/Source/DynamicMaterialEditor/Private/UI/Widgets/Editor/SDMMaterialProperties.cpp
+++ b/Engine/Plugins/Experimental/DynamicMaterial/Source/DynamicMaterialEditor/Private/UI/Widgets/Editor/SDMMaterialProperties.cpp
@@ -87,7 +87,7 @@ TSharedRef<SWidget> SDMMaterialProperties::CreateSlot_Content()
 	{
 		return SNullWidget::NullWidget;
 	}
-	UDynamicMaterialModelBase* MaterialModelBase = EditorWidget->GetPreviewMaterialModelBase();
+	UDynamicMaterialModelBase* MaterialModelBase = EditorWidget->GetMaterialModelBase();
 
 	if (!MaterialModelBase)
 	{
@@ -245,7 +245,7 @@ void SDMMaterialProperties::AddProperty(const TSharedRef<ICustomDetailsView>& In
 		return;
 	}
 
-	UDynamicMaterialModel* MaterialModel = EditorWidget->GetPreviewMaterialModel();
+	UDynamicMaterialModel* MaterialModel = EditorWidget->GetMaterialModel();
 
 	if (!MaterialModel)
 	{
@@ -275,7 +275,7 @@ TSharedRef<SWidget> SDMMaterialProperties::CreatePropertyRow(UDMMaterialProperty
 {
 	// There are all ensured to be valid by the caller of this method
 	TSharedPtr<SDMMaterialEditor> EditorWidget = EditorWidgetWeak.Pin();
-	UDynamicMaterialModelBase* MaterialModelBase = EditorWidget->GetPreviewMaterialModelBase();
+	UDynamicMaterialModelBase* MaterialModelBase = EditorWidget->GetMaterialModelBase();
 	UDynamicMaterialModelEditorOnlyData* EditorOnlyData = UDynamicMaterialModelEditorOnlyData::Get(MaterialModelBase);
 	UDynamicMaterialEditorSettings* Settings = UDynamicMaterialEditorSettings::Get();
 	const EDMMaterialPropertyType MaterialProperty = InProperty->GetMaterialProperty();
@@ -422,7 +422,7 @@ bool SDMMaterialProperties::GetPropertyEnabledEnabled(EDMMaterialPropertyType In
 		return false;
 	}
 
-	UDynamicMaterialModel* MaterialModel = EditorWidget->GetPreviewMaterialModel();
+	UDynamicMaterialModel* MaterialModel = EditorWidget->GetMaterialModel();
 
 	if (!MaterialModel)
 	{
@@ -455,7 +455,7 @@ ECheckBoxState SDMMaterialProperties::GetPropertyEnabledState(EDMMaterialPropert
 		return ECheckBoxState::Unchecked;
 	}
 
-	UDynamicMaterialModel* MaterialModel = EditorWidget->GetPreviewMaterialModel();
+	UDynamicMaterialModel* MaterialModel = EditorWidget->GetMaterialModel();
 
 	if (!MaterialModel)
 	{
@@ -490,7 +490,7 @@ void SDMMaterialProperties::OnPropertyEnabledStateChanged(ECheckBoxState InState
 		return;
 	}
 
-	UDynamicMaterialModel* MaterialModel = EditorWidget->GetPreviewMaterialModel();
+	UDynamicMaterialModel* MaterialModel = EditorWidget->GetMaterialModel();
 
 	if (!MaterialModel)
 	{
@@ -541,7 +541,7 @@ FReply SDMMaterialProperties::OnPropertyClicked(const FGeometry& InGeometry, con
 		return FReply::Handled();
 	}
 
-	UDynamicMaterialModel* MaterialModel = EditorWidget->GetPreviewMaterialModel();
+	UDynamicMaterialModel* MaterialModel = EditorWidget->GetMaterialModel();
 
 	if (!MaterialModel)
 	{
@@ -581,7 +581,7 @@ TSharedRef<SWidget> SDMMaterialProperties::CreateGlobalSlider(UDMMaterialPropert
 		return SNullWidget::NullWidget;
 	}
 
-	UDynamicMaterialModelBase* MaterialModelBase = EditorWidget->GetPreviewMaterialModelBase();
+	UDynamicMaterialModelBase* MaterialModelBase = EditorWidget->GetMaterialModelBase();
 
 	if (!MaterialModelBase)
 	{
@@ -717,7 +717,7 @@ void SDMMaterialProperties::OnExpansionStateChanged(const TSharedRef<ICustomDeta
 		return;
 	}
 
-	UDynamicMaterialModelBase* MaterialModelBase = EditorWidget->GetPreviewMaterialModelBase();
+	UDynamicMaterialModelBase* MaterialModelBase = EditorWidget->GetMaterialModelBase();
 
 	if (!MaterialModelBase)
 	{
@@ -767,7 +767,7 @@ bool SDMMaterialProperties::OnAssetDraggedOver(TArrayView<FAssetData> InAssets,
 		return false;
 	}
 
-	UDynamicMaterialModel* MaterialModel = EditorWidget->GetPreviewMaterialModel();
+	UDynamicMaterialModel* MaterialModel = EditorWidget->GetMaterialModel();
 
 	if (!MaterialModel)
 	{
@@ -844,7 +844,7 @@ void SDMMaterialProperties::HandleDrop_Texture(UTexture* InTexture, EDMMaterialP
 		return;
 	}
 
-	UDynamicMaterialModel* MaterialModel = EditorWidget->GetPreviewMaterialModel();
+	UDynamicMaterialModel* MaterialModel = EditorWidget->GetMaterialModel();
 
 	if (!MaterialModel)
 	{
diff --git a/Engine/Plugins/Experimental/DynamicMaterial/Source/DynamicMaterialEditor/Private/UI/Widgets/Editor/SDMMaterialPropertySelector.cpp b/Engine/Plugins/Experimental/DynamicMaterial/Source/DynamicMaterialEditor/Private/UI/Widgets/Editor/SDMMaterialPropertySelector.cpp
--- a/Engine/Plugins/Experimental/DynamicMaterial/Source/DynamicMaterialEditor/Private/UI/Widgets/Editor/SDMMaterialPropertySelector.cpp
+++ b/Engine/Plugins/Experimental/DynamicMaterial/Source/DynamicMaterialEditor/Private/UI/Widgets/Editor/SDMMaterialPropertySelector.cpp
@@ -108,7 +108,7 @@ UDynamicMaterialModelEditorOnlyData* SDMMaterialPropertySelector::GetEditorOnlyD
 		return nullptr;
 	}
 
-	UDynamicMaterialModel* MaterialModel = EditorWidget->GetPreviewMaterialModel();
+	UDynamicMaterialModel* MaterialModel = EditorWidget->GetMaterialModel();
 
 	if (!IsValid(MaterialModel))
 	{
diff --git a/Engine/Plugins/Experimental/DynamicMaterial/Source/DynamicMaterialEditor/Private/UI/Widgets/Editor/SDMMaterialSlotEditor.cpp b/Engine/Plugins/Experimental/DynamicMaterial/Source/DynamicMaterialEditor/Private/UI/Widgets/Editor/SDMMaterialSlotEditor.cpp
--- a/Engine/Plugins/Experimental/DynamicMaterial/Source/DynamicMaterialEditor/Private/UI/Widgets/Editor/SDMMaterialSlotEditor.cpp
+++ b/Engine/Plugins/Experimental/DynamicMaterial/Source/DynamicMaterialEditor/Private/UI/Widgets/Editor/SDMMaterialSlotEditor.cpp
@@ -81,7 +81,7 @@ void SDMMaterialSlotEditor::Construct(const FArguments& InArgs, const TSharedRef
 
 	SetCanTick(false);
 
-	bIsDynamic = !Cast<UDynamicMaterialModel>(InEditorWidget->GetPreviewMaterialModelBase());
+	bIsDynamic = !Cast<UDynamicMaterialModel>(InEditorWidget->GetMaterialModelBase());
 
 	ContentSlot = TDMWidgetSlot<SWidget>(SharedThis(this), 0, SNullWidget::NullWidget);
 
@@ -678,7 +678,7 @@ TSharedRef<SWidget> SDMMaterialSlotEditor::CreateSlot_LayerOpacity()
 						{
 							if (TSharedPtr<SDMMaterialEditor> EditorWidget = EditorWidgetWeak.Pin())
 							{
-								if (UDynamicMaterialModelBase* MaterialModelBase = EditorWidget->GetPreviewMaterialModelBase())
+								if (UDynamicMaterialModelBase* MaterialModelBase = EditorWidget->GetMaterialModelBase())
 								{
 									World = MaterialModelBase->GetWorld();
 								}
@@ -963,7 +963,7 @@ bool SDMMaterialSlotEditor::OnAreAssetsAcceptableForDrop(TArrayView<FAssetData>
 		return false;
 	}
 
-	UDynamicMaterialModelBase* MaterialModelBase = EditorWidget->GetPreviewMaterialModelBase();
+	UDynamicMaterialModelBase* MaterialModelBase = EditorWidget->GetMaterialModelBase();
 
 	if (MaterialModelBase->IsA<UDynamicMaterialModelDynamic>())
 	{
@@ -1023,7 +1023,7 @@ void SDMMaterialSlotEditor::OnAssetsDropped(const FDragDropEvent& InDragDropEven
 		return;
 	}
 
-	UDynamicMaterialModelBase* MaterialModelBase = EditorWidget->GetPreviewMaterialModelBase();
+	UDynamicMaterialModelBase* MaterialModelBase = EditorWidget->GetMaterialModelBase();
 
 	if (MaterialModelBase->IsA<UDynamicMaterialModelDynamic>())
 	{
diff --git a/Engine/Plugins/Experimental/DynamicMaterial/Source/DynamicMaterialEditor/Private/UI/Widgets/Editor/SDMObjectEditorWidgetBase.cpp b/Engine/Plugins/Experimental/DynamicMaterial/Source/DynamicMaterialEditor/Private/UI/Widgets/Editor/SDMObjectEditorWidgetBase.cpp
--- a/Engine/Plugins/Experimental/DynamicMaterial/Source/DynamicMaterialEditor/Private/UI/Widgets/Editor/SDMObjectEditorWidgetBase.cpp
+++ b/Engine/Plugins/Experimental/DynamicMaterial/Source/DynamicMaterialEditor/Private/UI/Widgets/Editor/SDMObjectEditorWidgetBase.cpp
@@ -42,7 +42,7 @@ void SDMObjectEditorWidgetBase::Construct(const FArguments& InArgs, const TShare
 
 	if (!WorldContext)
 	{
-		WorldContext = InEditorWidget->GetPreviewMaterialModelBase();
+		WorldContext = InEditorWidget->GetMaterialModelBase();
 	}
 
 	if (WorldContext)
diff --git a/Engine/Plugins/Experimental/DynamicMaterial/Source/DynamicMaterialEditor/Private/UI/Widgets/Editor/SDMStatusBar.cpp b/Engine/Plugins/Experimental/DynamicMaterial/Source/DynamicMaterialEditor/Private/UI/Widgets/Editor/SDMStatusBar.cpp
--- a/Engine/Plugins/Experimental/DynamicMaterial/Source/DynamicMaterialEditor/Private/UI/Widgets/Editor/SDMStatusBar.cpp
+++ b/Engine/Plugins/Experimental/DynamicMaterial/Source/DynamicMaterialEditor/Private/UI/Widgets/Editor/SDMStatusBar.cpp
@@ -26,7 +26,7 @@ SDMStatusBar::~SDMStatusBar()
 		return;
 	}
 
-	UDynamicMaterialModel* MaterialModel = EditorWidget->GetPreviewMaterialModel();
+	UDynamicMaterialModel* MaterialModel = EditorWidget->GetMaterialModel();
 
 	if (!MaterialModel)
 	{
@@ -51,7 +51,7 @@ void SDMStatusBar::Construct(const FArguments& InArgs, const TSharedRef<SDMMater
 
 	SetCanTick(false);
 
-	UDynamicMaterialModel* MaterialModel = InEditorWidget->GetPreviewMaterialModel();
+	UDynamicMaterialModel* MaterialModel = InEditorWidget->GetMaterialModel();
 
 	if (!MaterialModel)
 	{
diff --git a/Engine/Plugins/Experimental/DynamicMaterial/Source/DynamicMaterialEditor/Private/UI/Widgets/Editor/SDMToolBar.cpp b/Engine/Plugins/Experimental/DynamicMaterial/Source/DynamicMaterialEditor/Private/UI/Widgets/Editor/SDMToolBar.cpp
--- a/Engine/Plugins/Experimental/DynamicMaterial/Source/DynamicMaterialEditor/Private/UI/Widgets/Editor/SDMToolBar.cpp
+++ b/Engine/Plugins/Experimental/DynamicMaterial/Source/DynamicMaterialEditor/Private/UI/Widgets/Editor/SDMToolBar.cpp
@@ -16,6 +16,7 @@
 #include "Model/DynamicMaterialModel.h"
 #include "Model/DynamicMaterialModelBase.h"
 #include "Model/DynamicMaterialModelDynamic.h"
+#include "PackageTools.h"
 #include "Selection.h"
 #include "Styling/StyleColors.h"
 #include "UI/Menus/DMToolBarMenus.h"
@@ -67,7 +68,7 @@ void SDMToolBar::Construct(const FArguments& InArgs, const TSharedRef<SDMMateria
 		];
 
 	SetActorPropertySelected(InActor);
-	UpdateButtonVisibilities();
+	SetButtonVisibilities();
 }
 
 TSharedRef<SWidget> SDMToolBar::CreateToolBarEntries()
@@ -109,7 +110,7 @@ TSharedRef<SWidget> SDMToolBar::CreateToolBarEntries()
 						.Image(this, &SDMToolBar::GetSaveIcon)
 						.DesiredSizeOverride(LargeIconToolBarButtonSize)
 					]
-				]
+				]		
 				+ SHorizontalBox::Slot()
 				.AutoWidth()
 				.Padding(0.f)
@@ -128,22 +129,6 @@ TSharedRef<SWidget> SDMToolBar::CreateToolBarEntries()
 				]
 				+ SHorizontalBox::Slot()
 				.AutoWidth()
-				.Padding(0.f)
-				.VAlign(EVerticalAlignment::VAlign_Center)
-				[
-					SAssignNew(ApplyButtonWidget, SButton)
-					.ContentPadding(LargeIconToolBarButtonContentPadding)
-					.ButtonStyle(FDynamicMaterialEditorStyle::Get(), "HoverHintOnly")
-					.ToolTipText(LOCTEXT("MaterialDesignerApplyTooltip", "Apply the preview Material Designer Material back to the original asset.."))
-					.OnClicked(this, &SDMToolBar::OnApplyClicked)
-					[
-						SNew(STextBlock)
-						.TextStyle(FDynamicMaterialEditorStyle::Get(), "RegularFont")
-						.Text(LOCTEXT("Apply", "Apply"))
-					]
-				]	
-				+ SHorizontalBox::Slot()
-				.AutoWidth()
 				.VAlign(VAlign_Center)
 				.Padding(0.f)
 				[
@@ -340,7 +325,7 @@ void SDMToolBar::SetActorPropertySelected(AActor* InActor)
 		ActorRowWidget->SetVisibility(EVisibility::Visible);
 
 		TArray<FDMObjectMaterialProperty> ActorProperties = UDMMaterialInstanceFunctionLibrary::GetActorMaterialProperties(InActor);
-		UDynamicMaterialModelBase* MaterialModelBase = GetOriginalMaterialModelBase();
+		UDynamicMaterialModelBase* MaterialModelBase = GetMaterialModelBase();
 
 		const int32 ActorPropertyCount = ActorProperties.Num();
 		ActorMaterialProperties.Empty(ActorPropertyCount);
@@ -367,12 +352,12 @@ void SDMToolBar::SetActorPropertySelected(AActor* InActor)
 	PropertySelectorContainer->SetContent(CreateSlotsComboBoxWidget());
 }
 
-void SDMToolBar::UpdateButtonVisibilities()
+void SDMToolBar::SetButtonVisibilities()
 {
 	bool bIsAsset = false;
 	bool bIsDynamic = false;
 
-	UDynamicMaterialModelBase* MaterialModelBase = GetOriginalMaterialModelBase();
+	UDynamicMaterialModelBase* MaterialModelBase = GetMaterialModelBase();
 
 	if (IsValid(MaterialModelBase))
 	{
@@ -456,7 +441,7 @@ TSharedRef<SWidget> SDMToolBar::CreateToolBarButton(TAttribute<const FSlateBrush
 
 TSharedRef<SWidget> SDMToolBar::CreateSlotsComboBoxWidget()
 {
-	UDynamicMaterialModelBase* MaterialModelBase = GetOriginalMaterialModelBase();
+	UDynamicMaterialModelBase* MaterialModelBase = GetMaterialModelBase();
 
 	if (!MaterialActorWeak.IsValid() || !IsValid(MaterialModelBase))
 	{
@@ -542,11 +527,11 @@ void SDMToolBar::OnMaterialSlotChanged(TSharedPtr<FDMObjectMaterialProperty> InS
 	}
 }
 
-UDynamicMaterialModelBase* SDMToolBar::GetOriginalMaterialModelBase() const
+UDynamicMaterialModelBase* SDMToolBar::GetMaterialModelBase() const
 {
 	if (TSharedPtr<SDMMaterialEditor> EditorWidget = GetEditorWidget())
 	{
-		return EditorWidget->GetOriginalMaterialModelBase();
+		return EditorWidget->GetMaterialModelBase();
 	}
 
 	return nullptr;
@@ -606,7 +591,7 @@ FReply SDMToolBar::OnExportMaterialInstanceButtonClicked()
 		return FReply::Handled();
 	}
 
-	UDynamicMaterialModelBase* MaterialModelBase = GetOriginalMaterialModelBase();
+	UDynamicMaterialModelBase* MaterialModelBase = GetMaterialModelBase();
 
 	if (!MaterialModelBase)
 	{
@@ -627,7 +612,7 @@ FReply SDMToolBar::OnExportMaterialInstanceButtonClicked()
 
 FReply SDMToolBar::OnBrowseClicked()
 {
-	UDynamicMaterialModelBase* MaterialModelBase = GetOriginalMaterialModelBase();
+	UDynamicMaterialModelBase* MaterialModelBase = GetMaterialModelBase();
 
 	if (!IsValid(MaterialModelBase))
 	{
@@ -667,7 +652,7 @@ FReply SDMToolBar::OnUseClicked()
 		return FReply::Handled();
 	}
 
-	UDynamicMaterialModelBase* CurrentModelBase = GetOriginalMaterialModelBase();
+	UDynamicMaterialModelBase* CurrentModelBase = GetMaterialModelBase();
 	UDMWorldSubsystem* DMSubsystem = nullptr;
 
 	if (CurrentModelBase)
@@ -742,6 +727,23 @@ FReply SDMToolBar::OnUseClicked()
 	return FReply::Handled();
 }
 
+UPackage* SDMToolBar::GetSaveablePackage(UObject* InObject)
+{
+	if (!IsValid(InObject))
+	{
+		return nullptr;
+	}
+
+	UPackage* Package = InObject->GetPackage();
+
+	if (!Package || Package->HasAllFlags(RF_Transient))
+	{
+		return nullptr;
+	}
+
+	return Package;
+}
+
 FText SDMToolBar::GetActorName() const
 {
 	if (const AActor* const SlotActor = GetMaterialActor())
@@ -759,7 +761,7 @@ TSharedPtr<SDMMaterialEditor> SDMToolBar::GetEditorWidget() const
 
 FText SDMToolBar::GetAssetName() const
 {
-	if (UDynamicMaterialModelBase* MaterialModelBase = GetOriginalMaterialModelBase())
+	if (UDynamicMaterialModelBase* MaterialModelBase = GetMaterialModelBase())
 	{
 		if (UDynamicMaterialInstance* MaterialInstance = MaterialModelBase->GetDynamicMaterialInstance())
 		{
@@ -779,7 +781,7 @@ FText SDMToolBar::GetAssetName() const
 
 FText SDMToolBar::GetAssetToolTip() const
 {
-	if (UDynamicMaterialModelBase* MaterialModelBase = GetOriginalMaterialModelBase())
+	if (UDynamicMaterialModelBase* MaterialModelBase = GetMaterialModelBase())
 	{
 		if (UDynamicMaterialInstance* MaterialInstance = MaterialModelBase->GetDynamicMaterialInstance())
 		{
@@ -799,7 +801,7 @@ FText SDMToolBar::GetAssetToolTip() const
 
 const FSlateBrush* SDMToolBar::GetSaveIcon() const
 {
-	if (UPackage* Package = SDMMaterialEditor::GetSaveablePackage(GetOriginalMaterialModelBase()))
+	if (UPackage* Package = GetSaveablePackage(GetMaterialModelBase()))
 	{
 		if (Package->IsDirty())
 		{
@@ -812,19 +814,14 @@ const FSlateBrush* SDMToolBar::GetSaveIcon() const
 
 FReply SDMToolBar::OnSaveClicked()
 {
-	if (TSharedPtr<SDMMaterialEditor> EditorWidget = GetEditorWidget())
-	{
-		EditorWidget->SaveOriginal();
-	}
-
-	return FReply::Handled();
-}
-
-FReply SDMToolBar::OnApplyClicked()
-{
-	if (TSharedPtr<SDMMaterialEditor> EditorWidget = GetEditorWidget())
+	if (UDynamicMaterialModelBase* MaterialModelBase = GetMaterialModelBase())
 	{
-		EditorWidget->ApplyToOriginal();
+		if (UPackage* Package = GetSaveablePackage(MaterialModelBase))
+		{
+			TArray<UObject*> AssetsToSave;
+			AssetsToSave.Add(MaterialModelBase);
+			UPackageTools::SavePackagesForObjects(AssetsToSave);
+		}
 	}
 
 	return FReply::Handled();
@@ -836,7 +833,7 @@ FReply SDMToolBar::OnOpenParentClicked()
 	{
 		if (TSharedPtr<SDMMaterialDesigner> DesignerWidget = EditorWidget->GetDesignerWidget())
 		{
-			if (UDynamicMaterialModelDynamic* DynamicMaterialModel = Cast<UDynamicMaterialModelDynamic>(GetOriginalMaterialModelBase()))
+			if (UDynamicMaterialModelDynamic* DynamicMaterialModel = Cast<UDynamicMaterialModelDynamic>(EditorWidget->GetMaterialModelBase()))
 			{
 				if (UDynamicMaterialModel* ParentModel = DynamicMaterialModel->ResolveMaterialModel())
 				{
@@ -851,7 +848,7 @@ FReply SDMToolBar::OnOpenParentClicked()
 
 FReply SDMToolBar::OnConvertToEditableClicked()
 {
-	UDynamicMaterialModelDynamic* CurrentModelDynamic = Cast<UDynamicMaterialModelDynamic>(GetOriginalMaterialModelBase());
+	UDynamicMaterialModelDynamic* CurrentModelDynamic = Cast<UDynamicMaterialModelDynamic>(GetMaterialModelBase());
 
 	if (!CurrentModelDynamic)
 	{
@@ -951,7 +948,7 @@ FReply SDMToolBar::OnConvertToEditableClicked()
 
 FReply SDMToolBar::OnConvertToInstanceClicked()
 {
-	UDynamicMaterialModel* CurrentModel = Cast<UDynamicMaterialModel>(GetOriginalMaterialModelBase());
+	UDynamicMaterialModel* CurrentModel = Cast<UDynamicMaterialModel>(GetMaterialModelBase());
 
 	if (!CurrentModel)
 	{
diff --git a/Engine/Plugins/Experimental/DynamicMaterial/Source/DynamicMaterialEditor/Private/UI/Widgets/Editor/SDMToolBar.h b/Engine/Plugins/Experimental/DynamicMaterial/Source/DynamicMaterialEditor/Private/UI/Widgets/Editor/SDMToolBar.h
--- a/Engine/Plugins/Experimental/DynamicMaterial/Source/DynamicMaterialEditor/Private/UI/Widgets/Editor/SDMToolBar.h
+++ b/Engine/Plugins/Experimental/DynamicMaterial/Source/DynamicMaterialEditor/Private/UI/Widgets/Editor/SDMToolBar.h
@@ -47,6 +47,8 @@ public:
 	TSharedPtr<SDMMaterialEditor> GetEditorWidget() const;
 
 protected:
+	static UPackage* GetSaveablePackage(UObject* InObject);
+
 	TWeakPtr<SDMMaterialEditor> EditorWidgetWeak;
 	TWeakObjectPtr<AActor> MaterialActorWeak;
 
@@ -55,7 +57,6 @@ protected:
 
 	TSharedPtr<SBox> PropertySelectorContainer;
 	TSharedPtr<SWidget> SaveButtonWidget;
-	TSharedPtr<SWidget> ApplyButtonWidget;
 	TSharedPtr<SWidget> ActorRowWidget;
 	TSharedPtr<SWidget> AssetRowWidget;
 	TSharedPtr<STextBlock> ActorNameWidget;
@@ -65,13 +66,13 @@ protected:
 	TSharedPtr<SWidget> ConvertToEditableButton;
 	TSharedPtr<SWidget> ConvertToInstanceButton;
 
-	UDynamicMaterialModelBase* GetOriginalMaterialModelBase() const;
+	UDynamicMaterialModelBase* GetMaterialModelBase() const;
 
 	TSharedRef<SWidget> CreateToolBarEntries();
 
 	void SetActorPropertySelected(AActor* InActor);
 
-	void UpdateButtonVisibilities();
+	void SetButtonVisibilities();
 
 	TSharedRef<SWidget> CreateToolBarButton(TAttribute<const FSlateBrush*> InImageBrush, const TAttribute<FText>& InTooltipText, FOnClicked InOnClicked);
 
@@ -100,8 +101,6 @@ protected:
 
 	FReply OnSaveClicked();
 
-	FReply OnApplyClicked();
-
 	FReply OnOpenParentClicked();
 
 	FReply OnConvertToEditableClicked();
diff --git a/Engine/Plugins/Experimental/DynamicMaterial/Source/DynamicMaterialEditor/Private/UI/Widgets/Editor/SlotEditor/SDMMaterialSlotLayerEffectItem.cpp b/Engine/Plugins/Experimental/DynamicMaterial/Source/DynamicMaterialEditor/Private/UI/Widgets/Editor/SlotEditor/SDMMaterialSlotLayerEffectItem.cpp
--- a/Engine/Plugins/Experimental/DynamicMaterial/Source/DynamicMaterialEditor/Private/UI/Widgets/Editor/SlotEditor/SDMMaterialSlotLayerEffectItem.cpp
+++ b/Engine/Plugins/Experimental/DynamicMaterial/Source/DynamicMaterialEditor/Private/UI/Widgets/Editor/SlotEditor/SDMMaterialSlotLayerEffectItem.cpp
@@ -175,7 +175,7 @@ bool SDMMaterialSlotLayerEffectItem::CanModifyMaterialModel() const
 				{
 					if (TSharedPtr<SDMMaterialEditor> EditorWidget = SlotEditorWidget->GetEditorWidget())
 					{
-						if (UDynamicMaterialModelBase* MaterialModelBase = EditorWidget->GetPreviewMaterialModelBase())
+						if (UDynamicMaterialModelBase* MaterialModelBase = EditorWidget->GetMaterialModelBase())
 						{
 							return MaterialModelBase->IsA<UDynamicMaterialModel>();
 						}
diff --git a/Engine/Plugins/Experimental/DynamicMaterial/Source/DynamicMaterialEditor/Private/UI/Widgets/SDMMaterialDesigner.cpp b/Engine/Plugins/Experimental/DynamicMaterial/Source/DynamicMaterialEditor/Private/UI/Widgets/SDMMaterialDesigner.cpp
--- a/Engine/Plugins/Experimental/DynamicMaterial/Source/DynamicMaterialEditor/Private/UI/Widgets/SDMMaterialDesigner.cpp
+++ b/Engine/Plugins/Experimental/DynamicMaterial/Source/DynamicMaterialEditor/Private/UI/Widgets/SDMMaterialDesigner.cpp
@@ -155,7 +155,7 @@ UDynamicMaterialModelBase* SDMMaterialDesigner::GetMaterialModelBase() const
 
 	if (Content->GetWidgetClass().GetWidgetType() == SDMMaterialEditor::StaticWidgetClass().GetWidgetType())
 	{
-		return StaticCastSharedPtr<SDMMaterialEditor>(Content)->GetPreviewMaterialModelBase();
+		return StaticCastSharedPtr<SDMMaterialEditor>(Content)->GetMaterialModelBase();
 	}
 
 	if (Content->GetWidgetClass().GetWidgetType() == SDMMaterialWizard::StaticWidgetClass().GetWidgetType())
@@ -502,7 +502,7 @@ void SDMMaterialDesigner::OnLayoutChanged()
 	// Already assured
 	TSharedPtr<SDMMaterialEditor> CurrentEditor = StaticCastSharedRef<SDMMaterialEditor>(Content.ToSharedRef());
 
-	UDynamicMaterialModelBase* MaterialModelBase = CurrentEditor->GetPreviewMaterialModelBase();
+	UDynamicMaterialModelBase* MaterialModelBase = CurrentEditor->GetMaterialModelBase();
 	const FDMObjectMaterialProperty* MaterialObjectProperty = CurrentEditor->GetMaterialObjectProperty();
 	EDMMaterialEditorMode EditorMode = CurrentEditor->GetEditMode();
 	EDMMaterialPropertyType SelectedProperty = CurrentEditor->GetSelectedPropertyType();
diff --git a/Engine/Plugins/Experimental/DynamicMaterial/Source/DynamicMaterialEditor/Private/UI/Widgets/SDMMaterialEditor.cpp b/Engine/Plugins/Experimental/DynamicMaterial/Source/DynamicMaterialEditor/Private/UI/Widgets/SDMMaterialEditor.cpp
--- a/Engine/Plugins/Experimental/DynamicMaterial/Source/DynamicMaterialEditor/Private/UI/Widgets/SDMMaterialEditor.cpp
+++ b/Engine/Plugins/Experimental/DynamicMaterial/Source/DynamicMaterialEditor/Private/UI/Widgets/SDMMaterialEditor.cpp
@@ -16,8 +16,6 @@
 #include "Framework/Application/IInputProcessor.h"
 #include "Framework/Application/SlateApplication.h"
 #include "Framework/Commands/GenericCommands.h"
-#include "Material/DynamicMaterialInstance.h"
-#include "Material/DynamicMaterialInstanceFactory.h"
 #include "Materials/Material.h"
 #include "Materials/MaterialFunctionInterface.h"
 #include "Misc/MessageDialog.h"
@@ -25,10 +23,7 @@
 #include "Model/DynamicMaterialModelBase.h"
 #include "Model/DynamicMaterialModelDynamic.h"
 #include "Model/DynamicMaterialModelEditorOnlyData.h"
-#include "PackageTools.h"
-#include "Selection.h"
 #include "Styling/SlateIconFinder.h"
-#include "UObject/Linker.h"
 #include "UI/Utils/DMEditorSelectionContext.h"
 #include "UI/Utils/DMPreviewMaterialManager.h"
 #include "UI/Widgets/Editor/SDMMaterialComponentEditor.h"
@@ -201,51 +196,18 @@ TSharedPtr<SDMMaterialDesigner> SDMMaterialEditor::GetDesignerWidget() const
 	return DesignerWidgetWeak.Pin();
 }
 
-UDynamicMaterialModelBase* SDMMaterialEditor::GetOriginalMaterialModelBase() const
+UDynamicMaterialModelBase* SDMMaterialEditor::GetMaterialModelBase() const
 {
-	return OriginalMaterialModelBaseWeak.Get();
-}
-
-UDynamicMaterialModelBase* SDMMaterialEditor::GetPreviewMaterialModelBase() const
-{
-	return PreviewMaterialModelBase.Get();
+	return MaterialModelBaseWeak.Get();
 }
 
 void SDMMaterialEditor::SetMaterialModelBase(UDynamicMaterialModelBase* InMaterialModelBase)
 {
-	OriginalMaterialModelBaseWeak = InMaterialModelBase;
+	MaterialModelBaseWeak = InMaterialModelBase;
 
-	if (!InMaterialModelBase)
+	if (UDynamicMaterialModelDynamic* MaterialModelDynamic = Cast<UDynamicMaterialModelDynamic>(InMaterialModelBase))
 	{
-		PreviewMaterialModelBase = nullptr;
-	}
-	else
-	{
-		PreviewMaterialModelBase = Cast<UDynamicMaterialModelBase>(StaticDuplicateObject(InMaterialModelBase, GetTransientPackage()));
-		PreviewMaterialModelBase->SetDynamicMaterialInstance(nullptr);
-
-		if (UDynamicMaterialModelDynamic* MaterialModelDynamic = Cast<UDynamicMaterialModelDynamic>(PreviewMaterialModelBase))
-		{
-			MaterialModelDynamic->EnsureComponents();
-		}
-
-		if (UDynamicMaterialModelEditorOnlyData* EditorOnlyData = UDynamicMaterialModelEditorOnlyData::Get(PreviewMaterialModelBase.Get()))
-		{
-			EditorOnlyData->RequestMaterialBuild(/* Immediate */ true);
-		}
-
-		UDynamicMaterialInstanceFactory* Factory = NewObject<UDynamicMaterialInstanceFactory>();
-
-		UDynamicMaterialInstance* NewInstance = Cast<UDynamicMaterialInstance>(Factory->FactoryCreateNew(
-			UDynamicMaterialInstance::StaticClass(),
-			GetTransientPackage(),
-			NAME_None,
-			RF_Transactional,
-			PreviewMaterialModelBase,
-			nullptr
-		));
-
-		PreviewMaterialInstance = NewInstance;
+		MaterialModelDynamic->EnsureComponents();
 	}
 
 	EditGlobalSettings();
@@ -255,9 +217,9 @@ void SDMMaterialEditor::SetMaterialModelBase(UDynamicMaterialModelBase* InMateri
 	BindEditorOnlyDataUpdate(InMaterialModelBase);
 }
 
-UDynamicMaterialModel* SDMMaterialEditor::GetPreviewMaterialModel() const
+UDynamicMaterialModel* SDMMaterialEditor::GetMaterialModel() const
 {
-	if (UDynamicMaterialModelBase* MaterialModelBase = GetPreviewMaterialModelBase())
+	if (UDynamicMaterialModelBase* MaterialModelBase = MaterialModelBaseWeak.Get())
 	{
 		return MaterialModelBase->ResolveMaterialModel();
 	}
@@ -267,7 +229,7 @@ UDynamicMaterialModel* SDMMaterialEditor::GetPreviewMaterialModel() const
 
 bool SDMMaterialEditor::IsDynamicModel() const
 {
-	return !!Cast<UDynamicMaterialModelDynamic>(GetPreviewMaterialModelBase());
+	return !!Cast<UDynamicMaterialModelDynamic>(MaterialModelBaseWeak.Get());
 }
 
 const FDMObjectMaterialProperty* SDMMaterialEditor::GetMaterialObjectProperty() const
@@ -361,7 +323,7 @@ void SDMMaterialEditor::SelectProperty(EDMMaterialPropertyType InProperty, bool
 
 	PageHistoryAdd({EDMMaterialEditorMode::EditSlot, InProperty});
 
-	if (UDynamicMaterialModelEditorOnlyData* EditorOnlyData = UDynamicMaterialModelEditorOnlyData::Get(GetPreviewMaterialModelBase()))
+	if (UDynamicMaterialModelEditorOnlyData* EditorOnlyData = UDynamicMaterialModelEditorOnlyData::Get(MaterialModelBaseWeak))
 	{
 		UDMMaterialSlot* CurrentSlot = SelectionContext.Slot.Get();
 
@@ -476,7 +438,7 @@ void SDMMaterialEditor::OnEffectSelected(const TSharedRef<SDMMaterialSlotLayerEf
 
 void SDMMaterialEditor::OpenMaterialPreviewTab()
 {
-	UDynamicMaterialModelBase* MaterialModelBase = GetPreviewMaterialModelBase();
+	UDynamicMaterialModelBase* MaterialModelBase = GetMaterialModelBase();
 
 	if (!MaterialModelBase)
 	{
@@ -542,7 +504,7 @@ void SDMMaterialEditor::CloseMaterialPreviewTab()
 
 TSharedPtr<IToolTip> SDMMaterialEditor::GetMaterialPreviewToolTip()
 {
-	UDynamicMaterialModelBase* MaterialModelBase = GetPreviewMaterialModelBase();
+	UDynamicMaterialModelBase* MaterialModelBase = GetMaterialModelBase();
 
 	if (!MaterialModelBase)
 	{
@@ -607,7 +569,7 @@ void SDMMaterialEditor::Validate()
 		return;
 	}
 
-	UDynamicMaterialModelBase* MaterialModelBase = GetPreviewMaterialModelBase();
+	UDynamicMaterialModelBase* MaterialModelBase = GetMaterialModelBase();
 
 	if (!IsValid(MaterialModelBase))
 	{
@@ -709,41 +671,6 @@ void SDMMaterialEditor::PostRedo(bool bInSuccess)
 	OnUndo();
 }
 
-void SDMMaterialEditor::AddReferencedObjects(FReferenceCollector& Collector)
-{
-	if (PreviewMaterialInstance)
-	{
-		Collector.AddReferencedObject(PreviewMaterialInstance);
-	}
-
-	if (PreviewMaterialModelBase)
-	{
-		Collector.AddReferencedObject(PreviewMaterialModelBase);
-	}
-}
-
-FString SDMMaterialEditor::GetReferencerName() const
-{
-	return TEXT("SDMMaterialEditor");
-}
-
-UPackage* SDMMaterialEditor::GetSaveablePackage(UObject* InObject)
-{
-	if (!IsValid(InObject))
-	{
-		return nullptr;
-	}
-
-	UPackage* Package = InObject->GetPackage();
-
-	if (!Package || Package->HasAllFlags(RF_Transient))
-	{
-		return nullptr;
-	}
-
-	return Package;
-}
-
 void SDMMaterialEditor::BindCommands(SDMMaterialSlotEditor* InSlotEditor)
 {
 	const FGenericCommands& GenericCommands = FGenericCommands::Get();
@@ -829,7 +756,7 @@ void SDMMaterialEditor::BindCommands(SDMMaterialSlotEditor* InSlotEditor)
 
 bool SDMMaterialEditor::IsPropertyValidForModel(EDMMaterialPropertyType InProperty) const
 {
-	UDynamicMaterialModelEditorOnlyData* EditorOnlyData = UDynamicMaterialModelEditorOnlyData::Get(GetPreviewMaterialModelBase());
+	UDynamicMaterialModelEditorOnlyData* EditorOnlyData = UDynamicMaterialModelEditorOnlyData::Get(MaterialModelBaseWeak);
 
 	if (!EditorOnlyData)
 	{
@@ -986,57 +913,6 @@ void SDMMaterialEditor::PageHistoryAdd(const FDMMaterialEditorPage& InPage)
 	SelectionContext.PageHistoryCount = NewPageIndex + 1;
 }
 
-void SDMMaterialEditor::ApplyToOriginal()
-{
-	UDynamicMaterialModelBase* Copy = GetPreviewMaterialModelBase();
-
-	if (!Copy)
-	{
-		return;
-	}
-
-	UDynamicMaterialModelBase* Original = GetOriginalMaterialModelBase();
-
-	if (!Original)
-	{
-		return;
-	}
-
-	if (Original->IsSelected())
-	{
-		GEditor->GetSelectedObjects()->Deselect(Original);
-	}
-
-	ResetLoaders(Original->GetPackage());
-
-	UDynamicMaterialInstance* Outer = Cast<UDynamicMaterialInstance>(Original->GetOuter());
-
-	Original = Cast<UDynamicMaterialModelBase>(StaticDuplicateObject(
-		Copy, 
-		Original->GetOuter(),
-		Original->GetFName(),
-		RF_AllFlags, 
-		Original->GetClass()
-	));
-
-	Original->SetDynamicMaterialInstance(Outer);
-
-	if (Outer)
-	{
-		Outer->SetMaterialModel(Original);
-	}
-
-	if (UDynamicMaterialModelEditorOnlyData* EditorOnlyData = UDynamicMaterialModelEditorOnlyData::Get(Original))
-	{
-		EditorOnlyData->RequestMaterialBuild(/* Immediate */ true);
-	}
-
-	if (Outer)
-	{
-		Outer->InitializeMIDPublic();
-	}
-}
-
 bool SDMMaterialEditor::SetActivePage(const FDMMaterialEditorPage& InPage)
 {
 	switch (InPage.EditorMode)
@@ -1060,21 +936,6 @@ bool SDMMaterialEditor::SetActivePage(const FDMMaterialEditorPage& InPage)
 	}
 }
 
-void SDMMaterialEditor::SaveOriginal()
-{
-	if (UDynamicMaterialModelBase* MaterialModelBase = GetOriginalMaterialModelBase())
-	{
-		ApplyToOriginal();
-
-		if (!!GetSaveablePackage(MaterialModelBase))
-		{
-			TArray<UObject*> AssetsToSave;
-			AssetsToSave.Add(MaterialModelBase);
-			UPackageTools::SavePackagesForObjects(AssetsToSave);
-		}
-	}
-}
-
 void SDMMaterialEditor::HandleDrop_CreateTextureSet(const TArray<FAssetData>& InTextureAssets)
 {
 	if (InTextureAssets.Num() < 2)
@@ -1104,7 +965,7 @@ void SDMMaterialEditor::HandleDrop_TextureSet(UDMTextureSet* InTextureSet)
 		return;
 	}
 
-	UDynamicMaterialModel* MaterialModel = GetPreviewMaterialModel();
+	UDynamicMaterialModel* MaterialModel = GetMaterialModel();
 
 	if (!MaterialModel)
 	{
@@ -1243,7 +1104,7 @@ TSharedRef<SDMToolBar> SDMMaterialEditor::CreateSlot_ToolBar()
 
 TSharedRef<SDMMaterialGlobalSettingsEditor> SDMMaterialEditor::CreateSlot_GlobalSettingsEditor()
 {
-	return SNew(SDMMaterialGlobalSettingsEditor, SharedThis(this), GetPreviewMaterialModelBase());
+	return SNew(SDMMaterialGlobalSettingsEditor, SharedThis(this), GetMaterialModelBase());
 }
 
 TSharedRef<SDMMaterialProperties> SDMMaterialEditor::CreateSlot_MaterialProperties()
@@ -1256,7 +1117,7 @@ TSharedRef<SWidget> SDMMaterialEditor::CreateSlot_Preview()
 	return SNew(SOverlay)
 		+ SOverlay::Slot()
 		[
-			SNew(SDMMaterialPreview, SharedThis(this), GetPreviewMaterialModelBase())
+			SNew(SDMMaterialPreview, SharedThis(this), GetMaterialModelBase())
 		]
 		+ SOverlay::Slot()
 		.HAlign(EHorizontalAlignment::HAlign_Left)
@@ -1279,7 +1140,7 @@ TSharedRef<SDMMaterialPropertySelector> SDMMaterialEditor::CreateSlot_PropertySe
 
 	if (SelectionContext.EditorMode == EDMMaterialEditorMode::EditSlot && SelectionContext.Property == EDMMaterialPropertyType::None)
 	{
-		if (UDynamicMaterialModel* MaterialModel = GetPreviewMaterialModel())
+		if (UDynamicMaterialModel* MaterialModel = GetMaterialModel())
 		{
 			if (UDynamicMaterialModelEditorOnlyData* EditorOnlyData = UDynamicMaterialModelEditorOnlyData::Get(MaterialModel))
 			{
@@ -1333,12 +1194,12 @@ TSharedRef<SDMMaterialComponentEditor> SDMMaterialEditor::CreateSlot_ComponentEd
 
 TSharedRef<SDMStatusBar> SDMMaterialEditor::CreateSlot_StatusBar()
 {
-	return SNew(SDMStatusBar, SharedThis(this), GetPreviewMaterialModelBase());
+	return SNew(SDMStatusBar, SharedThis(this), GetMaterialModelBase());
 }
 
 void SDMMaterialEditor::OnUndo()
 {
-	UDynamicMaterialModelBase* MaterialModel = GetPreviewMaterialModelBase();
+	UDynamicMaterialModelBase* MaterialModel = GetMaterialModelBase();
 
 	if (!IsValid(MaterialModel))
 	{
@@ -1372,7 +1233,7 @@ void SDMMaterialEditor::ApplySelectionContext()
 
 	if (SelectionContext.EditorMode == EDMMaterialEditorMode::EditSlot)
 	{
-		if (UDynamicMaterialModelEditorOnlyData* EditorOnlyData = UDynamicMaterialModelEditorOnlyData::Get(GetPreviewMaterialModelBase()))
+		if (UDynamicMaterialModelEditorOnlyData* EditorOnlyData = UDynamicMaterialModelEditorOnlyData::Get(MaterialModelBaseWeak))
 		{
 			if (UDMMaterialSlot* Slot = SelectionContext.Slot.Get())
 			{
diff --git a/Engine/Plugins/Experimental/DynamicMaterial/Source/DynamicMaterialEditor/Private/UI/Widgets/SDMMaterialEditor.h b/Engine/Plugins/Experimental/DynamicMaterial/Source/DynamicMaterialEditor/Private/UI/Widgets/SDMMaterialEditor.h
--- a/Engine/Plugins/Experimental/DynamicMaterial/Source/DynamicMaterialEditor/Private/UI/Widgets/SDMMaterialEditor.h
+++ b/Engine/Plugins/Experimental/DynamicMaterial/Source/DynamicMaterialEditor/Private/UI/Widgets/SDMMaterialEditor.h
@@ -51,7 +51,7 @@ namespace UE::DynamicMaterialEditor::Private
 	inline const TCHAR* EditorLightBackground = TEXT("Brushes.Header");
 }
 
-class SDMMaterialEditor : public SCompoundWidget, public FSelfRegisteringEditorUndoClient, public FGCObject
+class SDMMaterialEditor : public SCompoundWidget, public FSelfRegisteringEditorUndoClient
 {
 	SLATE_DECLARE_WIDGET(SDMMaterialEditor, SCompoundWidget)
 
@@ -64,8 +64,6 @@ class SDMMaterialEditor : public SCompoundWidget, public FSelfRegisteringEditorU
 	SLATE_END_ARGS()
 
 public:
-	static UPackage* GetSaveablePackage(UObject* InObject);
-
 	SDMMaterialEditor();
 
 	virtual ~SDMMaterialEditor() override;
@@ -75,11 +73,9 @@ public:
 	TSharedPtr<SDMMaterialDesigner> GetDesignerWidget() const;
 
 	/** Material Model / Actor */
-	UDynamicMaterialModelBase* GetOriginalMaterialModelBase() const;
-
-	UDynamicMaterialModelBase* GetPreviewMaterialModelBase() const;
+	UDynamicMaterialModelBase* GetMaterialModelBase() const;
 
-	UDynamicMaterialModel* GetPreviewMaterialModel() const;
+	UDynamicMaterialModel* GetMaterialModel() const;
 
 	bool IsDynamicModel() const;
 
@@ -143,10 +139,6 @@ public:
 
 	bool SetActivePage(const FDMMaterialEditorPage& InPage);
 
-	void SaveOriginal();
-
-	void ApplyToOriginal();
-
 	void HandleDrop_CreateTextureSet(const TArray<FAssetData>& InTextureAssets);
 
 	void HandleDrop_TextureSet(UDMTextureSet* InTextureSet);
@@ -162,11 +154,6 @@ public:
 	virtual void PostRedo(bool bInSuccess) override;
 	//~ End FUndoClient
 
-	//~ Begin FGCObject
-	virtual void AddReferencedObjects(FReferenceCollector& Collector) override;
-	virtual FString GetReferencerName() const override;
-	//~ End FGCObject
-
 protected:
 	TWeakPtr<SDMMaterialDesigner> DesignerWidgetWeak;
 
@@ -182,9 +169,7 @@ protected:
 	TDMWidgetSlot<SDMMaterialComponentEditor> ComponentEditorSlot;
 	TDMWidgetSlot<SDMStatusBar> StatusBarSlot;
 
-	TWeakObjectPtr<UDynamicMaterialModelBase> OriginalMaterialModelBaseWeak;
-	TObjectPtr<UDynamicMaterialInstance> PreviewMaterialInstance;
-	TObjectPtr<UDynamicMaterialModelBase> PreviewMaterialModelBase;
+	TWeakObjectPtr<UDynamicMaterialModelBase> MaterialModelBaseWeak;
 	TOptional<FDMObjectMaterialProperty> ObjectMaterialPropertyOpt;
 
 	TSharedRef<FUICommandList> CommandList;
diff --git a/Engine/Plugins/Experimental/DynamicMaterial/Source/DynamicMaterialEditor/Private/UI/Widgets/Visualizers/SDMMaterialComponentPreview.cpp b/Engine/Plugins/Experimental/DynamicMaterial/Source/DynamicMaterialEditor/Private/UI/Widgets/Visualizers/SDMMaterialComponentPreview.cpp
--- a/Engine/Plugins/Experimental/DynamicMaterial/Source/DynamicMaterialEditor/Private/UI/Widgets/Visualizers/SDMMaterialComponentPreview.cpp
+++ b/Engine/Plugins/Experimental/DynamicMaterial/Source/DynamicMaterialEditor/Private/UI/Widgets/Visualizers/SDMMaterialComponentPreview.cpp
@@ -55,7 +55,7 @@ void SDMMaterialComponentPreview::Construct(const FArguments& InArgs, const TSha
 
 	PreviewMaterialBaseWeak = InEditorWidget->GetPreviewMaterialManager()->CreatePreviewMaterial(InComponent);
 	PreviewMaterialDynamicWeak = InEditorWidget->GetPreviewMaterialManager()->CreatePreviewMaterialDynamic(PreviewMaterialBaseWeak.Get());
-	MaterialModelBaseWeak = InEditorWidget->GetPreviewMaterialModelBase();
+	MaterialModelBaseWeak = InEditorWidget->GetMaterialModelBase();
 
 	SetCanTick(true);
 
@@ -211,7 +211,7 @@ void SDMMaterialComponentPreview::RecreateMaterial()
 		EditorWidget->GetPreviewMaterialManager()->FreePreviewMaterialDynamic(PreviewMaterialBase);
 		PreviewMaterialDynamicWeak = EditorWidget->GetPreviewMaterialManager()->CreatePreviewMaterialDynamic(PreviewMaterialBase);
 
-		UDynamicMaterialModelBase* MaterialModelBase = EditorWidget->GetPreviewMaterialModelBase();
+		UDynamicMaterialModelBase* MaterialModelBase = EditorWidget->GetMaterialModelBase();
 
 		if (UDynamicMaterialModel* MaterialModel = Cast<UDynamicMaterialModel>(MaterialModelBase))
 		{
@@ -231,7 +231,7 @@ void SDMMaterialComponentPreview::RecreateMaterial()
 		EditorWidget->GetPreviewMaterialManager()->FreePreviewMaterialDynamic(PreviewMaterialBase);
 		PreviewMaterialDynamicWeak = EditorWidget->GetPreviewMaterialManager()->CreatePreviewMaterialDynamic(PreviewMaterialBase);
 
-		UDynamicMaterialModelBase* MaterialModelBase = EditorWidget->GetPreviewMaterialModelBase();
+		UDynamicMaterialModelBase* MaterialModelBase = EditorWidget->GetMaterialModelBase();
 
 		if (UDynamicMaterialModel* MaterialModel = Cast<UDynamicMaterialModel>(MaterialModelBase))
 		{
