Material Designer: Apply button
- Materials edited in the designer are now preview materials which are applied to the source asset when clicking apply or the save buttons.
- Optional automatic transfer to values existing on both the original and preview material to the original when changing on the preview.
- Property row generators now generate property handles for both the preview and original materials.
- Property extension buttons are now based on the original material, rather than the preview (e.g. create keyframe, expose to remote control) - except for the reset to default button, which is still based on the preview.
- Properties that don't yet exist in the original material are marked with a red (!) instead of those extension buttons.

diff --git a/Engine/Plugins/Experimental/DynamicMaterial/Source/DynamicMaterial/Private/Components/DMMaterialValue.cpp b/Engine/Plugins/Experimental/DynamicMaterial/Source/DynamicMaterial/Private/Components/DMMaterialValue.cpp
--- a/Engine/Plugins/Experimental/DynamicMaterial/Source/DynamicMaterial/Private/Components/DMMaterialValue.cpp
+++ b/Engine/Plugins/Experimental/DynamicMaterial/Source/DynamicMaterial/Private/Components/DMMaterialValue.cpp
@@ -12,6 +12,7 @@
 #include "DMValueDefinition.h"
 #include "Modules/ModuleManager.h"
 #include "PropertyEditorModule.h"
+#include "PropertyHandle.h"
 #endif // WITH_EDITOR
 
 #define LOCTEXT_NAMESPACE "DMMaterialValue"
@@ -68,12 +69,12 @@ UDMMaterialValue::UDMMaterialValue(EDMValueType InType)
 	EditableProperties.Add(ValueName);
 #endif
 }
- 
+
 UDynamicMaterialModel* UDMMaterialValue::GetMaterialModel() const
 {
 	return Cast<UDynamicMaterialModel>(GetOuterSafe());
 }
- 
+
 #if WITH_EDITOR
 FText UDMMaterialValue::GetTypeName() const
 {
@@ -89,6 +90,15 @@ FText UDMMaterialValue::GetDescription() const
 	);
 }
 
+void UDMMaterialValue::ResetToDefault(TSharedPtr<IPropertyHandle> InPropertyHandle)
+{
+	// -> represents a child property. Child property resets are handled by individual class implementations.
+	if (InPropertyHandle->GetPropertyPath().Find(TEXT("->")) == INDEX_NONE)
+	{
+		ApplyDefaultValue();
+	}
+}
+
 void UDMMaterialValue::PostCDOContruct()
 {
 	Super::PostCDOContruct();
diff --git a/Engine/Plugins/Experimental/DynamicMaterial/Source/DynamicMaterial/Private/Components/DMMaterialValueDynamic.cpp b/Engine/Plugins/Experimental/DynamicMaterial/Source/DynamicMaterial/Private/Components/DMMaterialValueDynamic.cpp
--- a/Engine/Plugins/Experimental/DynamicMaterial/Source/DynamicMaterial/Private/Components/DMMaterialValueDynamic.cpp
+++ b/Engine/Plugins/Experimental/DynamicMaterial/Source/DynamicMaterial/Private/Components/DMMaterialValueDynamic.cpp
@@ -8,6 +8,7 @@
 
 #if WITH_EDITOR
 #include "Materials/Material.h"
+#include "PropertyHandle.h"
 #include "UObject/Package.h"
 #endif
 
@@ -41,6 +42,14 @@ UDMMaterialValueDynamic* UDMMaterialValueDynamic::CreateValueDynamic(TSubclassOf
 
 	return NewValueDynamic;
 }
+
+void UDMMaterialValueDynamic::ResetToDefault(TSharedPtr<IPropertyHandle> InPropertyHandle)
+{
+	if (InPropertyHandle->GetPropertyPath().Find(TEXT("->")) == INDEX_NONE)
+	{
+		ApplyDefaultValue();
+	}
+}
 #endif
 
 UDMMaterialValue* UDMMaterialValueDynamic::GetParentValue() const
diff --git a/Engine/Plugins/Experimental/DynamicMaterial/Source/DynamicMaterial/Private/Components/DMTextureUV.cpp b/Engine/Plugins/Experimental/DynamicMaterial/Source/DynamicMaterial/Private/Components/DMTextureUV.cpp
--- a/Engine/Plugins/Experimental/DynamicMaterial/Source/DynamicMaterial/Private/Components/DMTextureUV.cpp
+++ b/Engine/Plugins/Experimental/DynamicMaterial/Source/DynamicMaterial/Private/Components/DMTextureUV.cpp
@@ -638,6 +638,15 @@ void UDMTextureUV::RemoveParameterNames()
 }
 #endif
 
+void UDMTextureUV::CopyParametersFrom_Implementation(UObject* InOther)
+{
+	UDMTextureUV* OtherTextureUV = CastChecked<UDMTextureUV>(InOther);
+	OtherTextureUV->SetOffset(Offset);
+	OtherTextureUV->SetTiling(Tiling);
+	OtherTextureUV->SetPivot(Pivot);
+	OtherTextureUV->SetRotation(Rotation);
+}
+
 void UDMTextureUV::OnTextureUVChanged(EDMUpdateType InUpdateType)
 {
 	if (!IsComponentValid())
diff --git a/Engine/Plugins/Experimental/DynamicMaterial/Source/DynamicMaterial/Private/Components/DMTextureUVDynamic.cpp b/Engine/Plugins/Experimental/DynamicMaterial/Source/DynamicMaterial/Private/Components/DMTextureUVDynamic.cpp
--- a/Engine/Plugins/Experimental/DynamicMaterial/Source/DynamicMaterial/Private/Components/DMTextureUVDynamic.cpp
+++ b/Engine/Plugins/Experimental/DynamicMaterial/Source/DynamicMaterial/Private/Components/DMTextureUVDynamic.cpp
@@ -236,6 +236,15 @@ void UDMTextureUVDynamic::OnTextureUVChanged()
 	}
 }
 
+void UDMTextureUVDynamic::CopyParametersFrom_Implementation(UObject* InOther)
+{
+	UDMTextureUVDynamic* OtherTextureUVDynamic = CastChecked<UDMTextureUVDynamic>(InOther);
+	OtherTextureUVDynamic->SetOffset(Offset);
+	OtherTextureUVDynamic->SetTiling(Tiling);
+	OtherTextureUVDynamic->SetPivot(Pivot);
+	OtherTextureUVDynamic->SetRotation(Rotation);
+}
+
 #if WITH_EDITOR
 void UDMTextureUVDynamic::OnComponentAdded()
 {
diff --git a/Engine/Plugins/Experimental/DynamicMaterial/Source/DynamicMaterial/Private/Components/MaterialValues/DMMaterialValueBool.cpp b/Engine/Plugins/Experimental/DynamicMaterial/Source/DynamicMaterial/Private/Components/MaterialValues/DMMaterialValueBool.cpp
--- a/Engine/Plugins/Experimental/DynamicMaterial/Source/DynamicMaterial/Private/Components/MaterialValues/DMMaterialValueBool.cpp
+++ b/Engine/Plugins/Experimental/DynamicMaterial/Source/DynamicMaterial/Private/Components/MaterialValues/DMMaterialValueBool.cpp
@@ -99,6 +99,12 @@ FText UDMMaterialValueBool::GetComponentDescription() const
 	return LOCTEXT("Bool", "Bool");
 }
 
+void UDMMaterialValueBool::CopyParametersFrom_Implementation(UObject* InOther)
+{
+	UDMMaterialValueBool* OtherValue = CastChecked<UDMMaterialValueBool>(InOther);
+	OtherValue->SetValue(GetValue());
+}
+
 TSharedPtr<FJsonValue> UDMMaterialValueBool::JsonSerialize() const
 {
 	return FDMJsonUtils::Serialize(Value);
diff --git a/Engine/Plugins/Experimental/DynamicMaterial/Source/DynamicMaterial/Private/Components/MaterialValues/DMMaterialValueColorAtlas.cpp b/Engine/Plugins/Experimental/DynamicMaterial/Source/DynamicMaterial/Private/Components/MaterialValues/DMMaterialValueColorAtlas.cpp
--- a/Engine/Plugins/Experimental/DynamicMaterial/Source/DynamicMaterial/Private/Components/MaterialValues/DMMaterialValueColorAtlas.cpp
+++ b/Engine/Plugins/Experimental/DynamicMaterial/Source/DynamicMaterial/Private/Components/MaterialValues/DMMaterialValueColorAtlas.cpp
@@ -163,6 +163,12 @@ FText UDMMaterialValueColorAtlas::GetComponentDescription() const
 	return LOCTEXT("ColorAtlas", "Color Atlas");
 }
 
+void UDMMaterialValueColorAtlas::CopyParametersFrom_Implementation(UObject* InOther)
+{
+	UDMMaterialValueColorAtlas* OtherValue = CastChecked<UDMMaterialValueColorAtlas>(InOther);
+	OtherValue->SetValue(GetValue());
+}
+
 TSharedPtr<FJsonValue> UDMMaterialValueColorAtlas::JsonSerialize() const
 {
 	return FDMJsonUtils::Serialize({
diff --git a/Engine/Plugins/Experimental/DynamicMaterial/Source/DynamicMaterial/Private/Components/MaterialValues/DMMaterialValueFloat1.cpp b/Engine/Plugins/Experimental/DynamicMaterial/Source/DynamicMaterial/Private/Components/MaterialValues/DMMaterialValueFloat1.cpp
--- a/Engine/Plugins/Experimental/DynamicMaterial/Source/DynamicMaterial/Private/Components/MaterialValues/DMMaterialValueFloat1.cpp
+++ b/Engine/Plugins/Experimental/DynamicMaterial/Source/DynamicMaterial/Private/Components/MaterialValues/DMMaterialValueFloat1.cpp
@@ -80,6 +80,12 @@ FText UDMMaterialValueFloat1::GetComponentDescription() const
 	return LOCTEXT("Scalar", "Scalar");
 }
 
+void UDMMaterialValueFloat1::CopyParametersFrom_Implementation(UObject* InOther)
+{
+	UDMMaterialValueFloat1* OtherValue = CastChecked<UDMMaterialValueFloat1>(InOther);
+	OtherValue->SetValue(GetValue());
+}
+
 TSharedPtr<FJsonValue> UDMMaterialValueFloat1::JsonSerialize() const
 {
 	return FDMJsonUtils::Serialize(Value);
diff --git a/Engine/Plugins/Experimental/DynamicMaterial/Source/DynamicMaterial/Private/Components/MaterialValues/DMMaterialValueFloat2.cpp b/Engine/Plugins/Experimental/DynamicMaterial/Source/DynamicMaterial/Private/Components/MaterialValues/DMMaterialValueFloat2.cpp
--- a/Engine/Plugins/Experimental/DynamicMaterial/Source/DynamicMaterial/Private/Components/MaterialValues/DMMaterialValueFloat2.cpp
+++ b/Engine/Plugins/Experimental/DynamicMaterial/Source/DynamicMaterial/Private/Components/MaterialValues/DMMaterialValueFloat2.cpp
@@ -12,6 +12,7 @@
 #include "Materials/MaterialExpressionVectorParameter.h"
 #include "Model/IDMMaterialBuildStateInterface.h"
 #include "Model/IDMMaterialBuildUtilsInterface.h"
+#include "PropertyHandle.h"
 #include "Utils/DMUtils.h"
 #endif
 
@@ -65,7 +66,34 @@ void UDMMaterialValueFloat2::ApplyDefaultValue()
 {
 	SetValue(DefaultValue);
 }
- 
+
+void UDMMaterialValueFloat2::ResetToDefault(TSharedPtr<IPropertyHandle> InPropertyHandle)
+{
+	const FStringView PropertyPath = InPropertyHandle->GetPropertyPath();
+	const int32 IndexOfStructSeparator = PropertyPath.Find(TEXT("->"));
+
+	if (IndexOfStructSeparator == INDEX_NONE)
+	{
+		Super::ResetToDefault(InPropertyHandle);
+		return;
+	}
+
+	FStringView LeafName = PropertyPath.RightChop(IndexOfStructSeparator + 2);
+
+	if (LeafName.Equals(TEXT("X")))
+	{
+		FVector2D CurrentValue = GetValue();
+		CurrentValue.X = GetDefaultValue().X;
+		SetValue(CurrentValue);
+	}
+	else if (LeafName.Equals(TEXT("Y")))
+	{
+		FVector2D CurrentValue = GetValue();
+		CurrentValue.Y = GetDefaultValue().Y;
+		SetValue(CurrentValue);
+	}
+}
+
 void UDMMaterialValueFloat2::ResetDefaultValue()
 {
 	DefaultValue = FVector2D::ZeroVector;
@@ -89,6 +117,12 @@ FText UDMMaterialValueFloat2::GetComponentDescription() const
 	return LOCTEXT("Vector2", "Vector 2");
 }
 
+void UDMMaterialValueFloat2::CopyParametersFrom_Implementation(UObject* InOther)
+{
+	UDMMaterialValueFloat2* OtherValue = Cast<UDMMaterialValueFloat2>(InOther);
+	OtherValue->SetValue(GetValue());
+}
+
 TSharedPtr<FJsonValue> UDMMaterialValueFloat2::JsonSerialize() const
 {
 	return FDMJsonUtils::Serialize(Value);
diff --git a/Engine/Plugins/Experimental/DynamicMaterial/Source/DynamicMaterial/Private/Components/MaterialValues/DMMaterialValueFloat3RGB.cpp b/Engine/Plugins/Experimental/DynamicMaterial/Source/DynamicMaterial/Private/Components/MaterialValues/DMMaterialValueFloat3RGB.cpp
--- a/Engine/Plugins/Experimental/DynamicMaterial/Source/DynamicMaterial/Private/Components/MaterialValues/DMMaterialValueFloat3RGB.cpp
+++ b/Engine/Plugins/Experimental/DynamicMaterial/Source/DynamicMaterial/Private/Components/MaterialValues/DMMaterialValueFloat3RGB.cpp
@@ -11,6 +11,7 @@
 #include "Materials/MaterialExpressionVectorParameter.h"
 #include "Model/IDMMaterialBuildStateInterface.h"
 #include "Model/IDMMaterialBuildUtilsInterface.h"
+#include "PropertyHandle.h"
 #include "Utils/DMUtils.h"
 #endif
  
@@ -62,7 +63,40 @@ void UDMMaterialValueFloat3RGB::ApplyDefaultValue()
 {
 	SetValue(DefaultValue);
 }
- 
+
+void UDMMaterialValueFloat3RGB::ResetToDefault(TSharedPtr<IPropertyHandle> InPropertyHandle)
+{
+	const FStringView PropertyPath = InPropertyHandle->GetPropertyPath();
+	const int32 IndexOfStructSeparator = PropertyPath.Find(TEXT("->"));
+
+	if (IndexOfStructSeparator == INDEX_NONE)
+	{
+		Super::ResetToDefault(InPropertyHandle);
+		return;
+	}
+
+	FStringView LeafName = PropertyPath.RightChop(IndexOfStructSeparator + 2);
+
+	if (LeafName.Equals(TEXT("R")))
+	{
+		FLinearColor CurrentValue = GetValue();
+		CurrentValue.R = GetDefaultValue().R;
+		SetValue(CurrentValue);
+	}
+	else if (LeafName.Equals(TEXT("G")))
+	{
+		FLinearColor CurrentValue = GetValue();
+		CurrentValue.G = GetDefaultValue().G;
+		SetValue(CurrentValue);
+	}
+	else if (LeafName.Equals(TEXT("B")))
+	{
+		FLinearColor CurrentValue = GetValue();
+		CurrentValue.B = GetDefaultValue().B;
+		SetValue(CurrentValue);
+	}
+}
+
 void UDMMaterialValueFloat3RGB::ResetDefaultValue()
 {
 	DefaultValue = FLinearColor(0, 0, 0, 1);
@@ -86,6 +120,12 @@ FText UDMMaterialValueFloat3RGB::GetComponentDescription() const
 	return LOCTEXT("ColorRGB", "Color (RGB)");
 }
 
+void UDMMaterialValueFloat3RGB::CopyParametersFrom_Implementation(UObject* InOther)
+{
+	UDMMaterialValueFloat3RGB* OtherValue = CastChecked<UDMMaterialValueFloat3RGB>(InOther);
+	OtherValue->SetValue(GetValue());
+}
+
 TSharedPtr<FJsonValue> UDMMaterialValueFloat3RGB::JsonSerialize() const
 {
 	return FDMJsonUtils::Serialize(Value);
diff --git a/Engine/Plugins/Experimental/DynamicMaterial/Source/DynamicMaterial/Private/Components/MaterialValues/DMMaterialValueFloat3RPY.cpp b/Engine/Plugins/Experimental/DynamicMaterial/Source/DynamicMaterial/Private/Components/MaterialValues/DMMaterialValueFloat3RPY.cpp
--- a/Engine/Plugins/Experimental/DynamicMaterial/Source/DynamicMaterial/Private/Components/MaterialValues/DMMaterialValueFloat3RPY.cpp
+++ b/Engine/Plugins/Experimental/DynamicMaterial/Source/DynamicMaterial/Private/Components/MaterialValues/DMMaterialValueFloat3RPY.cpp
@@ -11,6 +11,7 @@
 #include "Materials/MaterialExpressionVectorParameter.h"
 #include "Model/IDMMaterialBuildStateInterface.h"
 #include "Model/IDMMaterialBuildUtilsInterface.h"
+#include "PropertyHandle.h"
 #include "Utils/DMUtils.h"
 #endif
  
@@ -62,7 +63,40 @@ void UDMMaterialValueFloat3RPY::ApplyDefaultValue()
 {
 	SetValue(DefaultValue);
 }
- 
+
+void UDMMaterialValueFloat3RPY::ResetToDefault(TSharedPtr<IPropertyHandle> InPropertyHandle)
+{
+	const FStringView PropertyPath = InPropertyHandle->GetPropertyPath();
+	const int32 IndexOfStructSeparator = PropertyPath.Find(TEXT("->"));
+
+	if (IndexOfStructSeparator == INDEX_NONE)
+	{
+		Super::ResetToDefault(InPropertyHandle);
+		return;
+	}
+
+	FStringView LeafName = PropertyPath.RightChop(IndexOfStructSeparator + 2);
+
+	if (LeafName.Equals(TEXT("Roll")))
+	{
+		FRotator CurrentValue = GetValue();
+		CurrentValue.Roll = GetDefaultValue().Roll;
+		SetValue(CurrentValue);
+	}
+	else if (LeafName.Equals(TEXT("Pitch")))
+	{
+		FRotator CurrentValue = GetValue();
+		CurrentValue.Pitch = GetDefaultValue().Pitch;
+		SetValue(CurrentValue);
+	}
+	else if (LeafName.Equals(TEXT("Yaw")))
+	{
+		FRotator CurrentValue = GetValue();
+		CurrentValue.Yaw = GetDefaultValue().Yaw;
+		SetValue(CurrentValue);
+	}
+}
+
 void UDMMaterialValueFloat3RPY::ResetDefaultValue()
 {
 	DefaultValue = FRotator::ZeroRotator;
@@ -86,6 +120,12 @@ FText UDMMaterialValueFloat3RPY::GetComponentDescription() const
 	return LOCTEXT("Rotator", "Rotator");
 }
 
+void UDMMaterialValueFloat3RPY::CopyParametersFrom_Implementation(UObject* InOther)
+{
+	UDMMaterialValueFloat3RPY* OtherValue = CastChecked<UDMMaterialValueFloat3RPY>(InOther);
+	OtherValue->SetValue(GetValue());
+}
+
 TSharedPtr<FJsonValue> UDMMaterialValueFloat3RPY::JsonSerialize() const
 {
 	return FDMJsonUtils::Serialize(Value);
diff --git a/Engine/Plugins/Experimental/DynamicMaterial/Source/DynamicMaterial/Private/Components/MaterialValues/DMMaterialValueFloat3XYZ.cpp b/Engine/Plugins/Experimental/DynamicMaterial/Source/DynamicMaterial/Private/Components/MaterialValues/DMMaterialValueFloat3XYZ.cpp
--- a/Engine/Plugins/Experimental/DynamicMaterial/Source/DynamicMaterial/Private/Components/MaterialValues/DMMaterialValueFloat3XYZ.cpp
+++ b/Engine/Plugins/Experimental/DynamicMaterial/Source/DynamicMaterial/Private/Components/MaterialValues/DMMaterialValueFloat3XYZ.cpp
@@ -11,6 +11,7 @@
 #include "Materials/MaterialExpressionVectorParameter.h"
 #include "Model/IDMMaterialBuildStateInterface.h"
 #include "Model/IDMMaterialBuildUtilsInterface.h"
+#include "PropertyHandle.h"
 #include "Utils/DMUtils.h"
 #endif
  
@@ -62,7 +63,40 @@ void UDMMaterialValueFloat3XYZ::ApplyDefaultValue()
 {
 	SetValue(DefaultValue);
 }
- 
+
+void UDMMaterialValueFloat3XYZ::ResetToDefault(TSharedPtr<IPropertyHandle> InPropertyHandle)
+{
+	const FStringView PropertyPath = InPropertyHandle->GetPropertyPath();
+	const int32 IndexOfStructSeparator = PropertyPath.Find(TEXT("->"));
+
+	if (IndexOfStructSeparator == INDEX_NONE)
+	{
+		Super::ResetToDefault(InPropertyHandle);
+		return;
+	}
+
+	FStringView LeafName = PropertyPath.RightChop(IndexOfStructSeparator + 2);
+
+	if (LeafName.Equals(TEXT("X")))
+	{
+		FVector CurrentValue = GetValue();
+		CurrentValue.X = GetDefaultValue().X;
+		SetValue(CurrentValue);
+	}
+	else if (LeafName.Equals(TEXT("Y")))
+	{
+		FVector CurrentValue = GetValue();
+		CurrentValue.Y = GetDefaultValue().Y;
+		SetValue(CurrentValue);
+	}
+	else if (LeafName.Equals(TEXT("Z")))
+	{
+		FVector CurrentValue = GetValue();
+		CurrentValue.Z = GetDefaultValue().Z;
+		SetValue(CurrentValue);
+	}
+}
+
 void UDMMaterialValueFloat3XYZ::ResetDefaultValue()
 {
 	DefaultValue = FVector::ZeroVector;
@@ -86,6 +120,12 @@ FText UDMMaterialValueFloat3XYZ::GetComponentDescription() const
 	return LOCTEXT("Vector3", "Vector 3");
 }
 
+void UDMMaterialValueFloat3XYZ::CopyParametersFrom_Implementation(UObject* InOther)
+{
+	UDMMaterialValueFloat3XYZ* OtherValue = CastChecked<UDMMaterialValueFloat3XYZ>(InOther);
+	OtherValue->SetValue(GetValue());
+}
+
 TSharedPtr<FJsonValue> UDMMaterialValueFloat3XYZ::JsonSerialize() const
 {
 	return FDMJsonUtils::Serialize(Value);
diff --git a/Engine/Plugins/Experimental/DynamicMaterial/Source/DynamicMaterial/Private/Components/MaterialValues/DMMaterialValueFloat4.cpp b/Engine/Plugins/Experimental/DynamicMaterial/Source/DynamicMaterial/Private/Components/MaterialValues/DMMaterialValueFloat4.cpp
--- a/Engine/Plugins/Experimental/DynamicMaterial/Source/DynamicMaterial/Private/Components/MaterialValues/DMMaterialValueFloat4.cpp
+++ b/Engine/Plugins/Experimental/DynamicMaterial/Source/DynamicMaterial/Private/Components/MaterialValues/DMMaterialValueFloat4.cpp
@@ -11,6 +11,7 @@
 #include "Materials/MaterialExpressionVectorParameter.h"
 #include "Model/IDMMaterialBuildStateInterface.h"
 #include "Model/IDMMaterialBuildUtilsInterface.h"
+#include "PropertyHandle.h"
 #include "Utils/DMUtils.h"
 #endif
  
@@ -69,6 +70,45 @@ void UDMMaterialValueFloat4::ResetDefaultValue()
 	DefaultValue = FLinearColor(0, 0, 0, 1);
 }
 
+void UDMMaterialValueFloat4::ResetToDefault(TSharedPtr<IPropertyHandle> InPropertyHandle)
+{
+	const FStringView PropertyPath = InPropertyHandle->GetPropertyPath();
+	const int32 IndexOfStructSeparator = PropertyPath.Find(TEXT("->"));
+
+	if (IndexOfStructSeparator == INDEX_NONE)
+	{
+		Super::ResetToDefault(InPropertyHandle);
+		return;
+	}
+
+	FStringView LeafName = PropertyPath.RightChop(IndexOfStructSeparator + 2);
+
+	if (LeafName.Equals(TEXT("R")))
+	{
+		FLinearColor CurrentValue = GetValue();
+		CurrentValue.R = GetDefaultValue().R;
+		SetValue(CurrentValue);
+	}
+	else if (LeafName.Equals(TEXT("G")))
+	{
+		FLinearColor CurrentValue = GetValue();
+		CurrentValue.G = GetDefaultValue().G;
+		SetValue(CurrentValue);
+	}
+	else if (LeafName.Equals(TEXT("B")))
+	{
+		FLinearColor CurrentValue = GetValue();
+		CurrentValue.B = GetDefaultValue().B;
+		SetValue(CurrentValue);
+	}
+	else if (LeafName.Equals(TEXT("A")))
+	{
+		FLinearColor CurrentValue = GetValue();
+		CurrentValue.A = GetDefaultValue().A;
+		SetValue(CurrentValue);
+	}
+}
+
 UDMMaterialValueDynamic* UDMMaterialValueFloat4::ToDynamic(UDynamicMaterialModelDynamic* InMaterialModelDynamic)
 {
 	UDMMaterialValueFloat4Dynamic* ValueDynamic = UDMMaterialValueDynamic::CreateValueDynamic<UDMMaterialValueFloat4Dynamic>(InMaterialModelDynamic, this);
@@ -87,6 +127,12 @@ FText UDMMaterialValueFloat4::GetComponentDescription() const
 	return LOCTEXT("ColorRGBA", "Color (RGBA)");
 }
 
+void UDMMaterialValueFloat4::CopyParametersFrom_Implementation(UObject* InOther)
+{
+	UDMMaterialValueFloat4* OtherValue = CastChecked<UDMMaterialValueFloat4>(InOther);
+	OtherValue->SetValue(GetValue());
+}
+
 TSharedPtr<FJsonValue> UDMMaterialValueFloat4::JsonSerialize() const
 {
 	return FDMJsonUtils::Serialize(Value);
diff --git a/Engine/Plugins/Experimental/DynamicMaterial/Source/DynamicMaterial/Private/Components/MaterialValues/DMMaterialValueRenderTarget.cpp b/Engine/Plugins/Experimental/DynamicMaterial/Source/DynamicMaterial/Private/Components/MaterialValues/DMMaterialValueRenderTarget.cpp
--- a/Engine/Plugins/Experimental/DynamicMaterial/Source/DynamicMaterial/Private/Components/MaterialValues/DMMaterialValueRenderTarget.cpp
+++ b/Engine/Plugins/Experimental/DynamicMaterial/Source/DynamicMaterial/Private/Components/MaterialValues/DMMaterialValueRenderTarget.cpp
@@ -346,6 +346,11 @@ void UDMMaterialValueRenderTarget::UpdateRenderTarget()
 	}
 }
 
+void UDMMaterialValueRenderTarget::CopyParametersFrom_Implementation(UObject* InOther)
+{
+	// Overridden to do nothing
+}
+
 UDMMaterialComponent* UDMMaterialValueRenderTarget::GetSubComponentByPath(FDMComponentPath& InPath, const FDMComponentPathSegment& InPathSegment) const
 {
 	if (InPathSegment.GetToken() == RendererPathToken)
diff --git a/Engine/Plugins/Experimental/DynamicMaterial/Source/DynamicMaterial/Private/Components/MaterialValues/DMMaterialValueTexture.cpp b/Engine/Plugins/Experimental/DynamicMaterial/Source/DynamicMaterial/Private/Components/MaterialValues/DMMaterialValueTexture.cpp
--- a/Engine/Plugins/Experimental/DynamicMaterial/Source/DynamicMaterial/Private/Components/MaterialValues/DMMaterialValueTexture.cpp
+++ b/Engine/Plugins/Experimental/DynamicMaterial/Source/DynamicMaterial/Private/Components/MaterialValues/DMMaterialValueTexture.cpp
@@ -156,6 +156,12 @@ void UDMMaterialValueTexture::PostLoad()
 	Type = EDMValueType::VT_Texture;
 }
 
+void UDMMaterialValueTexture::CopyParametersFrom_Implementation(UObject* InOther)
+{
+	UDMMaterialValueTexture* OtherValue = CastChecked<UDMMaterialValueTexture>(InOther);
+	OtherValue->SetValue(GetValue());
+}
+
 #if WITH_EDITOR
 void UDMMaterialValueTexture::ApplyDefaultValue()
 {
diff --git a/Engine/Plugins/Experimental/DynamicMaterial/Source/DynamicMaterial/Private/Components/MaterialValuesDynamic/DMMaterialValueBoolDynamic.cpp b/Engine/Plugins/Experimental/DynamicMaterial/Source/DynamicMaterial/Private/Components/MaterialValuesDynamic/DMMaterialValueBoolDynamic.cpp
--- a/Engine/Plugins/Experimental/DynamicMaterial/Source/DynamicMaterial/Private/Components/MaterialValuesDynamic/DMMaterialValueBoolDynamic.cpp
+++ b/Engine/Plugins/Experimental/DynamicMaterial/Source/DynamicMaterial/Private/Components/MaterialValuesDynamic/DMMaterialValueBoolDynamic.cpp
@@ -41,6 +41,12 @@ void UDMMaterialValueBoolDynamic::CopyDynamicPropertiesTo(UDMMaterialComponent*
 	}
 }
 
+void UDMMaterialValueBoolDynamic::CopyParametersFrom_Implementation(UObject* InOther)
+{
+	UDMMaterialValueBoolDynamic* OtherValue = CastChecked<UDMMaterialValueBoolDynamic>(InOther);
+	OtherValue->SetValue(GetValue());
+}
+
 TSharedPtr<FJsonValue> UDMMaterialValueBoolDynamic::JsonSerialize() const
 {
 	return FDMJsonUtils::Serialize(Value);
diff --git a/Engine/Plugins/Experimental/DynamicMaterial/Source/DynamicMaterial/Private/Components/MaterialValuesDynamic/DMMaterialValueColorAtlasDynamic.cpp b/Engine/Plugins/Experimental/DynamicMaterial/Source/DynamicMaterial/Private/Components/MaterialValuesDynamic/DMMaterialValueColorAtlasDynamic.cpp
--- a/Engine/Plugins/Experimental/DynamicMaterial/Source/DynamicMaterial/Private/Components/MaterialValuesDynamic/DMMaterialValueColorAtlasDynamic.cpp
+++ b/Engine/Plugins/Experimental/DynamicMaterial/Source/DynamicMaterial/Private/Components/MaterialValuesDynamic/DMMaterialValueColorAtlasDynamic.cpp
@@ -41,6 +41,12 @@ void UDMMaterialValueColorAtlasDynamic::CopyDynamicPropertiesTo(UDMMaterialCompo
 	}
 }
 
+void UDMMaterialValueColorAtlasDynamic::CopyParametersFrom_Implementation(UObject* InOther)
+{
+	UDMMaterialValueColorAtlasDynamic* OtherValue = CastChecked<UDMMaterialValueColorAtlasDynamic>(InOther);
+	OtherValue->SetValue(GetValue());
+}
+
 TSharedPtr<FJsonValue> UDMMaterialValueColorAtlasDynamic::JsonSerialize() const
 {
 	return FDMJsonUtils::Serialize(Value);
diff --git a/Engine/Plugins/Experimental/DynamicMaterial/Source/DynamicMaterial/Private/Components/MaterialValuesDynamic/DMMaterialValueFloat1Dynamic.cpp b/Engine/Plugins/Experimental/DynamicMaterial/Source/DynamicMaterial/Private/Components/MaterialValuesDynamic/DMMaterialValueFloat1Dynamic.cpp
--- a/Engine/Plugins/Experimental/DynamicMaterial/Source/DynamicMaterial/Private/Components/MaterialValuesDynamic/DMMaterialValueFloat1Dynamic.cpp
+++ b/Engine/Plugins/Experimental/DynamicMaterial/Source/DynamicMaterial/Private/Components/MaterialValuesDynamic/DMMaterialValueFloat1Dynamic.cpp
@@ -41,6 +41,12 @@ void UDMMaterialValueFloat1Dynamic::CopyDynamicPropertiesTo(UDMMaterialComponent
 	}
 }
 
+void UDMMaterialValueFloat1Dynamic::CopyParametersFrom_Implementation(UObject* InOther)
+{
+	UDMMaterialValueFloat1Dynamic* OtherValue = CastChecked<UDMMaterialValueFloat1Dynamic>(InOther);
+	OtherValue->SetValue(GetValue());
+}
+
 TSharedPtr<FJsonValue> UDMMaterialValueFloat1Dynamic::JsonSerialize() const
 {
 	return FDMJsonUtils::Serialize(Value);
diff --git a/Engine/Plugins/Experimental/DynamicMaterial/Source/DynamicMaterial/Private/Components/MaterialValuesDynamic/DMMaterialValueFloat2Dynamic.cpp b/Engine/Plugins/Experimental/DynamicMaterial/Source/DynamicMaterial/Private/Components/MaterialValuesDynamic/DMMaterialValueFloat2Dynamic.cpp
--- a/Engine/Plugins/Experimental/DynamicMaterial/Source/DynamicMaterial/Private/Components/MaterialValuesDynamic/DMMaterialValueFloat2Dynamic.cpp
+++ b/Engine/Plugins/Experimental/DynamicMaterial/Source/DynamicMaterial/Private/Components/MaterialValuesDynamic/DMMaterialValueFloat2Dynamic.cpp
@@ -6,6 +6,10 @@
 #include "Components/MaterialValues/DMMaterialValueFloat2.h"
 #include "Materials/MaterialInstanceDynamic.h"
 
+#if WITH_EDITOR
+#include "PropertyHandle.h"
+#endif
+
 UDMMaterialValueFloat2Dynamic::UDMMaterialValueFloat2Dynamic()
 	: UDMMaterialValueDynamic()
 	, Value(FVector2D::ZeroVector)
@@ -33,6 +37,33 @@ void UDMMaterialValueFloat2Dynamic::ApplyDefaultValue()
 	SetValue(GetDefaultValue());
 }
 
+void UDMMaterialValueFloat2Dynamic::ResetToDefault(TSharedPtr<IPropertyHandle> InPropertyHandle)
+{
+	const FStringView PropertyPath = InPropertyHandle->GetPropertyPath();
+	const int32 IndexOfStructSeparator = PropertyPath.Find(TEXT("->"));
+
+	if (IndexOfStructSeparator == INDEX_NONE)
+	{
+		Super::ResetToDefault(InPropertyHandle);
+		return;
+	}
+
+	FStringView LeafName = PropertyPath.RightChop(IndexOfStructSeparator + 2);
+
+	if (LeafName.Equals(TEXT("X")))
+	{
+		FVector2D CurrentValue = GetValue();
+		CurrentValue.X = GetDefaultValue().X;
+		SetValue(CurrentValue);
+	}
+	else if (LeafName.Equals(TEXT("Y")))
+	{
+		FVector2D CurrentValue = GetValue();
+		CurrentValue.Y = GetDefaultValue().Y;
+		SetValue(CurrentValue);
+	}
+}
+
 void UDMMaterialValueFloat2Dynamic::CopyDynamicPropertiesTo(UDMMaterialComponent* InDestinationComponent) const
 {
 	if (UDMMaterialValueFloat2* DestinationValueFloat2 = Cast<UDMMaterialValueFloat2>(InDestinationComponent))
@@ -41,6 +72,12 @@ void UDMMaterialValueFloat2Dynamic::CopyDynamicPropertiesTo(UDMMaterialComponent
 	}
 }
 
+void UDMMaterialValueFloat2Dynamic::CopyParametersFrom_Implementation(UObject* InOther)
+{
+	UDMMaterialValueFloat2Dynamic* OtherValue = CastChecked<UDMMaterialValueFloat2Dynamic>(InOther);
+	OtherValue->SetValue(GetValue());
+}
+
 TSharedPtr<FJsonValue> UDMMaterialValueFloat2Dynamic::JsonSerialize() const
 {
 	return FDMJsonUtils::Serialize(Value);
diff --git a/Engine/Plugins/Experimental/DynamicMaterial/Source/DynamicMaterial/Private/Components/MaterialValuesDynamic/DMMaterialValueFloat3RGBDynamic.cpp b/Engine/Plugins/Experimental/DynamicMaterial/Source/DynamicMaterial/Private/Components/MaterialValuesDynamic/DMMaterialValueFloat3RGBDynamic.cpp
--- a/Engine/Plugins/Experimental/DynamicMaterial/Source/DynamicMaterial/Private/Components/MaterialValuesDynamic/DMMaterialValueFloat3RGBDynamic.cpp
+++ b/Engine/Plugins/Experimental/DynamicMaterial/Source/DynamicMaterial/Private/Components/MaterialValuesDynamic/DMMaterialValueFloat3RGBDynamic.cpp
@@ -6,6 +6,10 @@
 #include "Components/MaterialValues/DMMaterialValueFloat3RGB.h"
 #include "Materials/MaterialInstanceDynamic.h"
 
+#if WITH_EDITOR
+#include "PropertyHandle.h"
+#endif
+
 UDMMaterialValueFloat3RGBDynamic::UDMMaterialValueFloat3RGBDynamic()
 	: UDMMaterialValueDynamic()
 	, Value(FLinearColor::Black)
@@ -33,6 +37,39 @@ void UDMMaterialValueFloat3RGBDynamic::ApplyDefaultValue()
 	SetValue(GetDefaultValue());
 }
 
+void UDMMaterialValueFloat3RGBDynamic::ResetToDefault(TSharedPtr<IPropertyHandle> InPropertyHandle)
+{
+	const FStringView PropertyPath = InPropertyHandle->GetPropertyPath();
+	const int32 IndexOfStructSeparator = PropertyPath.Find(TEXT("->"));
+
+	if (IndexOfStructSeparator == INDEX_NONE)
+	{
+		Super::ResetToDefault(InPropertyHandle);
+		return;
+	}
+
+	FStringView LeafName = PropertyPath.RightChop(IndexOfStructSeparator + 2);
+
+	if (LeafName.Equals(TEXT("R")))
+	{
+		FLinearColor CurrentValue = GetValue();
+		CurrentValue.R = GetDefaultValue().R;
+		SetValue(CurrentValue);
+	}
+	else if (LeafName.Equals(TEXT("G")))
+	{
+		FLinearColor CurrentValue = GetValue();
+		CurrentValue.G = GetDefaultValue().G;
+		SetValue(CurrentValue);
+	}
+	else if (LeafName.Equals(TEXT("B")))
+	{
+		FLinearColor CurrentValue = GetValue();
+		CurrentValue.B = GetDefaultValue().B;
+		SetValue(CurrentValue);
+	}
+}
+
 void UDMMaterialValueFloat3RGBDynamic::CopyDynamicPropertiesTo(UDMMaterialComponent* InDestinationComponent) const
 {
 	if (UDMMaterialValueFloat3RGB* DestinationValueFloat3RGB = Cast<UDMMaterialValueFloat3RGB>(InDestinationComponent))
@@ -41,6 +78,12 @@ void UDMMaterialValueFloat3RGBDynamic::CopyDynamicPropertiesTo(UDMMaterialCompon
 	}
 }
 
+void UDMMaterialValueFloat3RGBDynamic::CopyParametersFrom_Implementation(UObject* InOther)
+{
+	UDMMaterialValueFloat3RGBDynamic* OtherValue = CastChecked<UDMMaterialValueFloat3RGBDynamic>(InOther);
+	OtherValue->SetValue(GetValue());
+}
+
 TSharedPtr<FJsonValue> UDMMaterialValueFloat3RGBDynamic::JsonSerialize() const
 {
 	return FDMJsonUtils::Serialize(Value);
diff --git a/Engine/Plugins/Experimental/DynamicMaterial/Source/DynamicMaterial/Private/Components/MaterialValuesDynamic/DMMaterialValueFloat3RPYDynamic.cpp b/Engine/Plugins/Experimental/DynamicMaterial/Source/DynamicMaterial/Private/Components/MaterialValuesDynamic/DMMaterialValueFloat3RPYDynamic.cpp
--- a/Engine/Plugins/Experimental/DynamicMaterial/Source/DynamicMaterial/Private/Components/MaterialValuesDynamic/DMMaterialValueFloat3RPYDynamic.cpp
+++ b/Engine/Plugins/Experimental/DynamicMaterial/Source/DynamicMaterial/Private/Components/MaterialValuesDynamic/DMMaterialValueFloat3RPYDynamic.cpp
@@ -6,6 +6,10 @@
 #include "Components/MaterialValues/DMMaterialValueFloat3RPY.h"
 #include "Materials/MaterialInstanceDynamic.h"
 
+#if WITH_EDITOR
+#include "PropertyHandle.h"
+#endif
+
 UDMMaterialValueFloat3RPYDynamic::UDMMaterialValueFloat3RPYDynamic()
 	: UDMMaterialValueDynamic()
 	, Value(FRotator::ZeroRotator)
@@ -33,6 +37,39 @@ void UDMMaterialValueFloat3RPYDynamic::ApplyDefaultValue()
 	SetValue(GetDefaultValue());
 }
 
+void UDMMaterialValueFloat3RPYDynamic::ResetToDefault(TSharedPtr<IPropertyHandle> InPropertyHandle)
+{
+	const FStringView PropertyPath = InPropertyHandle->GetPropertyPath();
+	const int32 IndexOfStructSeparator = PropertyPath.Find(TEXT("->"));
+
+	if (IndexOfStructSeparator == INDEX_NONE)
+	{
+		Super::ResetToDefault(InPropertyHandle);
+		return;
+	}
+
+	FStringView LeafName = PropertyPath.RightChop(IndexOfStructSeparator + 2);
+
+	if (LeafName.Equals(TEXT("Roll")))
+	{
+		FRotator CurrentValue = GetValue();
+		CurrentValue.Roll = GetDefaultValue().Roll;
+		SetValue(CurrentValue);
+	}
+	else if (LeafName.Equals(TEXT("Pitch")))
+	{
+		FRotator CurrentValue = GetValue();
+		CurrentValue.Pitch = GetDefaultValue().Pitch;
+		SetValue(CurrentValue);
+	}
+	else if (LeafName.Equals(TEXT("Yaw")))
+	{
+		FRotator CurrentValue = GetValue();
+		CurrentValue.Yaw = GetDefaultValue().Yaw;
+		SetValue(CurrentValue);
+	}
+}
+
 void UDMMaterialValueFloat3RPYDynamic::CopyDynamicPropertiesTo(UDMMaterialComponent* InDestinationComponent) const
 {
 	if (UDMMaterialValueFloat3RPY* DestinationValueFloat3RPY = Cast<UDMMaterialValueFloat3RPY>(InDestinationComponent))
@@ -41,6 +78,12 @@ void UDMMaterialValueFloat3RPYDynamic::CopyDynamicPropertiesTo(UDMMaterialCompon
 	}
 }
 
+void UDMMaterialValueFloat3RPYDynamic::CopyParametersFrom_Implementation(UObject* InOther)
+{
+	UDMMaterialValueFloat3RPYDynamic* OtherValue = CastChecked<UDMMaterialValueFloat3RPYDynamic>(InOther);
+	OtherValue->SetValue(GetValue());
+}
+
 TSharedPtr<FJsonValue> UDMMaterialValueFloat3RPYDynamic::JsonSerialize() const
 {
 	return FDMJsonUtils::Serialize(Value);
diff --git a/Engine/Plugins/Experimental/DynamicMaterial/Source/DynamicMaterial/Private/Components/MaterialValuesDynamic/DMMaterialValueFloat3XYZDynamic.cpp b/Engine/Plugins/Experimental/DynamicMaterial/Source/DynamicMaterial/Private/Components/MaterialValuesDynamic/DMMaterialValueFloat3XYZDynamic.cpp
--- a/Engine/Plugins/Experimental/DynamicMaterial/Source/DynamicMaterial/Private/Components/MaterialValuesDynamic/DMMaterialValueFloat3XYZDynamic.cpp
+++ b/Engine/Plugins/Experimental/DynamicMaterial/Source/DynamicMaterial/Private/Components/MaterialValuesDynamic/DMMaterialValueFloat3XYZDynamic.cpp
@@ -6,6 +6,10 @@
 #include "Components/MaterialValues/DMMaterialValueFloat3XYZ.h"
 #include "Materials/MaterialInstanceDynamic.h"
 
+#if WITH_EDITOR
+#include "PropertyHandle.h"
+#endif
+
 UDMMaterialValueFloat3XYZDynamic::UDMMaterialValueFloat3XYZDynamic()
 	: UDMMaterialValueDynamic()
 	, Value(FVector::ZeroVector)
@@ -33,6 +37,39 @@ void UDMMaterialValueFloat3XYZDynamic::ApplyDefaultValue()
 	SetValue(GetDefaultValue());
 }
 
+void UDMMaterialValueFloat3XYZDynamic::ResetToDefault(TSharedPtr<IPropertyHandle> InPropertyHandle)
+{
+	const FStringView PropertyPath = InPropertyHandle->GetPropertyPath();
+	const int32 IndexOfStructSeparator = PropertyPath.Find(TEXT("->"));
+
+	if (IndexOfStructSeparator == INDEX_NONE)
+	{
+		Super::ResetToDefault(InPropertyHandle);
+		return;
+	}
+
+	FStringView LeafName = PropertyPath.RightChop(IndexOfStructSeparator + 2);
+
+	if (LeafName.Equals(TEXT("X")))
+	{
+		FVector CurrentValue = GetValue();
+		CurrentValue.X = GetDefaultValue().X;
+		SetValue(CurrentValue);
+	}
+	else if (LeafName.Equals(TEXT("Y")))
+	{
+		FVector CurrentValue = GetValue();
+		CurrentValue.Y = GetDefaultValue().Y;
+		SetValue(CurrentValue);
+	}
+	else if (LeafName.Equals(TEXT("Z")))
+	{
+		FVector CurrentValue = GetValue();
+		CurrentValue.Z = GetDefaultValue().Z;
+		SetValue(CurrentValue);
+	}
+}
+
 void UDMMaterialValueFloat3XYZDynamic::CopyDynamicPropertiesTo(UDMMaterialComponent* InDestinationComponent) const
 {
 	if (UDMMaterialValueFloat3XYZ* DestinationValueFloat3XYZ = Cast<UDMMaterialValueFloat3XYZ>(InDestinationComponent))
@@ -41,6 +78,12 @@ void UDMMaterialValueFloat3XYZDynamic::CopyDynamicPropertiesTo(UDMMaterialCompon
 	}
 }
 
+void UDMMaterialValueFloat3XYZDynamic::CopyParametersFrom_Implementation(UObject* InOther)
+{
+	UDMMaterialValueFloat3XYZDynamic* OtherValue = CastChecked<UDMMaterialValueFloat3XYZDynamic>(InOther);
+	OtherValue->SetValue(GetValue());
+}
+
 TSharedPtr<FJsonValue> UDMMaterialValueFloat3XYZDynamic::JsonSerialize() const
 {
 	return FDMJsonUtils::Serialize(Value);
diff --git a/Engine/Plugins/Experimental/DynamicMaterial/Source/DynamicMaterial/Private/Components/MaterialValuesDynamic/DMMaterialValueFloat4Dynamic.cpp b/Engine/Plugins/Experimental/DynamicMaterial/Source/DynamicMaterial/Private/Components/MaterialValuesDynamic/DMMaterialValueFloat4Dynamic.cpp
--- a/Engine/Plugins/Experimental/DynamicMaterial/Source/DynamicMaterial/Private/Components/MaterialValuesDynamic/DMMaterialValueFloat4Dynamic.cpp
+++ b/Engine/Plugins/Experimental/DynamicMaterial/Source/DynamicMaterial/Private/Components/MaterialValuesDynamic/DMMaterialValueFloat4Dynamic.cpp
@@ -6,6 +6,10 @@
 #include "Components/MaterialValues/DMMaterialValueFloat4.h"
 #include "Materials/MaterialInstanceDynamic.h"
 
+#if WITH_EDITOR
+#include "PropertyHandle.h"
+#endif
+
 UDMMaterialValueFloat4Dynamic::UDMMaterialValueFloat4Dynamic()
 	: UDMMaterialValueDynamic()
 	, Value(FLinearColor::Black)
@@ -33,6 +37,45 @@ void UDMMaterialValueFloat4Dynamic::ApplyDefaultValue()
 	SetValue(GetDefaultValue());
 }
 
+void UDMMaterialValueFloat4Dynamic::ResetToDefault(TSharedPtr<IPropertyHandle> InPropertyHandle)
+{
+	const FStringView PropertyPath = InPropertyHandle->GetPropertyPath();
+	const int32 IndexOfStructSeparator = PropertyPath.Find(TEXT("->"));
+
+	if (IndexOfStructSeparator == INDEX_NONE)
+	{
+		Super::ResetToDefault(InPropertyHandle);
+		return;
+	}
+
+	FStringView LeafName = PropertyPath.RightChop(IndexOfStructSeparator + 2);
+
+	if (LeafName.Equals(TEXT("R")))
+	{
+		FLinearColor CurrentValue = GetValue();
+		CurrentValue.R = GetDefaultValue().R;
+		SetValue(CurrentValue);
+	}
+	else if (LeafName.Equals(TEXT("G")))
+	{
+		FLinearColor CurrentValue = GetValue();
+		CurrentValue.G = GetDefaultValue().G;
+		SetValue(CurrentValue);
+	}
+	else if (LeafName.Equals(TEXT("B")))
+	{
+		FLinearColor CurrentValue = GetValue();
+		CurrentValue.B = GetDefaultValue().B;
+		SetValue(CurrentValue);
+	}
+	else if (LeafName.Equals(TEXT("A")))
+	{
+		FLinearColor CurrentValue = GetValue();
+		CurrentValue.A = GetDefaultValue().A;
+		SetValue(CurrentValue);
+	}
+}
+
 void UDMMaterialValueFloat4Dynamic::CopyDynamicPropertiesTo(UDMMaterialComponent* InDestinationComponent) const
 {
 	if (UDMMaterialValueFloat4* DestinationValueFloat4 = Cast<UDMMaterialValueFloat4>(InDestinationComponent))
@@ -41,6 +84,12 @@ void UDMMaterialValueFloat4Dynamic::CopyDynamicPropertiesTo(UDMMaterialComponent
 	}
 }
 
+void UDMMaterialValueFloat4Dynamic::CopyParametersFrom_Implementation(UObject* InOther)
+{
+	UDMMaterialValueFloat4Dynamic* OtherValue = CastChecked<UDMMaterialValueFloat4Dynamic>(InOther);
+	OtherValue->SetValue(GetValue());
+}
+
 TSharedPtr<FJsonValue> UDMMaterialValueFloat4Dynamic::JsonSerialize() const
 {
 	return FDMJsonUtils::Serialize(Value);
diff --git a/Engine/Plugins/Experimental/DynamicMaterial/Source/DynamicMaterial/Private/Components/MaterialValuesDynamic/DMMaterialValueRenderTargetDynamic.cpp b/Engine/Plugins/Experimental/DynamicMaterial/Source/DynamicMaterial/Private/Components/MaterialValuesDynamic/DMMaterialValueRenderTargetDynamic.cpp
--- a/Engine/Plugins/Experimental/DynamicMaterial/Source/DynamicMaterial/Private/Components/MaterialValuesDynamic/DMMaterialValueRenderTargetDynamic.cpp
+++ b/Engine/Plugins/Experimental/DynamicMaterial/Source/DynamicMaterial/Private/Components/MaterialValuesDynamic/DMMaterialValueRenderTargetDynamic.cpp
@@ -31,6 +31,11 @@ void UDMMaterialValueRenderTargetDynamic::CopyDynamicPropertiesTo(UDMMaterialCom
 }
 #endif
 
+void UDMMaterialValueRenderTargetDynamic::CopyParametersFrom_Implementation(UObject* InOther)
+{
+	// Overridden to do nothing
+}
+
 void UDMMaterialValueRenderTargetDynamic::SetMIDParameter(UMaterialInstanceDynamic* InMID) const
 {
 }
diff --git a/Engine/Plugins/Experimental/DynamicMaterial/Source/DynamicMaterial/Private/Components/MaterialValuesDynamic/DMMaterialValueTextureDynamic.cpp b/Engine/Plugins/Experimental/DynamicMaterial/Source/DynamicMaterial/Private/Components/MaterialValuesDynamic/DMMaterialValueTextureDynamic.cpp
--- a/Engine/Plugins/Experimental/DynamicMaterial/Source/DynamicMaterial/Private/Components/MaterialValuesDynamic/DMMaterialValueTextureDynamic.cpp
+++ b/Engine/Plugins/Experimental/DynamicMaterial/Source/DynamicMaterial/Private/Components/MaterialValuesDynamic/DMMaterialValueTextureDynamic.cpp
@@ -45,6 +45,12 @@ void UDMMaterialValueTextureDynamic::CopyDynamicPropertiesTo(UDMMaterialComponen
 	}
 }
 
+void UDMMaterialValueTextureDynamic::CopyParametersFrom_Implementation(UObject* InOther)
+{
+	UDMMaterialValueTextureDynamic* OtherValue = CastChecked<UDMMaterialValueTextureDynamic>(InOther);
+	OtherValue->SetValue(GetValue());
+}
+
 TSharedPtr<FJsonValue> UDMMaterialValueTextureDynamic::JsonSerialize() const
 {
 	return FDMJsonUtils::Serialize(Value);
diff --git a/Engine/Plugins/Experimental/DynamicMaterial/Source/DynamicMaterial/Private/Material/DynamicMaterialInstance.cpp b/Engine/Plugins/Experimental/DynamicMaterial/Source/DynamicMaterial/Private/Material/DynamicMaterialInstance.cpp
--- a/Engine/Plugins/Experimental/DynamicMaterial/Source/DynamicMaterial/Private/Material/DynamicMaterialInstance.cpp
+++ b/Engine/Plugins/Experimental/DynamicMaterial/Source/DynamicMaterial/Private/Material/DynamicMaterialInstance.cpp
@@ -113,7 +113,7 @@ void UDynamicMaterialInstance::SetMaterialModel(UDynamicMaterialModelBase* InMat
 
 	if (InMaterialModel)
 	{
-		InMaterialModel->Rename(nullptr, this, UE::DynamicMaterial::RenameFlags);
+		InMaterialModel->Rename(*InMaterialModel->GetName(), this, UE::DynamicMaterial::RenameFlags);
 	}
 }
 
diff --git a/Engine/Plugins/Experimental/DynamicMaterial/Source/DynamicMaterial/Public/Components/DMMaterialComponent.h b/Engine/Plugins/Experimental/DynamicMaterial/Source/DynamicMaterial/Public/Components/DMMaterialComponent.h
--- a/Engine/Plugins/Experimental/DynamicMaterial/Source/DynamicMaterial/Public/Components/DMMaterialComponent.h
+++ b/Engine/Plugins/Experimental/DynamicMaterial/Source/DynamicMaterial/Public/Components/DMMaterialComponent.h
@@ -161,10 +161,6 @@ public:
 	/** Event called when this component's state changes to Removed. */
 	FOnLifetimeStateChanged::RegistrationType& GetOnRemoved() { return OnRemoved; }
 
-	//~ Begin FNotifyHook
-	virtual void NotifyPreChange(FEditPropertyChain* InPropertyAboutToChange) {}
-	//~ End FNotifyHook
-
 	//~ Begin UObject
 	DYNAMICMATERIAL_API virtual bool Modify(bool bInAlwaysMarkDirty = true) override;
 	DYNAMICMATERIAL_API virtual void PostLoad() override;
diff --git a/Engine/Plugins/Experimental/DynamicMaterial/Source/DynamicMaterial/Public/Components/DMMaterialValue.h b/Engine/Plugins/Experimental/DynamicMaterial/Source/DynamicMaterial/Public/Components/DMMaterialValue.h
--- a/Engine/Plugins/Experimental/DynamicMaterial/Source/DynamicMaterial/Public/Components/DMMaterialValue.h
+++ b/Engine/Plugins/Experimental/DynamicMaterial/Source/DynamicMaterial/Public/Components/DMMaterialValue.h
@@ -3,10 +3,13 @@
 #pragma once
  
 #include "Components/DMMaterialLinkedComponent.h"
+#include "IDMParameterContainer.h"
+#include "Utils/DMJsonUtils.h"
+
 #include "DMDefs.h"
 #include "Templates/SharedPointer.h"
 #include "UObject/StrongObjectPtr.h"
-#include "Utils/DMJsonUtils.h"
+
 #include "DMMaterialValue.generated.h"
 
 class IPropertyHandle;
@@ -27,7 +30,7 @@ struct IDMMaterialBuildStateInterface;
  * A value used in a material. Manages its own parameter.
  */
 UCLASS(MinimalAPI, BlueprintType, Abstract, ClassGroup = "Material Designer", meta = (DisplayName = "Material Designer Value"))
-class UDMMaterialValue : public UDMMaterialLinkedComponent, public IDMJsonSerializable
+class UDMMaterialValue : public UDMMaterialLinkedComponent, public IDMJsonSerializable, public IDMParameterContainer
 {
 	GENERATED_BODY()
 
@@ -75,8 +78,8 @@ public:
 	UFUNCTION(BlueprintCallable, Category = "Material Designer")
 	virtual void ApplyDefaultValue() PURE_VIRTUAL(UDMMaterialValue::ApplyDefaultValue)
 
-	void ResetToDefault(TSharedPtr<IPropertyHandle> InPropertyHandle) { ApplyDefaultValue(); }
- 
+	DYNAMICMATERIAL_API virtual void ResetToDefault(TSharedPtr<IPropertyHandle> InPropertyHandle);
+
 	/**
 	 * Resets to the default default value. 0, nullptr, etc.
 	 */
diff --git a/Engine/Plugins/Experimental/DynamicMaterial/Source/DynamicMaterial/Public/Components/DMMaterialValueDynamic.h b/Engine/Plugins/Experimental/DynamicMaterial/Source/DynamicMaterial/Public/Components/DMMaterialValueDynamic.h
--- a/Engine/Plugins/Experimental/DynamicMaterial/Source/DynamicMaterial/Public/Components/DMMaterialValueDynamic.h
+++ b/Engine/Plugins/Experimental/DynamicMaterial/Source/DynamicMaterial/Public/Components/DMMaterialValueDynamic.h
@@ -3,9 +3,11 @@
 #pragma once
 
 #include "Components/DMMaterialComponentDynamic.h"
-#include "Templates/SharedPointerFwd.h"
+#include "IDMParameterContainer.h"
 #include "Utils/DMJsonUtils.h"
 
+#include "Templates/SharedPointerFwd.h"
+
 #include "DMMaterialValueDynamic.generated.h"
 
 class IPropertyHandle;
@@ -17,7 +19,7 @@ class UMaterialInstanceDynamic;
  * A value used inside an instanced material instance. Links to the original value in the parent material.
  */
 UCLASS(MinimalAPI, BlueprintType, Abstract, ClassGroup = "Material Designer", meta = (DisplayName = "Material Designer Value Instance"))
-class UDMMaterialValueDynamic : public UDMMaterialComponentDynamic, public IDMJsonSerializable
+class UDMMaterialValueDynamic : public UDMMaterialComponentDynamic, public IDMJsonSerializable, public IDMParameterContainer
 {
 	GENERATED_BODY()
 
@@ -48,7 +50,7 @@ public:
 	UFUNCTION(BlueprintCallable, Category = "Material Designer")
 	virtual void ApplyDefaultValue() PURE_VIRTUAL(UDMMaterialValue::ApplyDefaultValue)
 
-	void ResetToDefault(TSharedPtr<IPropertyHandle> InPropertyHandle) { ApplyDefaultValue(); }
+	DYNAMICMATERIAL_API virtual void ResetToDefault(TSharedPtr<IPropertyHandle> InPropertyHandle);
 #endif
 
 	virtual void SetMIDParameter(UMaterialInstanceDynamic* InMID) const PURE_VIRTUAL(UDMMaterialValueDynamic::SetMIDParameter)
diff --git a/Engine/Plugins/Experimental/DynamicMaterial/Source/DynamicMaterial/Public/Components/DMTextureUV.h b/Engine/Plugins/Experimental/DynamicMaterial/Source/DynamicMaterial/Public/Components/DMTextureUV.h
--- a/Engine/Plugins/Experimental/DynamicMaterial/Source/DynamicMaterial/Public/Components/DMTextureUV.h
+++ b/Engine/Plugins/Experimental/DynamicMaterial/Source/DynamicMaterial/Public/Components/DMTextureUV.h
@@ -3,6 +3,8 @@
 #pragma once
 
 #include "Components/DMMaterialLinkedComponent.h"
+#include "IDMParameterContainer.h"
+
 #include "DMDefs.h"
 #include "Misc/Guid.h"
 #include "Templates/SharedPointer.h"
@@ -49,7 +51,7 @@ namespace UE::DynamicMaterial::ParamID
  * Represents a Texture UV material function with the following parameters: offset, tiling, pivot and rotation.
  */
 UCLASS(MinimalAPI, BlueprintType, ClassGroup = "Material Designer", meta = (DisplayName = "Material Designer Texture UV"))
-class UDMTextureUV : public UDMMaterialLinkedComponent
+class UDMTextureUV : public UDMMaterialLinkedComponent, public IDMParameterContainer
 {
 	GENERATED_BODY()
 
@@ -250,6 +252,10 @@ protected:
 	void RemoveParameterNames();
 #endif
 
+	//~ Begin IDMParameterContainer
+	virtual void CopyParametersFrom_Implementation(UObject* InOther) override;
+	//~ End IDMParameterContainer
+
 	void OnTextureUVChanged(EDMUpdateType InUpdateType);
 
 #if WITH_EDITOR
diff --git a/Engine/Plugins/Experimental/DynamicMaterial/Source/DynamicMaterial/Public/Components/DMTextureUVDynamic.h b/Engine/Plugins/Experimental/DynamicMaterial/Source/DynamicMaterial/Public/Components/DMTextureUVDynamic.h
--- a/Engine/Plugins/Experimental/DynamicMaterial/Source/DynamicMaterial/Public/Components/DMTextureUVDynamic.h
+++ b/Engine/Plugins/Experimental/DynamicMaterial/Source/DynamicMaterial/Public/Components/DMTextureUVDynamic.h
@@ -3,6 +3,8 @@
 #pragma once
 
 #include "Components/DMMaterialComponentDynamic.h"
+#include "IDMParameterContainer.h"
+
 #include "Templates/SharedPointerFwd.h"
 
 #include "DMTextureUVDynamic.generated.h"
@@ -16,7 +18,7 @@ class UMaterialInstanceDynamic;
  * A texture uv used inside a instanced material instance. Links to the original texture uv in the parent material.
  */
 UCLASS(MinimalAPI, BlueprintType, ClassGroup = "Material Designer", meta = (DisplayName = "Material Designer Texture UV Instance"))
-class UDMTextureUVDynamic : public UDMMaterialComponentDynamic
+class UDMTextureUVDynamic : public UDMMaterialComponentDynamic, public IDMParameterContainer
 {
 	GENERATED_BODY()
 
@@ -105,6 +107,10 @@ protected:
 	/** Called when any value in this texture uv changes. */
 	DYNAMICMATERIAL_API virtual void OnTextureUVChanged();
 
+	//~ Begin IDMParameterContainer
+	virtual void CopyParametersFrom_Implementation(UObject* InOther) override;
+	//~ End IDMParameterContainer
+
 #if WITH_EDITOR
 	//~ Begin UDMMaterialComponent
 	DYNAMICMATERIAL_API virtual void OnComponentAdded() override;
diff --git a/Engine/Plugins/Experimental/DynamicMaterial/Source/DynamicMaterial/Public/Components/MaterialValues/DMMaterialValueBool.h b/Engine/Plugins/Experimental/DynamicMaterial/Source/DynamicMaterial/Public/Components/MaterialValues/DMMaterialValueBool.h
--- a/Engine/Plugins/Experimental/DynamicMaterial/Source/DynamicMaterial/Public/Components/MaterialValues/DMMaterialValueBool.h
+++ b/Engine/Plugins/Experimental/DynamicMaterial/Source/DynamicMaterial/Public/Components/MaterialValues/DMMaterialValueBool.h
@@ -65,4 +65,8 @@ protected:
 		meta = (AllowPrivateAccess = "true"))
 	bool bDefaultValue;
 #endif
+
+	//~ Begin IDMParameterContainer
+	virtual void CopyParametersFrom_Implementation(UObject* InOther) override;
+	//~ End IDMParameterContainer
 };
diff --git a/Engine/Plugins/Experimental/DynamicMaterial/Source/DynamicMaterial/Public/Components/MaterialValues/DMMaterialValueColorAtlas.h b/Engine/Plugins/Experimental/DynamicMaterial/Source/DynamicMaterial/Public/Components/MaterialValues/DMMaterialValueColorAtlas.h
--- a/Engine/Plugins/Experimental/DynamicMaterial/Source/DynamicMaterial/Public/Components/MaterialValues/DMMaterialValueColorAtlas.h
+++ b/Engine/Plugins/Experimental/DynamicMaterial/Source/DynamicMaterial/Public/Components/MaterialValues/DMMaterialValueColorAtlas.h
@@ -92,4 +92,8 @@ protected:
 		meta = (AllowPrivateAccess = "true", NotKeyframeable, NoCreate))
 	TObjectPtr<UCurveLinearColor> Curve;
 #endif
+
+	//~ Begin IDMParameterContainer
+	virtual void CopyParametersFrom_Implementation(UObject* InOther) override;
+	//~ End IDMParameterContainer
 };
diff --git a/Engine/Plugins/Experimental/DynamicMaterial/Source/DynamicMaterial/Public/Components/MaterialValues/DMMaterialValueFloat2.h b/Engine/Plugins/Experimental/DynamicMaterial/Source/DynamicMaterial/Public/Components/MaterialValues/DMMaterialValueFloat2.h
--- a/Engine/Plugins/Experimental/DynamicMaterial/Source/DynamicMaterial/Public/Components/MaterialValues/DMMaterialValueFloat2.h
+++ b/Engine/Plugins/Experimental/DynamicMaterial/Source/DynamicMaterial/Public/Components/MaterialValues/DMMaterialValueFloat2.h
@@ -44,6 +44,7 @@ public:
 	DYNAMICMATERIAL_API virtual int32 GetInnateMaskOutput(int32 OutputChannels) const override;
 	DYNAMICMATERIAL_API virtual bool IsDefaultValue() const override;
 	DYNAMICMATERIAL_API virtual void ApplyDefaultValue() override;
+	DYNAMICMATERIAL_API virtual void ResetToDefault(TSharedPtr<IPropertyHandle> InPropertyHandle) override;
 	DYNAMICMATERIAL_API virtual void ResetDefaultValue() override;
 	DYNAMICMATERIAL_API virtual UDMMaterialValueDynamic* ToDynamic(UDynamicMaterialModelDynamic* InMaterialModelDynamic) override;
 #endif
@@ -66,4 +67,8 @@ protected:
 		meta = (AllowPrivateAccess = "true"))
 	FVector2D DefaultValue;
 #endif
+
+	//~ Begin IDMParameterContainer
+	virtual void CopyParametersFrom_Implementation(UObject* InOther) override;
+	//~ End IDMParameterContainer
 };
diff --git a/Engine/Plugins/Experimental/DynamicMaterial/Source/DynamicMaterial/Public/Components/MaterialValues/DMMaterialValueFloat3RGB.h b/Engine/Plugins/Experimental/DynamicMaterial/Source/DynamicMaterial/Public/Components/MaterialValues/DMMaterialValueFloat3RGB.h
--- a/Engine/Plugins/Experimental/DynamicMaterial/Source/DynamicMaterial/Public/Components/MaterialValues/DMMaterialValueFloat3RGB.h
+++ b/Engine/Plugins/Experimental/DynamicMaterial/Source/DynamicMaterial/Public/Components/MaterialValues/DMMaterialValueFloat3RGB.h
@@ -44,6 +44,7 @@ public:
 	DYNAMICMATERIAL_API virtual int32 GetInnateMaskOutput(int32 OutputChannels) const override;
 	DYNAMICMATERIAL_API virtual bool IsDefaultValue() const override;
 	DYNAMICMATERIAL_API virtual void ApplyDefaultValue() override;
+	DYNAMICMATERIAL_API virtual void ResetToDefault(TSharedPtr<IPropertyHandle> InPropertyHandle) override;
 	DYNAMICMATERIAL_API virtual void ResetDefaultValue() override;
 	DYNAMICMATERIAL_API virtual UDMMaterialValueDynamic* ToDynamic(UDynamicMaterialModelDynamic* InMaterialModelDynamic) override;
 	//~ End UDMMaterialValue
@@ -66,4 +67,8 @@ protected:
 		meta = (AllowPrivateAccess = "true", HideAlphaChannel))
 	FLinearColor DefaultValue;
 #endif
+
+	//~ Begin IDMParameterContainer
+	virtual void CopyParametersFrom_Implementation(UObject* InOther) override;
+	//~ End IDMParameterContainer
 };
diff --git a/Engine/Plugins/Experimental/DynamicMaterial/Source/DynamicMaterial/Public/Components/MaterialValues/DMMaterialValueRenderTarget.h b/Engine/Plugins/Experimental/DynamicMaterial/Source/DynamicMaterial/Public/Components/MaterialValues/DMMaterialValueRenderTarget.h
--- a/Engine/Plugins/Experimental/DynamicMaterial/Source/DynamicMaterial/Public/Components/MaterialValues/DMMaterialValueRenderTarget.h
+++ b/Engine/Plugins/Experimental/DynamicMaterial/Source/DynamicMaterial/Public/Components/MaterialValues/DMMaterialValueRenderTarget.h
@@ -117,6 +117,10 @@ protected:
 
 	void UpdateRenderTarget();
 
+	//~ Begin IDMParameterContainer
+	virtual void CopyParametersFrom_Implementation(UObject* InOther) override;
+	//~ End IDMParameterContainer
+
 	//~ Begin UDMMaterialComponent
 	DYNAMICMATERIAL_API virtual UDMMaterialComponent* GetSubComponentByPath(FDMComponentPath& InPath, const FDMComponentPathSegment& InPathSegment) const override;
 #if WITH_EDITOR
diff --git a/Engine/Plugins/Experimental/DynamicMaterial/Source/DynamicMaterial/Public/Components/MaterialValuesDynamic/DMMaterialValueFloat2Dynamic.h b/Engine/Plugins/Experimental/DynamicMaterial/Source/DynamicMaterial/Public/Components/MaterialValuesDynamic/DMMaterialValueFloat2Dynamic.h
--- a/Engine/Plugins/Experimental/DynamicMaterial/Source/DynamicMaterial/Public/Components/MaterialValuesDynamic/DMMaterialValueFloat2Dynamic.h
+++ b/Engine/Plugins/Experimental/DynamicMaterial/Source/DynamicMaterial/Public/Components/MaterialValuesDynamic/DMMaterialValueFloat2Dynamic.h
@@ -38,6 +38,7 @@ public:
 #if WITH_EDITOR
 	DYNAMICMATERIAL_API virtual bool IsDefaultValue() const override;
 	DYNAMICMATERIAL_API virtual void ApplyDefaultValue() override;
+	DYNAMICMATERIAL_API virtual void ResetToDefault(TSharedPtr<IPropertyHandle> InPropertyHandle);
 	DYNAMICMATERIAL_API virtual void CopyDynamicPropertiesTo(UDMMaterialComponent* InDestinationComponent) const override;
 #endif
 	// ~End UDMMaterialValueDynamic
@@ -46,4 +47,8 @@ protected:
 	UPROPERTY(EditInstanceOnly, BlueprintReadWrite, Getter, Setter, BlueprintSetter = SetValue, Category = "Material Designer",
 		meta = (AllowPrivateAccess = "true", DisplayName = "Vector2D", AllowPreserveRatio, Delta = 0.001))
 	FVector2D Value;
+
+	//~ Begin IDMParameterContainer
+	virtual void CopyParametersFrom_Implementation(UObject* InOther) override;
+	//~ End IDMParameterContainer
 };
diff --git a/Engine/Plugins/Experimental/DynamicMaterial/Source/DynamicMaterial/Public/Components/MaterialValuesDynamic/DMMaterialValueFloat3RGBDynamic.h b/Engine/Plugins/Experimental/DynamicMaterial/Source/DynamicMaterial/Public/Components/MaterialValuesDynamic/DMMaterialValueFloat3RGBDynamic.h
--- a/Engine/Plugins/Experimental/DynamicMaterial/Source/DynamicMaterial/Public/Components/MaterialValuesDynamic/DMMaterialValueFloat3RGBDynamic.h
+++ b/Engine/Plugins/Experimental/DynamicMaterial/Source/DynamicMaterial/Public/Components/MaterialValuesDynamic/DMMaterialValueFloat3RGBDynamic.h
@@ -38,6 +38,7 @@ public:
 #if WITH_EDITOR
 	DYNAMICMATERIAL_API virtual bool IsDefaultValue() const override;
 	DYNAMICMATERIAL_API virtual void ApplyDefaultValue() override;
+	DYNAMICMATERIAL_API virtual void ResetToDefault(TSharedPtr<IPropertyHandle> InPropertyHandle);
 	DYNAMICMATERIAL_API virtual void CopyDynamicPropertiesTo(UDMMaterialComponent* InDestinationComponent) const override;
 #endif
 	// ~End UDMMaterialValueDynamic
@@ -46,4 +47,8 @@ protected:
 	UPROPERTY(EditInstanceOnly, BlueprintReadWrite, Getter, Setter, BlueprintSetter = SetValue, Category = "Material Designer",
 		meta = (AllowPrivateAccess = "true", DisplayName = "Color", HideAlphaChannel))
 	FLinearColor Value;
+
+	//~ Begin IDMParameterContainer
+	virtual void CopyParametersFrom_Implementation(UObject* InOther) override;
+	//~ End IDMParameterContainer
 };
diff --git a/Engine/Plugins/Experimental/DynamicMaterial/Source/DynamicMaterial/Public/Components/MaterialValuesDynamic/DMMaterialValueRenderTargetDynamic.h b/Engine/Plugins/Experimental/DynamicMaterial/Source/DynamicMaterial/Public/Components/MaterialValuesDynamic/DMMaterialValueRenderTargetDynamic.h
--- a/Engine/Plugins/Experimental/DynamicMaterial/Source/DynamicMaterial/Public/Components/MaterialValuesDynamic/DMMaterialValueRenderTargetDynamic.h
+++ b/Engine/Plugins/Experimental/DynamicMaterial/Source/DynamicMaterial/Public/Components/MaterialValuesDynamic/DMMaterialValueRenderTargetDynamic.h
@@ -29,4 +29,9 @@ public:
 	DYNAMICMATERIAL_API virtual void CopyDynamicPropertiesTo(UDMMaterialComponent* InDestinationComponent) const override;
 #endif
 	// ~End UDMMaterialValueDynamic
+
+protected:
+	//~ Begin IDMParameterContainer
+	virtual void CopyParametersFrom_Implementation(UObject* InOther) override;
+	//~ End IDMParameterContainer
 };
diff --git a/Engine/Plugins/Experimental/DynamicMaterial/Source/DynamicMaterial/Public/Model/DynamicMaterialModelBase.h b/Engine/Plugins/Experimental/DynamicMaterial/Source/DynamicMaterial/Public/Model/DynamicMaterialModelBase.h
--- a/Engine/Plugins/Experimental/DynamicMaterial/Source/DynamicMaterial/Public/Model/DynamicMaterialModelBase.h
+++ b/Engine/Plugins/Experimental/DynamicMaterial/Source/DynamicMaterial/Public/Model/DynamicMaterialModelBase.h
@@ -41,4 +41,18 @@ public:
 
 	/** Apply all registered components to the given MID. */
 	virtual void ApplyComponents(UMaterialInstanceDynamic* InMID) PURE_VIRTUAL(UDynamicMaterialModelBase::ApplyComponents)
+
+#if WITH_EDITOR
+	bool IsPreviewModified() const { return bPreviewModified; }
+
+	void MarkPreviewModified() { bPreviewModified = true; }
+
+	void MarkOriginalUpdated() { bPreviewModified = false; }
+#endif
+
+protected:
+#if WITH_EDITORONLY_DATA
+	/** Set to true when the material designer makes a change to a preview model. */
+	bool bPreviewModified = false;
+#endif
 };
diff --git a/Engine/Plugins/Experimental/DynamicMaterial/Source/DynamicMaterialEditor/Private/Components/DMMaterialStage.cpp b/Engine/Plugins/Experimental/DynamicMaterial/Source/DynamicMaterialEditor/Private/Components/DMMaterialStage.cpp
--- a/Engine/Plugins/Experimental/DynamicMaterial/Source/DynamicMaterialEditor/Private/Components/DMMaterialStage.cpp
+++ b/Engine/Plugins/Experimental/DynamicMaterial/Source/DynamicMaterialEditor/Private/Components/DMMaterialStage.cpp
@@ -409,7 +409,11 @@ void UDMMaterialStage::InputUpdated(UDMMaterialStageInput* InInput, EDMUpdateTyp
 	}
 
 	int32 InputIdx = Inputs.Find(InInput);
-	check(InputIdx != INDEX_NONE);
+
+	if (InputIdx == INDEX_NONE)
+	{
+		return;
+	}
 
 	InputIdx += FDMMaterialStageConnectorChannel::FIRST_STAGE_INPUT;
 
diff --git a/Engine/Plugins/Experimental/DynamicMaterial/Source/DynamicMaterialEditor/Private/DynamicMaterialEditorModule.cpp b/Engine/Plugins/Experimental/DynamicMaterial/Source/DynamicMaterialEditor/Private/DynamicMaterialEditorModule.cpp
--- a/Engine/Plugins/Experimental/DynamicMaterial/Source/DynamicMaterialEditor/Private/DynamicMaterialEditorModule.cpp
+++ b/Engine/Plugins/Experimental/DynamicMaterial/Source/DynamicMaterialEditor/Private/DynamicMaterialEditorModule.cpp
@@ -114,33 +114,27 @@ FDMComponentPropertyRowGeneratorDelegate FDynamicMaterialEditorModule::GetCompon
 	if (!FoundClass)
 	{
 		// Return an invalid lambda
-		static FDMComponentPropertyRowGeneratorDelegate NullLambda = FDMComponentPropertyRowGeneratorDelegate::CreateLambda(
-			[](const TSharedRef<SDMMaterialComponentEditor>& InComponentEditorWidget, UDMMaterialComponent* InComponent, TArray<FDMPropertyHandle>& InOutPropertyRows,
-				TSet<UDMMaterialComponent*>& InOutProcessedObjects)
-			{
-			});
-
+		static FDMComponentPropertyRowGeneratorDelegate NullLambda;
 		return NullLambda;
 	}
 
 	return ComponentPropertyRowGenerators[FoundClass];
 }
 
-void FDynamicMaterialEditorModule::GeneratorComponentPropertyRows(const TSharedRef<SDMMaterialComponentEditor>& InComponentEditorWidget, UDMMaterialComponent* InComponent,
-	TArray<FDMPropertyHandle>& InOutPropertyRows, TSet<UDMMaterialComponent*>& InOutProcessedObjects)
+void FDynamicMaterialEditorModule::GeneratorComponentPropertyRows(FDMComponentPropertyRowGeneratorParams& InParams)
 {
-	if (!IsValid(InComponent))
+	if (!IsValid(InParams.Object))
 	{
 		return;
 	}
 
-	if (InOutProcessedObjects.Contains(InComponent))
+	if (InParams.ProcessedObjects.Contains(InParams.Object))
 	{
 		return;
 	}
 
-	FDMComponentPropertyRowGeneratorDelegate RowGenerator = GetComponentPropertyRowGeneratorDelegate(InComponent->GetClass());
-	RowGenerator.ExecuteIfBound(InComponentEditorWidget, InComponent, InOutPropertyRows, InOutProcessedObjects);
+	FDMComponentPropertyRowGeneratorDelegate RowGenerator = GetComponentPropertyRowGeneratorDelegate(InParams.Object->GetClass());
+	RowGenerator.ExecuteIfBound(InParams);
 }
 
 void FDynamicMaterialEditorModule::RegisterCustomMaterialPropertyGenerator(UClass* InClass, FDMGetObjectMaterialPropertiesDelegate InGenerator)
@@ -410,7 +404,7 @@ TSharedRef<SWidget> FDynamicMaterialEditorModule::CreateEditor(UDynamicMaterialM
 
 		if (IsValid(WorldSubsystem))
 		{
-			WorldSubsystem->GetGetCustomEditorModelDelegate().BindSP(NewDesigner, &SDMMaterialDesigner::GetMaterialModelBase);
+			WorldSubsystem->GetGetCustomEditorModelDelegate().BindSP(NewDesigner, &SDMMaterialDesigner::GetOriginalMaterialModelBase);
 			WorldSubsystem->GetSetCustomEditorActorDelegate().BindSP(NewDesigner, &SDMMaterialDesigner::OnActorSelected);
 
 			TWeakPtr<SDMMaterialDesigner> NewDesignerWeak = NewDesigner;
@@ -450,7 +444,7 @@ UDynamicMaterialModelBase* FDynamicMaterialEditorModule::GetOpenedMaterialModel(
 {
 	if (TSharedPtr<SDMMaterialDesigner> Designer = FDMLevelEditorIntegration::GetMaterialDesignerForWorld(InWorld))
 	{
-		return Designer->GetMaterialModelBase();
+		return Designer->GetOriginalMaterialModelBase();
 	}
 
 	if (IsValid(InWorld))
diff --git a/Engine/Plugins/Experimental/DynamicMaterial/Source/DynamicMaterialEditor/Private/DynamicMaterialEditorModule.h b/Engine/Plugins/Experimental/DynamicMaterial/Source/DynamicMaterialEditor/Private/DynamicMaterialEditorModule.h
--- a/Engine/Plugins/Experimental/DynamicMaterial/Source/DynamicMaterialEditor/Private/DynamicMaterialEditorModule.h
+++ b/Engine/Plugins/Experimental/DynamicMaterial/Source/DynamicMaterialEditor/Private/DynamicMaterialEditorModule.h
@@ -20,6 +20,7 @@ class UDMMaterialComponent;
 class UDMMaterialStageSource;
 class UDynamicMaterialInstance;
 class UDynamicMaterialModelBase;
+struct FDMPropertyHandleGenerateParams;
 
 DECLARE_LOG_CATEGORY_EXTERN(LogDynamicMaterialEditor, Log, All);
 
@@ -47,8 +48,7 @@ public:
 	static FDynamicMaterialEditorModule& Get();
 
 	static FDMComponentPropertyRowGeneratorDelegate GetComponentPropertyRowGeneratorDelegate(UClass* InClass);
-	static void GeneratorComponentPropertyRows(const TSharedRef<SDMMaterialComponentEditor>& InComponentEditorWidget, UDMMaterialComponent* InComponent, 
-		TArray<FDMPropertyHandle>& InOutPropertyRows, TSet<UDMMaterialComponent*>& InOutProcessedObjects);
+	static void GeneratorComponentPropertyRows(FDMComponentPropertyRowGeneratorParams& InParams);
 
 	static FDMGetObjectMaterialPropertiesDelegate GetCustomMaterialPropertyGenerator(UClass* InClass);
 
diff --git a/Engine/Plugins/Experimental/DynamicMaterial/Source/DynamicMaterialEditor/Private/DynamicMaterialEditorSettings.cpp b/Engine/Plugins/Experimental/DynamicMaterial/Source/DynamicMaterialEditor/Private/DynamicMaterialEditorSettings.cpp
--- a/Engine/Plugins/Experimental/DynamicMaterial/Source/DynamicMaterialEditor/Private/DynamicMaterialEditorSettings.cpp
+++ b/Engine/Plugins/Experimental/DynamicMaterial/Source/DynamicMaterialEditor/Private/DynamicMaterialEditorSettings.cpp
@@ -134,6 +134,7 @@ UDynamicMaterialEditorSettings::UDynamicMaterialEditorSettings()
 
 	bFollowSelection = true;
 	bAddDetailsPanelButton = false;
+	bAutomaticallyCopyParametersToSourceMaterial = true;
 	bUseLinearColorForVectors = true;
 
 	ResetAllLayoutSettings();
diff --git a/Engine/Plugins/Experimental/DynamicMaterial/Source/DynamicMaterialEditor/Private/DynamicMaterialEditorSettings.h b/Engine/Plugins/Experimental/DynamicMaterial/Source/DynamicMaterialEditor/Private/DynamicMaterialEditorSettings.h
--- a/Engine/Plugins/Experimental/DynamicMaterial/Source/DynamicMaterialEditor/Private/DynamicMaterialEditorSettings.h
+++ b/Engine/Plugins/Experimental/DynamicMaterial/Source/DynamicMaterialEditor/Private/DynamicMaterialEditorSettings.h
@@ -192,6 +192,13 @@ public:
 	UPROPERTY(Config, EditAnywhere, BlueprintReadWrite, Category = "Editor")
 	bool bAddDetailsPanelButton;
 
+	/**
+	 * When properties on the preview material are changed, if they are parameter-based then automatically
+	 * also set those values on the original material (in the world or the asset).
+	 */
+	UPROPERTY(Config, EditAnywhere, BlueprintReadWrite, Category = "Editor")
+	bool bAutomaticallyCopyParametersToSourceMaterial;
+
 	UPROPERTY(Config, EditAnywhere, BlueprintReadWrite, Category = "Layout")
 	EDMMaterialEditorLayout Layout;
 
diff --git a/Engine/Plugins/Experimental/DynamicMaterial/Source/DynamicMaterialEditor/Private/LevelEditor/DMLevelEditorIntegration.cpp b/Engine/Plugins/Experimental/DynamicMaterial/Source/DynamicMaterialEditor/Private/LevelEditor/DMLevelEditorIntegration.cpp
--- a/Engine/Plugins/Experimental/DynamicMaterial/Source/DynamicMaterialEditor/Private/LevelEditor/DMLevelEditorIntegration.cpp
+++ b/Engine/Plugins/Experimental/DynamicMaterial/Source/DynamicMaterialEditor/Private/LevelEditor/DMLevelEditorIntegration.cpp
@@ -159,7 +159,7 @@ void FDMLevelEditorIntegration::OnMapTearDown(UWorld* InWorld)
 		return;
 	}
 
-	UDynamicMaterialModelBase* MaterialModelBase = Designer->GetMaterialModelBase();
+	UDynamicMaterialModelBase* MaterialModelBase = Designer->GetOriginalMaterialModelBase();
 	Designer->Empty();
 
 	if (MaterialModelBase)
diff --git a/Engine/Plugins/Experimental/DynamicMaterial/Source/DynamicMaterialEditor/Private/Model/DynamicMaterialModelEditorOnlyData.cpp b/Engine/Plugins/Experimental/DynamicMaterial/Source/DynamicMaterialEditor/Private/Model/DynamicMaterialModelEditorOnlyData.cpp
--- a/Engine/Plugins/Experimental/DynamicMaterial/Source/DynamicMaterialEditor/Private/Model/DynamicMaterialModelEditorOnlyData.cpp
+++ b/Engine/Plugins/Experimental/DynamicMaterial/Source/DynamicMaterialEditor/Private/Model/DynamicMaterialModelEditorOnlyData.cpp
@@ -592,6 +592,14 @@ void UDynamicMaterialModelEditorOnlyData::RequestMaterialBuild(bool bInImmediate
 {
 	if (!HasAnyFlags(RF_ClassDefaultObject | RF_ArchetypeObject))
 	{
+		if (IsIn(GetTransientPackage()))
+		{
+			if (MaterialModel)
+			{
+				MaterialModel->MarkPreviewModified();
+			}
+		}
+
 		if (bInImmediate)
 		{
 			BuildMaterial(/* Dirty Assets */ false);
diff --git a/Engine/Plugins/Experimental/DynamicMaterial/Source/DynamicMaterialEditor/Private/UI/Menus/DMMaterialSlotLayerMenus.cpp b/Engine/Plugins/Experimental/DynamicMaterial/Source/DynamicMaterialEditor/Private/UI/Menus/DMMaterialSlotLayerMenus.cpp
--- a/Engine/Plugins/Experimental/DynamicMaterial/Source/DynamicMaterialEditor/Private/UI/Menus/DMMaterialSlotLayerMenus.cpp
+++ b/Engine/Plugins/Experimental/DynamicMaterial/Source/DynamicMaterialEditor/Private/UI/Menus/DMMaterialSlotLayerMenus.cpp
@@ -340,7 +340,7 @@ void FDMMaterialSlotLayerMenus::AddGlobalValueSection(FToolMenuSection& InSectio
 		return;
 	}
 
-	UDynamicMaterialModel* MaterialModel = MenuContext->GetModel();
+	UDynamicMaterialModel* MaterialModel = MenuContext->GetPreviewModel();
 
 	if (!MaterialModel)
 	{
@@ -367,7 +367,7 @@ void FDMMaterialSlotLayerMenus::AddGlobalValueSection(FToolMenuSection& InSectio
 					return;
 				}
 
-				UDynamicMaterialModel* MaterialModel = MenuContext->GetModel();
+				UDynamicMaterialModel* MaterialModel = MenuContext->GetPreviewModel();
 
 				if (!MaterialModel)
 				{
@@ -510,7 +510,7 @@ void FDMMaterialSlotLayerMenus::AddLayerInputsMenu_Slot_Properties(UToolMenu* In
 		return;
 	}
 
-	UDynamicMaterialModel* MaterialModel = MenuContext->GetModel();
+	UDynamicMaterialModel* MaterialModel = MenuContext->GetPreviewModel();
 
 	if (!MaterialModel)
 	{
diff --git a/Engine/Plugins/Experimental/DynamicMaterial/Source/DynamicMaterialEditor/Private/UI/Menus/DMMaterialStageSourceMenus.cpp b/Engine/Plugins/Experimental/DynamicMaterial/Source/DynamicMaterialEditor/Private/UI/Menus/DMMaterialStageSourceMenus.cpp
--- a/Engine/Plugins/Experimental/DynamicMaterial/Source/DynamicMaterialEditor/Private/UI/Menus/DMMaterialStageSourceMenus.cpp
+++ b/Engine/Plugins/Experimental/DynamicMaterial/Source/DynamicMaterialEditor/Private/UI/Menus/DMMaterialStageSourceMenus.cpp
@@ -231,7 +231,7 @@ void FDMMaterialStageSourceMenus::GenerateChangeSourceMenu_GlobalValues(UToolMen
 		return;
 	}
 
-	UDynamicMaterialModel* MaterialModel = MenuContext->GetModel();
+	UDynamicMaterialModel* MaterialModel = MenuContext->GetPreviewModel();
 	if (!ensure(MaterialModel))
 	{
 		return;
@@ -438,7 +438,7 @@ void FDMMaterialStageSourceMenus::GenerateChangeSourceMenu_Slot_Properties(UTool
 		return;
 	}
 
-	UDynamicMaterialModelBase* MaterialModel = MenuContext->GetModelBase();
+	UDynamicMaterialModelBase* MaterialModel = MenuContext->GetPreviewModelBase();
 	if (!MaterialModel)
 	{
 		return;
@@ -560,7 +560,7 @@ void FDMMaterialStageSourceMenus::GenerateChangeSourceMenu_Slots(UToolMenu* InMe
 		return;
 	}
 
-	UDynamicMaterialModelBase* MaterialModelBase = MenuContext->GetModelBase();
+	UDynamicMaterialModelBase* MaterialModelBase = MenuContext->GetPreviewModelBase();
 	if (!MaterialModelBase)
 	{
 		return;
@@ -802,7 +802,7 @@ void FDMMaterialStageSourceMenus::GenerateChangeSourceMenu_Advanced(UToolMenu* c
 
 	UDMMaterialSlot* Slot = MenuContext->GetSlot();
 
-	UDynamicMaterialModel* MaterialModel = MenuContext->GetModel();
+	UDynamicMaterialModel* MaterialModel = MenuContext->GetPreviewModel();
 
 	if (!ensure(MaterialModel))
 	{
@@ -1448,7 +1448,7 @@ void FDMMaterialStageSourceMenus::CreateChangeMaterialStageSource(FToolMenuSecti
 
 	UDMMaterialSlot* Slot = MenuContext->GetSlot();
 
-	UDynamicMaterialModel* MaterialModel = MenuContext->GetModel();
+	UDynamicMaterialModel* MaterialModel = MenuContext->GetPreviewModel();
 
 	if (!ensure(MaterialModel))
 	{
diff --git a/Engine/Plugins/Experimental/DynamicMaterial/Source/DynamicMaterialEditor/Private/UI/Menus/DMMenuContext.cpp b/Engine/Plugins/Experimental/DynamicMaterial/Source/DynamicMaterialEditor/Private/UI/Menus/DMMenuContext.cpp
--- a/Engine/Plugins/Experimental/DynamicMaterial/Source/DynamicMaterialEditor/Private/UI/Menus/DMMenuContext.cpp
+++ b/Engine/Plugins/Experimental/DynamicMaterial/Source/DynamicMaterialEditor/Private/UI/Menus/DMMenuContext.cpp
@@ -112,21 +112,21 @@ UDMMaterialSlot* UDMMenuContext::GetSlot() const
 	return nullptr;;
 }
 
-UDynamicMaterialModelBase* UDMMenuContext::GetModelBase() const
+UDynamicMaterialModelBase* UDMMenuContext::GetPreviewModelBase() const
 {
 	if (TSharedPtr<SDMMaterialEditor> EditorWidget = EditorWidgetWeak.Pin())
 	{
-		return EditorWidget->GetMaterialModelBase();
+		return EditorWidget->GetPreviewMaterialModelBase();
 	}
 
 	return nullptr;
 }
 
-UDynamicMaterialModel* UDMMenuContext::GetModel() const
+UDynamicMaterialModel* UDMMenuContext::GetPreviewModel() const
 {
 	if (TSharedPtr<SDMMaterialEditor> EditorWidget = EditorWidgetWeak.Pin())
 	{
-		return EditorWidget->GetMaterialModel();
+		return EditorWidget->GetPreviewMaterialModel();
 	}
 
 	return nullptr;
diff --git a/Engine/Plugins/Experimental/DynamicMaterial/Source/DynamicMaterialEditor/Private/UI/Menus/DMToolBarMenus.cpp b/Engine/Plugins/Experimental/DynamicMaterial/Source/DynamicMaterialEditor/Private/UI/Menus/DMToolBarMenus.cpp
--- a/Engine/Plugins/Experimental/DynamicMaterial/Source/DynamicMaterialEditor/Private/UI/Menus/DMToolBarMenus.cpp
+++ b/Engine/Plugins/Experimental/DynamicMaterial/Source/DynamicMaterialEditor/Private/UI/Menus/DMToolBarMenus.cpp
@@ -88,7 +88,7 @@ void FDMToolBarMenus::AddExportMenu(UToolMenu* InMenu)
 		return;
 	}
 
-	UDynamicMaterialModelBase* const MaterialModelBase = MenuContext->GetModelBase();
+	UDynamicMaterialModelBase* const MaterialModelBase = MenuContext->GetPreviewModelBase();
 
 	if (!MaterialModelBase)
 	{
@@ -279,7 +279,7 @@ void FDMToolBarMenus::OpenMaterialEditorFromContext(UDMMenuContext* InMenuContex
 {
 	if (IsValid(InMenuContext))
 	{
-		if (UDynamicMaterialModelEditorOnlyData* const ModelEditorOnlyData = UDynamicMaterialModelEditorOnlyData::Get(InMenuContext->GetModel()))
+		if (UDynamicMaterialModelEditorOnlyData* const ModelEditorOnlyData = UDynamicMaterialModelEditorOnlyData::Get(InMenuContext->GetPreviewModel()))
 		{
 			if (FEngineAnalytics::IsAvailable())
 			{
@@ -506,7 +506,7 @@ void FDMToolBarMenus::CreateSnapshotMaterialMenu(UToolMenu* InMenu)
 		return;
 	}
 
-	UDynamicMaterialModelBase* const MaterialModelBase = MenuContext->GetModelBase();
+	UDynamicMaterialModelBase* const MaterialModelBase = MenuContext->GetPreviewModelBase();
 
 	if (!MaterialModelBase)
 	{
diff --git a/Engine/Plugins/Experimental/DynamicMaterial/Source/DynamicMaterialEditor/Private/UI/PropertyGenerators/DMComponentPropertyRowGenerator.cpp b/Engine/Plugins/Experimental/DynamicMaterial/Source/DynamicMaterialEditor/Private/UI/PropertyGenerators/DMComponentPropertyRowGenerator.cpp
--- a/Engine/Plugins/Experimental/DynamicMaterial/Source/DynamicMaterialEditor/Private/UI/PropertyGenerators/DMComponentPropertyRowGenerator.cpp
+++ b/Engine/Plugins/Experimental/DynamicMaterial/Source/DynamicMaterialEditor/Private/UI/PropertyGenerators/DMComponentPropertyRowGenerator.cpp
@@ -7,8 +7,6 @@
 #include "Model/DynamicMaterialModelDynamic.h"
 #include "Runtime/Launch/Resources/Version.h"
 #include "UI/Utils/DMWidgetLibrary.h"
-#include "UI/Widgets/Editor/SDMMaterialComponentEditor.h"
-#include "UI/Widgets/SDMMaterialEditor.h"
 
 const TSharedRef<FDMComponentPropertyRowGenerator>& FDMComponentPropertyRowGenerator::Get()
 {
@@ -16,48 +14,53 @@ const TSharedRef<FDMComponentPropertyRowGenerator>& FDMComponentPropertyRowGener
 	return Generator;
 }
 
-void FDMComponentPropertyRowGenerator::AddComponentProperties(const TSharedRef<SDMMaterialComponentEditor>& InComponentEditorWidget, UDMMaterialComponent* InComponent,
-	TArray<FDMPropertyHandle>& InOutPropertyRows, TSet<UDMMaterialComponent*>& InOutProcessedObjects)
+void FDMComponentPropertyRowGenerator::AddComponentProperties(FDMComponentPropertyRowGeneratorParams& InParams)
 {
-	if (!IsValid(InComponent))
+	if (!IsValid(InParams.Object))
 	{
 		return;
 	}
 
-	if (InOutProcessedObjects.Contains(InComponent))
+	if (InParams.ProcessedObjects.Contains(InParams.Object))
 	{
 		return;
 	}
 
-	InOutProcessedObjects.Add(InComponent);
+	InParams.ProcessedObjects.Add(InParams.Object);
 
-	const TArray<FName>& Properties = InComponent->GetEditableProperties();
+	UDMMaterialComponent* Component = Cast<UDMMaterialComponent>(InParams.Object);
+
+	if (!Component)
+	{
+		return;
+	}
+
+	const TArray<FName>& Properties = Component->GetEditableProperties();
 
 	for (const FName& Property : Properties)
 	{
-		if (InComponent->IsPropertyVisible(Property))
+		if (Component->IsPropertyVisible(Property))
 		{
-			AddPropertyEditRows(InComponentEditorWidget, InComponent, Property, InOutPropertyRows, InOutProcessedObjects);
+			AddPropertyEditRows(InParams, Property);
 		}
 	}
 }
 
-void FDMComponentPropertyRowGenerator::AddPropertyEditRows(const TSharedRef<SDMMaterialComponentEditor>& InComponentEditorWidget, UDMMaterialComponent* InComponent, 
-	const FName& InProperty, TArray<FDMPropertyHandle>& InOutPropertyRows, TSet<UDMMaterialComponent*>& InOutProcessedObjects)
+void FDMComponentPropertyRowGenerator::AddPropertyEditRows(FDMComponentPropertyRowGeneratorParams& InParams, const FName& InProperty)
 {
-	if (!IsValid(InComponent))
+	if (!IsValid(InParams.Object))
 	{
 		return;
 	}
 
-	FProperty* Property = InComponent->GetClass()->FindPropertyByName(InProperty);
+	FProperty* Property = InParams.Object->GetClass()->FindPropertyByName(InProperty);
 
 	if (!Property)
 	{
 		return;
 	}
 
-	void* MemoryPtr = Property->ContainerPtrToValuePtr<void>(InComponent);
+	void* MemoryPtr = Property->ContainerPtrToValuePtr<void>(InParams.Object);
 
 	if (FArrayProperty* ArrayProperty = CastField<FArrayProperty>(Property))
 	{
@@ -70,17 +73,17 @@ void FDMComponentPropertyRowGenerator::AddPropertyEditRows(const TSharedRef<SDMM
 #else
 			void* ElemPtr = ArrayHelper.GetElementPtr(Idx);
 #endif
-			AddPropertyEditRows(InComponentEditorWidget, InComponent, ArrayProperty->Inner, ElemPtr, InOutPropertyRows, InOutProcessedObjects);
+			AddPropertyEditRows(InParams, ArrayProperty->Inner, ElemPtr);
 		}
 	}
 	else
 	{
-		AddPropertyEditRows(InComponentEditorWidget, InComponent, Property, MemoryPtr, InOutPropertyRows, InOutProcessedObjects);
+		AddPropertyEditRows(InParams, Property, MemoryPtr);
 	}
 }
 
-void FDMComponentPropertyRowGenerator::AddPropertyEditRows(const TSharedRef<SDMMaterialComponentEditor>& InComponentEditorWidget, UDMMaterialComponent* InComponent,
-	FProperty* InProperty, void* MemoryPtr, TArray<FDMPropertyHandle>& InOutPropertyRows, TSet<UDMMaterialComponent*>& InOutProcessedObjects)
+void FDMComponentPropertyRowGenerator::AddPropertyEditRows(FDMComponentPropertyRowGeneratorParams& InParams, FProperty* InProperty, 
+	void* InMemoryPtr)
 {
 	if (InProperty->IsA<FArrayProperty>())
 	{
@@ -91,16 +94,20 @@ void FDMComponentPropertyRowGenerator::AddPropertyEditRows(const TSharedRef<SDMM
 	{
 		if (ObjectProperty->PropertyClass->IsChildOf(UDMMaterialComponent::StaticClass()))
 		{
-			UObject** ValuePtr = static_cast<UObject**>(MemoryPtr);
+			UObject** ValuePtr = static_cast<UObject**>(InMemoryPtr);
 			UObject* Value = *ValuePtr;
 			UDMMaterialComponent* ComponentValue = Cast<UDMMaterialComponent>(Value);
-			FDynamicMaterialEditorModule::GeneratorComponentPropertyRows(InComponentEditorWidget, ComponentValue, InOutPropertyRows, InOutProcessedObjects);
+
+			FDMComponentPropertyRowGeneratorParams ChildParams = InParams;
+			ChildParams.Object = ComponentValue;
+
+			FDynamicMaterialEditorModule::GeneratorComponentPropertyRows(ChildParams);
 			return;
 		}
 	}
 
-	FDMPropertyHandle& Handle = InOutPropertyRows.Add_GetRef(FDMWidgetLibrary::Get().GetPropertyHandle(&*InComponentEditorWidget, InComponent, InProperty->GetFName()));
-	Handle.bEnabled = !IsDynamic(InComponentEditorWidget);
+	FDMPropertyHandle& Handle = InParams.PropertyRows.Add_GetRef(FDMWidgetLibrary::Get().GetPropertyHandle(InParams.CreatePropertyHandleParams(InProperty->GetFName())));
+	Handle.bEnabled = !IsDynamic(InParams);
 }
 
 bool FDMComponentPropertyRowGenerator::AllowKeyframeButton(UDMMaterialComponent* InComponent, FProperty* InProperty)
@@ -108,15 +115,7 @@ bool FDMComponentPropertyRowGenerator::AllowKeyframeButton(UDMMaterialComponent*
 	return false;
 }
 
-bool FDMComponentPropertyRowGenerator::IsDynamic(const TSharedRef<SDMMaterialComponentEditor>& InComponentEditorWidget)
+bool FDMComponentPropertyRowGenerator::IsDynamic(FDMComponentPropertyRowGeneratorParams& InParams)
 {
-	if (TSharedPtr<SDMMaterialEditor> EditorWidget = InComponentEditorWidget->GetEditorWidget())
-	{
-		if (UDynamicMaterialModelBase* MaterialModelBase = EditorWidget->GetMaterialModelBase())
-		{
-			return !!Cast<UDynamicMaterialModelDynamic>(MaterialModelBase);
-		}
-	}
-
-	return false;
+	return InParams.PreviewMaterialModelBase && InParams.PreviewMaterialModelBase->IsA<UDynamicMaterialModelDynamic>();
 }
diff --git a/Engine/Plugins/Experimental/DynamicMaterial/Source/DynamicMaterialEditor/Private/UI/PropertyGenerators/DMInputThroughputPropertyRowGenerator.cpp b/Engine/Plugins/Experimental/DynamicMaterial/Source/DynamicMaterialEditor/Private/UI/PropertyGenerators/DMInputThroughputPropertyRowGenerator.cpp
--- a/Engine/Plugins/Experimental/DynamicMaterial/Source/DynamicMaterialEditor/Private/UI/PropertyGenerators/DMInputThroughputPropertyRowGenerator.cpp
+++ b/Engine/Plugins/Experimental/DynamicMaterial/Source/DynamicMaterialEditor/Private/UI/PropertyGenerators/DMInputThroughputPropertyRowGenerator.cpp
@@ -5,6 +5,7 @@
 #include "Components/DMMaterialStageThroughput.h"
 #include "Components/DMMaterialSubStage.h"
 #include "Components/MaterialStageInputs/DMMSIThroughput.h"
+#include "IDynamicMaterialEditorModule.h"
 
 const TSharedRef<FDMInputThroughputPropertyRowGenerator>& FDMInputThroughputPropertyRowGenerator::Get()
 {
@@ -12,32 +13,30 @@ const TSharedRef<FDMInputThroughputPropertyRowGenerator>& FDMInputThroughputProp
 	return Generator;
 }
 
-void FDMInputThroughputPropertyRowGenerator::AddComponentProperties(const TSharedRef<SDMMaterialComponentEditor>& InComponentEditorWidget,
-	UDMMaterialComponent* InComponent, TArray<FDMPropertyHandle>& InOutPropertyRows,
-	TSet<UDMMaterialComponent*>& InOutProcessedObjects)
+void FDMInputThroughputPropertyRowGenerator::AddComponentProperties(FDMComponentPropertyRowGeneratorParams& InParams)
 {
-	if (!IsValid(InComponent))
+	if (!IsValid(InParams.Object))
 	{
 		return;
 	}
 
-	if (InOutProcessedObjects.Contains(InComponent))
+	if (InParams.ProcessedObjects.Contains(InParams.Object))
 	{
 		return;
 	}
 
-	UDMMaterialStageInputThroughput* InputThroughput = Cast<UDMMaterialStageInputThroughput>(InComponent);
+	UDMMaterialStageInputThroughput* InputThroughput = Cast<UDMMaterialStageInputThroughput>(InParams.Object);
 
 	if (!InputThroughput)
 	{
 		return;
 	}
 
-	InOutProcessedObjects.Add(InputThroughput);
+	InParams.ProcessedObjects.Add(InputThroughput);
 
 	if (UDMMaterialSubStage* SubStage = InputThroughput->GetSubStage())
 	{
-		InOutProcessedObjects.Add(SubStage);
+		InParams.ProcessedObjects.Add(SubStage);
 	}
 
 	UDMMaterialStageThroughput* Throughput = InputThroughput->GetMaterialStageThroughput();
@@ -47,5 +46,8 @@ void FDMInputThroughputPropertyRowGenerator::AddComponentProperties(const TShare
 		return;
 	}
 
-	FDMThroughputPropertyRowGenerator::AddComponentProperties(InComponentEditorWidget, Throughput, InOutPropertyRows, InOutProcessedObjects);
+	FDMComponentPropertyRowGeneratorParams ThroughputParams = InParams;
+	ThroughputParams.Object = Throughput;
+
+	FDMThroughputPropertyRowGenerator::AddComponentProperties(ThroughputParams);
 }
diff --git a/Engine/Plugins/Experimental/DynamicMaterial/Source/DynamicMaterialEditor/Private/UI/PropertyGenerators/DMInputThroughputPropertyRowGenerator.h b/Engine/Plugins/Experimental/DynamicMaterial/Source/DynamicMaterialEditor/Private/UI/PropertyGenerators/DMInputThroughputPropertyRowGenerator.h
--- a/Engine/Plugins/Experimental/DynamicMaterial/Source/DynamicMaterialEditor/Private/UI/PropertyGenerators/DMInputThroughputPropertyRowGenerator.h
+++ b/Engine/Plugins/Experimental/DynamicMaterial/Source/DynamicMaterialEditor/Private/UI/PropertyGenerators/DMInputThroughputPropertyRowGenerator.h
@@ -12,7 +12,6 @@ public:
 	virtual ~FDMInputThroughputPropertyRowGenerator() override = default;
 
 	//~ Begin FDMComponentPropertyRowGenerator
-	virtual void AddComponentProperties(const TSharedRef<SDMMaterialComponentEditor>& InComponentEditorWidget, UDMMaterialComponent* InComponent,
-		TArray<FDMPropertyHandle>& InOutPropertyRows, TSet<UDMMaterialComponent*>& InOutProcessedObjects) override;
+	virtual void AddComponentProperties(FDMComponentPropertyRowGeneratorParams& InParams) override;
 	//~ End FDMComponentPropertyRowGenerator
 };
diff --git a/Engine/Plugins/Experimental/DynamicMaterial/Source/DynamicMaterialEditor/Private/UI/PropertyGenerators/DMMaterialEffectFunctionPropertyRowGenerator.cpp b/Engine/Plugins/Experimental/DynamicMaterial/Source/DynamicMaterialEditor/Private/UI/PropertyGenerators/DMMaterialEffectFunctionPropertyRowGenerator.cpp
--- a/Engine/Plugins/Experimental/DynamicMaterial/Source/DynamicMaterialEditor/Private/UI/PropertyGenerators/DMMaterialEffectFunctionPropertyRowGenerator.cpp
+++ b/Engine/Plugins/Experimental/DynamicMaterial/Source/DynamicMaterialEditor/Private/UI/PropertyGenerators/DMMaterialEffectFunctionPropertyRowGenerator.cpp
@@ -9,7 +9,6 @@
 #include "Materials/MaterialExpressionFunctionInput.h"
 #include "Materials/MaterialExpressionMaterialFunctionCall.h"
 #include "Materials/MaterialFunctionInterface.h"
-#include "UI/Widgets/Editor/SDMMaterialComponentEditor.h"
 #include "Utils/DMMaterialFunctionFunctionLibrary.h"
 
 #define LOCTEXT_NAMESPACE "DMMaterialEffectFunctionPropertyRowGenerator"
@@ -20,27 +19,26 @@ const TSharedRef<FDMMaterialEffectFunctionPropertyRowGenerator>& FDMMaterialEffe
 	return Generator;
 }
 
-void FDMMaterialEffectFunctionPropertyRowGenerator::AddComponentProperties(const TSharedRef<SDMMaterialComponentEditor>& InComponentEditorWidget, UDMMaterialComponent* InComponent,
-	TArray<FDMPropertyHandle>& InOutPropertyRows, TSet<UDMMaterialComponent*>& InOutProcessedObjects)
+void FDMMaterialEffectFunctionPropertyRowGenerator::AddComponentProperties(FDMComponentPropertyRowGeneratorParams& InParams)
 {
-	if (!IsValid(InComponent))
+	if (!IsValid(InParams.Object))
 	{
 		return;
 	}
 
-	if (InOutProcessedObjects.Contains(InComponent))
+	if (InParams.ProcessedObjects.Contains(InParams.Object))
 	{
 		return;
 	}
 
-	UDMMaterialEffectFunction* EffectFunction = Cast<UDMMaterialEffectFunction>(InComponent);
+	UDMMaterialEffectFunction* EffectFunction = Cast<UDMMaterialEffectFunction>(InParams.Object);
 
 	if (!EffectFunction)
 	{
 		return;
 	}
 
-	InOutProcessedObjects.Add(InComponent);
+	InParams.ProcessedObjects.Add(InParams.Object);
 
 	UMaterialFunctionInterface* MaterialFunction = EffectFunction->GetMaterialFunction();
 
@@ -75,7 +73,12 @@ void FDMMaterialEffectFunctionPropertyRowGenerator::AddComponentProperties(const
 		}
 
 		TArray<FDMPropertyHandle> ValuePropertyRows;
-		FDynamicMaterialEditorModule::GeneratorComponentPropertyRows(InComponentEditorWidget, Value, ValuePropertyRows, InOutProcessedObjects);
+
+		FDMComponentPropertyRowGeneratorParams ValueParams = InParams;
+		ValueParams.Object = Value;
+		ValueParams.PropertyRows = ValuePropertyRows;
+
+		FDynamicMaterialEditorModule::GeneratorComponentPropertyRows(ValueParams);
 
 		if (ValuePropertyRows.Num() == 1)
 		{
@@ -106,17 +109,17 @@ void FDMMaterialEffectFunctionPropertyRowGenerator::AddComponentProperties(const
 			ValuePropertyRow.NameToolTipOverride = Description;
 			ValuePropertyRow.CategoryOverrideName = MaterialInputName;
 
-			if (ValuePropertyRow.PropertyHandle.IsValid())
+			if (ValuePropertyRow.PreviewHandle.PropertyHandle.IsValid())
 			{
 				UDMMaterialFunctionFunctionLibrary::ApplyMetaData(
 					Inputs[InputIndex],
-					ValuePropertyRow.PropertyHandle.ToSharedRef()
+					ValuePropertyRow.PreviewHandle.PropertyHandle.ToSharedRef()
 				);
 			}
 
 		}
 
-		InOutPropertyRows.Append(ValuePropertyRows);
+		InParams.PropertyRows.Append(ValuePropertyRows);
 	}
 }
 
diff --git a/Engine/Plugins/Experimental/DynamicMaterial/Source/DynamicMaterialEditor/Private/UI/PropertyGenerators/DMMaterialEffectFunctionPropertyRowGenerator.h b/Engine/Plugins/Experimental/DynamicMaterial/Source/DynamicMaterialEditor/Private/UI/PropertyGenerators/DMMaterialEffectFunctionPropertyRowGenerator.h
--- a/Engine/Plugins/Experimental/DynamicMaterial/Source/DynamicMaterialEditor/Private/UI/PropertyGenerators/DMMaterialEffectFunctionPropertyRowGenerator.h
+++ b/Engine/Plugins/Experimental/DynamicMaterial/Source/DynamicMaterialEditor/Private/UI/PropertyGenerators/DMMaterialEffectFunctionPropertyRowGenerator.h
@@ -12,7 +12,6 @@ public:
 	virtual ~FDMMaterialEffectFunctionPropertyRowGenerator() override = default;
 
 	//~ Begin FDMComponentPropertyRowGenerator
-	virtual void AddComponentProperties(const TSharedRef<SDMMaterialComponentEditor>& InComponentEditorWidget, UDMMaterialComponent* InComponent,
-		TArray<FDMPropertyHandle>& InOutPropertyRows, TSet<UDMMaterialComponent*>& InOutProcessedObjects) override;
+	virtual void AddComponentProperties(FDMComponentPropertyRowGeneratorParams& InParams) override;
 	//~ End FDMComponentPropertyRowGenerator
 };
diff --git a/Engine/Plugins/Experimental/DynamicMaterial/Source/DynamicMaterialEditor/Private/UI/PropertyGenerators/DMMaterialModelPropertyRowGenerator.cpp b/Engine/Plugins/Experimental/DynamicMaterial/Source/DynamicMaterialEditor/Private/UI/PropertyGenerators/DMMaterialModelPropertyRowGenerator.cpp
--- a/Engine/Plugins/Experimental/DynamicMaterial/Source/DynamicMaterialEditor/Private/UI/PropertyGenerators/DMMaterialModelPropertyRowGenerator.cpp
+++ b/Engine/Plugins/Experimental/DynamicMaterial/Source/DynamicMaterialEditor/Private/UI/PropertyGenerators/DMMaterialModelPropertyRowGenerator.cpp
@@ -4,76 +4,80 @@
 
 #include "Components/DMMaterialProperty.h"
 #include "Components/MaterialValues/DMMaterialValueFloat1.h"
+#include "IDynamicMaterialEditorModule.h"
 #include "Model/DynamicMaterialModel.h"
 #include "Model/DynamicMaterialModelDynamic.h"
 #include "Model/DynamicMaterialModelEditorOnlyData.h"
 #include "UI/Utils/DMWidgetLibrary.h"
-#include "UI/Widgets/Editor/SDMMaterialGlobalSettingsEditor.h"
-#include "Utils/DMPrivate.h"
 
 #define LOCTEXT_NAMESPACE "DMMaterialModelPropertyRowGenerator"
 
-void FDMMaterialModelPropertyRowGenerator::AddMaterialModelProperties(const TSharedRef<SDMMaterialGlobalSettingsEditor>& InGlobalSettingEditorWidget, 
-	UDynamicMaterialModelBase* InMaterialModelBase, TArray<FDMPropertyHandle>& InOutPropertyRows)
+void FDMMaterialModelPropertyRowGenerator::AddMaterialModelProperties(FDMComponentPropertyRowGeneratorParams& InParams)
 {
-	UDynamicMaterialModel* MaterialModel = InMaterialModelBase->ResolveMaterialModel();
+	UDynamicMaterialModelBase* MaterialModelBase = Cast<UDynamicMaterialModelBase>(InParams.Object);
+
+	if (!MaterialModelBase)
+	{
+		return;
+	}
+
+	UDynamicMaterialModel* MaterialModel = MaterialModelBase->ResolveMaterialModel();
 
 	if (!MaterialModel)
 	{
 		return;
 	}
 
-	AddGlobalValue(InGlobalSettingEditorWidget, InMaterialModelBase, InOutPropertyRows, 
+	AddGlobalValue(InParams, 
 		MaterialModel->GetGlobalParameterValue(UDynamicMaterialModel::GlobalOffsetValueName), 
 		LOCTEXT("GlobalOffset", "Global Offset"));
 
-	AddGlobalValue(InGlobalSettingEditorWidget, InMaterialModelBase, InOutPropertyRows,
+	AddGlobalValue(InParams,
 		MaterialModel->GetGlobalParameterValue(UDynamicMaterialModel::GlobalTilingValueName),
 		LOCTEXT("GlobalTiling", "Global Tiling"));
 
-	AddGlobalValue(InGlobalSettingEditorWidget, InMaterialModelBase, InOutPropertyRows,
+	AddGlobalValue(InParams,
 		MaterialModel->GetGlobalParameterValue(UDynamicMaterialModel::GlobalRotationValueName), 
 		LOCTEXT("GlobalRotation", "Global Rotation"));
 
-	if (UDynamicMaterialModelEditorOnlyData* EditorOnlyData = UDynamicMaterialModelEditorOnlyData::Get(InMaterialModelBase))
+	if (UDynamicMaterialModelEditorOnlyData* EditorOnlyData = UDynamicMaterialModelEditorOnlyData::Get(MaterialModelBase))
 	{
-		AddVariable(InGlobalSettingEditorWidget, InMaterialModelBase, InOutPropertyRows, EditorOnlyData,
+		AddVariable(InParams, EditorOnlyData,
 			GET_MEMBER_NAME_CHECKED(UDynamicMaterialModelEditorOnlyData, Domain));
 
-		AddVariable(InGlobalSettingEditorWidget, InMaterialModelBase, InOutPropertyRows, EditorOnlyData,
+		AddVariable(InParams, EditorOnlyData,
 			GET_MEMBER_NAME_CHECKED(UDynamicMaterialModelEditorOnlyData, BlendMode));
 
-		AddVariable(InGlobalSettingEditorWidget, InMaterialModelBase, InOutPropertyRows, EditorOnlyData,
+		AddVariable(InParams, EditorOnlyData,
 			GET_MEMBER_NAME_CHECKED(UDynamicMaterialModelEditorOnlyData, ShadingModel));
 
-		AddVariable(InGlobalSettingEditorWidget, InMaterialModelBase, InOutPropertyRows, EditorOnlyData,
+		AddVariable(InParams, EditorOnlyData,
 			GET_MEMBER_NAME_CHECKED(UDynamicMaterialModelEditorOnlyData, bHasPixelAnimation));
 
-		AddVariable(InGlobalSettingEditorWidget, InMaterialModelBase, InOutPropertyRows, EditorOnlyData,
+		AddVariable(InParams, EditorOnlyData,
 			GET_MEMBER_NAME_CHECKED(UDynamicMaterialModelEditorOnlyData, bTwoSided));
 
-		AddVariable(InGlobalSettingEditorWidget, InMaterialModelBase, InOutPropertyRows, EditorOnlyData,
+		AddVariable(InParams, EditorOnlyData,
 			GET_MEMBER_NAME_CHECKED(UDynamicMaterialModelEditorOnlyData, bResponsiveAAEnabled));
 
-		AddVariable(InGlobalSettingEditorWidget, InMaterialModelBase, InOutPropertyRows, EditorOnlyData,
+		AddVariable(InParams, EditorOnlyData,
 			GET_MEMBER_NAME_CHECKED(UDynamicMaterialModelEditorOnlyData, bOutputTranslucentVelocityEnabled));
 
-		AddVariable(InGlobalSettingEditorWidget, InMaterialModelBase, InOutPropertyRows, EditorOnlyData,
+		AddVariable(InParams, EditorOnlyData,
 			GET_MEMBER_NAME_CHECKED(UDynamicMaterialModelEditorOnlyData, bNaniteTessellationEnabled));
 
-		AddVariable(InGlobalSettingEditorWidget, InMaterialModelBase, InOutPropertyRows, EditorOnlyData,
+		AddVariable(InParams, EditorOnlyData,
 			GET_MEMBER_NAME_CHECKED(UDynamicMaterialModelEditorOnlyData, DisplacementMagnitude));
 
-		AddVariable(InGlobalSettingEditorWidget, InMaterialModelBase, InOutPropertyRows, EditorOnlyData,
+		AddVariable(InParams, EditorOnlyData,
 			GET_MEMBER_NAME_CHECKED(UDynamicMaterialModelEditorOnlyData, DisplacementCenter));
 	}
 }
 
-void FDMMaterialModelPropertyRowGenerator::AddGlobalValue(const TSharedRef<SDMMaterialGlobalSettingsEditor>& InGlobalSettingEditorWidget,
-	UDynamicMaterialModelBase* InMaterialModelBase, TArray<FDMPropertyHandle>& InOutPropertyRows, UDMMaterialComponent* InComponent, 
+void FDMMaterialModelPropertyRowGenerator::AddGlobalValue(FDMComponentPropertyRowGeneratorParams& InParams, UDMMaterialComponent* InComponent,
 	const FText& InNameOverride)
 {
-	if (UDynamicMaterialModelDynamic* MaterialModelDynamic = Cast<UDynamicMaterialModelDynamic>(InMaterialModelBase))
+	if (UDynamicMaterialModelDynamic* MaterialModelDynamic = Cast<UDynamicMaterialModelDynamic>(InParams.Object))
 	{
 		InComponent = MaterialModelDynamic->GetComponentDynamic(InComponent->GetFName());
 
@@ -83,10 +87,11 @@ void FDMMaterialModelPropertyRowGenerator::AddGlobalValue(const TSharedRef<SDMMa
 		}
 	}
 
-	FDMPropertyHandle& ComponentHandle = InOutPropertyRows.Add_GetRef(FDMWidgetLibrary::Get().GetPropertyHandle(
-		&*InGlobalSettingEditorWidget,
-		InComponent,
-		UDMMaterialValue::ValueName
+	FDMComponentPropertyRowGeneratorParams GlobalValueParams = InParams;
+	GlobalValueParams.Object = InComponent;
+
+	FDMPropertyHandle& ComponentHandle = InParams.PropertyRows.Add_GetRef(FDMWidgetLibrary::Get().GetPropertyHandle(
+		GlobalValueParams.CreatePropertyHandleParams(UDMMaterialValue::ValueName)
 	));
 
 	ComponentHandle.CategoryOverrideName = TEXT("Material Settings");
@@ -110,16 +115,17 @@ void FDMMaterialModelPropertyRowGenerator::AddGlobalValue(const TSharedRef<SDMMa
 	}
 }
 
-void FDMMaterialModelPropertyRowGenerator::AddVariable(const TSharedRef<SDMMaterialGlobalSettingsEditor>& InGlobalSettingEditorWidget, 
-	UDynamicMaterialModelBase* InMaterialModelBase, TArray<FDMPropertyHandle>& InOutPropertyRows, UObject* InObject, FName InPropertyName)
+void FDMMaterialModelPropertyRowGenerator::AddVariable(FDMComponentPropertyRowGeneratorParams& InParams, UObject* InObject, FName InPropertyName)
 {
-	FDMPropertyHandle& ValueHandle = InOutPropertyRows.Add_GetRef(FDMWidgetLibrary::Get().GetPropertyHandle(
-		&*InGlobalSettingEditorWidget, 
-		InObject, InPropertyName
+	FDMComponentPropertyRowGeneratorParams VariableParams = InParams;
+	VariableParams.Object = InObject;
+
+	FDMPropertyHandle& ValueHandle = InParams.PropertyRows.Add_GetRef(FDMWidgetLibrary::Get().GetPropertyHandle(
+		VariableParams.CreatePropertyHandleParams(InPropertyName)
 	));
 
 	ValueHandle.CategoryOverrideName = TEXT("Material Type");
-	ValueHandle.bEnabled = !IsDynamic(InMaterialModelBase);
+	ValueHandle.bEnabled = !IsDynamic(Cast<UDynamicMaterialModelBase>(InParams.Object));
 }
 
 bool FDMMaterialModelPropertyRowGenerator::IsDynamic(UDynamicMaterialModelBase* InMaterialModelBase)
diff --git a/Engine/Plugins/Experimental/DynamicMaterial/Source/DynamicMaterialEditor/Private/UI/PropertyGenerators/DMMaterialModelPropertyRowGenerator.h b/Engine/Plugins/Experimental/DynamicMaterial/Source/DynamicMaterialEditor/Private/UI/PropertyGenerators/DMMaterialModelPropertyRowGenerator.h
--- a/Engine/Plugins/Experimental/DynamicMaterial/Source/DynamicMaterialEditor/Private/UI/PropertyGenerators/DMMaterialModelPropertyRowGenerator.h
+++ b/Engine/Plugins/Experimental/DynamicMaterial/Source/DynamicMaterialEditor/Private/UI/PropertyGenerators/DMMaterialModelPropertyRowGenerator.h
@@ -12,6 +12,7 @@ class UDMMaterialComponent;
 class UDynamicMaterialModelBase;
 class UDynamicMaterialModelEditorOnlyData;
 enum class EDMMaterialPropertyType : uint8;
+struct FDMComponentPropertyRowGeneratorParams;
 struct FDMPropertyHandle;
 
 /** Used to generate editable properties for material models. */
@@ -24,15 +25,12 @@ public:
 	 * @param InMaterialModelBase The model being edited.
 	 * @param InOutPropertyRows The generated rows.
 	 */
-	static void AddMaterialModelProperties(const TSharedRef<SDMMaterialGlobalSettingsEditor>& InGlobalSettingEditorWidget, 
-		UDynamicMaterialModelBase* InMaterialModelBase, TArray<FDMPropertyHandle>& InOutPropertyRows);
+	static void AddMaterialModelProperties(FDMComponentPropertyRowGeneratorParams& InParams);
 
-	static void AddGlobalValue(const TSharedRef<SDMMaterialGlobalSettingsEditor>& InGlobalSettingEditorWidget,
-		UDynamicMaterialModelBase* InMaterialModelBase, TArray<FDMPropertyHandle>& InOutPropertyRows, UDMMaterialComponent* InComponent,
+	static void AddGlobalValue(FDMComponentPropertyRowGeneratorParams& InParams, UDMMaterialComponent* InComponent,
 		const FText& InNameOverride);
 
-	static void AddVariable(const TSharedRef<SDMMaterialGlobalSettingsEditor>& InGlobalSettingEditorWidget,
-		UDynamicMaterialModelBase* InMaterialModelBase, TArray<FDMPropertyHandle>& InOutPropertyRows, UObject* InObject, FName InPropertyName);
+	static void AddVariable(FDMComponentPropertyRowGeneratorParams& InParams, UObject* InObject, FName InPropertyName);
 
 protected:
 	/** Returns true if the component edit widget is editing a Material Designer Dynamic. */
