Material Designer: Apply button

diff --git a/Engine/Plugins/Experimental/DynamicMaterial/Source/DynamicMaterialEditor/Private/UI/PropertyGenerators/DMMaterialStageFunctionPropertyRowGenerator.cpp b/Engine/Plugins/Experimental/DynamicMaterial/Source/DynamicMaterialEditor/Private/UI/PropertyGenerators/DMMaterialStageFunctionPropertyRowGenerator.cpp
--- a/Engine/Plugins/Experimental/DynamicMaterial/Source/DynamicMaterialEditor/Private/UI/PropertyGenerators/DMMaterialStageFunctionPropertyRowGenerator.cpp
+++ b/Engine/Plugins/Experimental/DynamicMaterial/Source/DynamicMaterialEditor/Private/UI/PropertyGenerators/DMMaterialStageFunctionPropertyRowGenerator.cpp
@@ -10,7 +10,6 @@
 #include "Materials/MaterialExpressionFunctionInput.h"
 #include "Materials/MaterialExpressionMaterialFunctionCall.h"
 #include "Materials/MaterialFunctionInterface.h"
-#include "UI/Widgets/Editor/SDMMaterialComponentEditor.h"
 #include "Utils/DMMaterialFunctionFunctionLibrary.h"
 
 #define LOCTEXT_NAMESPACE "DMMaterialStageFunctionPropertyRowGenerator"
@@ -21,128 +20,129 @@ const TSharedRef<FDMMaterialStageFunctionPropertyRowGenerator>& FDMMaterialStage
 	return Generator;
 }
 
-void FDMMaterialStageFunctionPropertyRowGenerator::AddComponentProperties(const TSharedRef<SDMMaterialComponentEditor>& InComponentEditorWidget, UDMMaterialComponent* InComponent,
-	TArray<FDMPropertyHandle>& InOutPropertyRows, TSet<UDMMaterialComponent*>& InOutProcessedObjects)
+void FDMMaterialStageFunctionPropertyRowGenerator::AddComponentProperties(FDMComponentPropertyRowGeneratorParams& InParams)
 {
-	if (!IsValid(InComponent))
+	if (!IsValid(InParams.Object))
 	{
 		return;
 	}
 
-	if (InOutProcessedObjects.Contains(InComponent))
+	if (InParams.ProcessedObjects.Contains(InParams.Object))
+	{
+		return;
+	}
+
+	UDMMaterialStageInputFunction* StageInputFunction = Cast<UDMMaterialStageInputFunction>(InParams.Object);
+
+	if (!StageInputFunction)
+	{
+		return;
+	}
+
+	UDMMaterialStageFunction* MaterialStageFunction = StageInputFunction->GetMaterialStageFunction();
+
+	if (!MaterialStageFunction)
 	{
 		return;
 	}
 
-	UDMMaterialStageFunction* MaterialStageFunction = nullptr;
 	TArray<FDMPropertyHandle> AllValuePropertyRows;
 
-	if (UDMMaterialStageInputFunction* StageInputFunction = Cast<UDMMaterialStageInputFunction>(InComponent))
+	InParams.ProcessedObjects.Add(InParams.Object);
+
+	if (UMaterialFunctionInterface* MaterialFunction = MaterialStageFunction->GetMaterialFunction())
 	{
-		MaterialStageFunction = StageInputFunction->GetMaterialStageFunction();
+		TArray<FFunctionExpressionInput> Inputs;
+		TArray<FFunctionExpressionOutput> Outputs;
+		MaterialFunction->GetInputsAndOutputs(Inputs, Outputs);
+
+		TArray<UDMMaterialValue*> InputValues = MaterialStageFunction->GetInputValues();
+
+		// 1 input, the previous stage, does not have a value.
+		if (Inputs.Num() != (InputValues.Num() + 1))
+		{
+			return;
+		}
 
-		if (MaterialStageFunction)
+		for (int32 InputIndex = 0; InputIndex < InputValues.Num(); ++InputIndex)
 		{
-			InOutProcessedObjects.Add(InComponent);
+			FFunctionExpressionInput& Input = Inputs[InputIndex + 1];
+			UDMMaterialValue* Value = InputValues[InputIndex];
 
-			if (UMaterialFunctionInterface* MaterialFunction = MaterialStageFunction->GetMaterialFunction())
+			if (!IsValid(Value))
 			{
-				TArray<FFunctionExpressionInput> Inputs;
-				TArray<FFunctionExpressionOutput> Outputs;
-				MaterialFunction->GetInputsAndOutputs(Inputs, Outputs);
+				continue;
+			}
+
+			if (!Input.ExpressionInput)
+			{
+				continue;
+			}
 
-				TArray<UDMMaterialValue*> InputValues = MaterialStageFunction->GetInputValues();
+			TArray<FDMPropertyHandle> ValuePropertyRows;
 
-				// 1 input, the previous stage, does not have a value.
-				if (Inputs.Num() != (InputValues.Num() + 1))
+			FDMComponentPropertyRowGeneratorParams ValueParams = InParams;
+			ValueParams.Object = Value;
+			ValueParams.PropertyRows = ValuePropertyRows;
+
+			FDynamicMaterialEditorModule::GeneratorComponentPropertyRows(ValueParams);
+
+			if (ValuePropertyRows.Num() == 1)
+			{
+				ValuePropertyRows[0].NameOverride = FText::FromName(Input.ExpressionInput->InputName);
+			}
+			else
+			{
+				static const FText NameFormat = LOCTEXT("ValueFormat", "{0}[{1}]");
+
+				for (int32 ValuePropertyIndex = 0; ValuePropertyIndex < ValuePropertyRows.Num(); ++ValuePropertyIndex)
 				{
-					return;
+					ValuePropertyRows[ValuePropertyIndex].NameOverride = FText::Format(
+						NameFormat,
+						FText::FromName(Input.ExpressionInput->InputName),
+						FText::AsNumber(ValuePropertyIndex + 1)
+					);
 				}
+			}
+
+			const FText Description = FText::FromString(Input.ExpressionInput->Description);
+
+			FText MaterialInputText = FText::FromString(MaterialFunction->GetUserExposedCaption());
 
-				for (int32 InputIndex = 0; InputIndex < InputValues.Num(); ++InputIndex)
+			if (MaterialInputText.IsEmpty())
+			{
+				MaterialInputText = LOCTEXT("Function", "Function");
+			}
+
+			const FText MaterialInputFormat = LOCTEXT("MaterialInputFormat", "{0} Inputs");
+			MaterialInputText = FText::Format(MaterialInputFormat, MaterialInputText);
+			const FName MaterialInputName = *MaterialInputText.ToString();
+
+			for (FDMPropertyHandle& ValuePropertyRow : ValuePropertyRows)
+			{
+				ValuePropertyRow.NameToolTipOverride = Description;
+				ValuePropertyRow.CategoryOverrideName = MaterialInputName;
+
+				if (ValuePropertyRow.PreviewHandle.PropertyHandle.IsValid())
 				{
-					FFunctionExpressionInput& Input = Inputs[InputIndex + 1];
-					UDMMaterialValue* Value = InputValues[InputIndex];
-
-					if (!IsValid(Value))
-					{
-						continue;
-					}
-
-					if (!Input.ExpressionInput)
-					{
-						continue;
-					}
-
-					TArray<FDMPropertyHandle> ValuePropertyRows;
-
-					FDynamicMaterialEditorModule::GeneratorComponentPropertyRows(
-						InComponentEditorWidget,
-						Value,
-						ValuePropertyRows,
-						InOutProcessedObjects
+					// The first input is the previous layer, so it does not have a value.
+					UDMMaterialFunctionFunctionLibrary::ApplyMetaData(
+						Inputs[InputIndex + 1],
+						ValuePropertyRow.PreviewHandle.PropertyHandle.ToSharedRef()
 					);
-
-					if (ValuePropertyRows.Num() == 1)
-					{
-						ValuePropertyRows[0].NameOverride = FText::FromName(Input.ExpressionInput->InputName);
-					}
-					else
-					{
-						static const FText NameFormat = LOCTEXT("ValueFormat", "{0}[{1}]");
-
-						for (int32 ValuePropertyIndex = 0; ValuePropertyIndex < ValuePropertyRows.Num(); ++ValuePropertyIndex)
-						{
-							ValuePropertyRows[ValuePropertyIndex].NameOverride = FText::Format(
-								NameFormat,
-								FText::FromName(Input.ExpressionInput->InputName),
-								FText::AsNumber(ValuePropertyIndex + 1)
-							);
-						}
-					}
-
-					const FText Description = FText::FromString(Input.ExpressionInput->Description);
-
-					FText MaterialInputText = FText::FromString(MaterialFunction->GetUserExposedCaption());
-
-					if (MaterialInputText.IsEmpty())
-					{
-						MaterialInputText = LOCTEXT("Function", "Function");
-					}
-
-					const FText MaterialInputFormat = LOCTEXT("MaterialInputFormat", "{0} Inputs");
-					MaterialInputText = FText::Format(MaterialInputFormat, MaterialInputText);
-					const FName MaterialInputName = *MaterialInputText.ToString();
-
-					for (FDMPropertyHandle& ValuePropertyRow : ValuePropertyRows)
-					{
-						ValuePropertyRow.NameToolTipOverride = Description;
-						ValuePropertyRow.CategoryOverrideName = MaterialInputName;
-
-						if (ValuePropertyRow.PropertyHandle.IsValid())
-						{
-							// The first input is the previous layer, so it does not have a value.
-							UDMMaterialFunctionFunctionLibrary::ApplyMetaData(
-								Inputs[InputIndex + 1],
-								ValuePropertyRow.PropertyHandle.ToSharedRef()
-							);
-						}
-					}
-
-					AllValuePropertyRows.Append(ValuePropertyRows);
 				}
 			}
+
+			AllValuePropertyRows.Append(ValuePropertyRows);
 		}
 	}
 
-	FDynamicMaterialEditorModule::GeneratorComponentPropertyRows(
-		InComponentEditorWidget, 
-		MaterialStageFunction, 
-		InOutPropertyRows, 
-		InOutProcessedObjects
-	);
+	FDMComponentPropertyRowGeneratorParams FunctionParams = InParams;
+	FunctionParams.Object = MaterialStageFunction;
+
+	FDynamicMaterialEditorModule::GeneratorComponentPropertyRows(FunctionParams);
 
-	InOutPropertyRows.Append(AllValuePropertyRows);
+	InParams.PropertyRows.Append(AllValuePropertyRows);
 }
 
 #undef LOCTEXT_NAMESPACE
diff --git a/Engine/Plugins/Experimental/DynamicMaterial/Source/DynamicMaterialEditor/Private/UI/PropertyGenerators/DMMaterialValueDynamicPropertyRowGenerator.cpp b/Engine/Plugins/Experimental/DynamicMaterial/Source/DynamicMaterialEditor/Private/UI/PropertyGenerators/DMMaterialValueDynamicPropertyRowGenerator.cpp
--- a/Engine/Plugins/Experimental/DynamicMaterial/Source/DynamicMaterialEditor/Private/UI/PropertyGenerators/DMMaterialValueDynamicPropertyRowGenerator.cpp
+++ b/Engine/Plugins/Experimental/DynamicMaterial/Source/DynamicMaterialEditor/Private/UI/PropertyGenerators/DMMaterialValueDynamicPropertyRowGenerator.cpp
@@ -7,8 +7,8 @@
 #include "Components/DMMaterialValueDynamic.h"
 #include "DMEDefs.h"
 #include "IDetailPropertyRow.h"
+#include "IDynamicMaterialEditorModule.h"
 #include "UI/Utils/DMWidgetLibrary.h"
-#include "UI/Widgets/Editor/SDMMaterialComponentEditor.h"
 
 const TSharedRef<FDMMaterialValueDynamicPropertyRowGenerator>& FDMMaterialValueDynamicPropertyRowGenerator::Get()
 {
@@ -16,20 +16,19 @@ const TSharedRef<FDMMaterialValueDynamicPropertyRowGenerator>& FDMMaterialValueD
 	return Generator;
 }
 
-void FDMMaterialValueDynamicPropertyRowGenerator::AddComponentProperties(const TSharedRef<SDMMaterialComponentEditor>& InComponentEditorWidget, UDMMaterialComponent* InComponent,
-	TArray<FDMPropertyHandle>& InOutPropertyRows, TSet<UDMMaterialComponent*>& InOutProcessedObjects)
+void FDMMaterialValueDynamicPropertyRowGenerator::AddComponentProperties(FDMComponentPropertyRowGeneratorParams& InParams)
 {
-	if (!IsValid(InComponent))
+	if (!IsValid(InParams.Object))
 	{
 		return;
 	}
 
-	if (InOutProcessedObjects.Contains(InComponent))
+	if (InParams.ProcessedObjects.Contains(InParams.Object))
 	{
 		return;
 	}
 
-	UDMMaterialValueDynamic* ValueDynamic = Cast<UDMMaterialValueDynamic>(InComponent);
+	UDMMaterialValueDynamic* ValueDynamic = Cast<UDMMaterialValueDynamic>(InParams.Object);
 
 	if (!ValueDynamic)
 	{
@@ -49,11 +48,13 @@ void FDMMaterialValueDynamicPropertyRowGenerator::AddComponentProperties(const T
 		return;
 	}
 
-	InOutProcessedObjects.Add(InComponent);
+	InParams.ProcessedObjects.Add(InParams.Object);
 
 	if (ParentValue->AllowEditValue())
 	{
-		FDMPropertyHandle Handle = FDMWidgetLibrary::Get().GetPropertyHandle(&*InComponentEditorWidget, ValueDynamic, UDMMaterialValue::ValueName);
+		FDMPropertyHandleGenerateParams Params;
+		Params.Widget = InParams.Owner;
+		FDMPropertyHandle Handle = FDMWidgetLibrary::Get().GetPropertyHandle(InParams.CreatePropertyHandleParams(UDMMaterialValue::ValueName));
 
 		Handle.ResetToDefaultOverride = FResetToDefaultOverride::Create(
 			FIsResetToDefaultVisible::CreateUObject(ValueDynamic, &UDMMaterialValueDynamic::CanResetToDefault),
@@ -62,11 +63,11 @@ void FDMMaterialValueDynamicPropertyRowGenerator::AddComponentProperties(const T
 
 		Handle.bEnabled = true;
 
-		InOutPropertyRows.Add(Handle);
+		InParams.PropertyRows.Add(Handle);
 	}
 
 	const TArray<FName>& Properties = ParentValue->GetEditableProperties();
-	const int32 StartRow = InOutPropertyRows.Num();
+	const int32 StartRow = InParams.PropertyRows.Num();
 
 	for (const FName& Property : Properties)
 	{
@@ -75,16 +76,16 @@ void FDMMaterialValueDynamicPropertyRowGenerator::AddComponentProperties(const T
 			continue;
 		}
 			
-		if (InComponent->IsPropertyVisible(Property))
+		if (ValueDynamic->IsPropertyVisible(Property))
 		{
-			AddPropertyEditRows(InComponentEditorWidget, InComponent, Property, InOutPropertyRows, InOutProcessedObjects);
+			AddPropertyEditRows(InParams, Property);
 		}
 	}
 
-	const int32 EndRow = InOutPropertyRows.Num();
+	const int32 EndRow = InParams.PropertyRows.Num();
 
 	for (int32 RowIndex = StartRow; RowIndex < EndRow; ++RowIndex)
 	{
-		InOutPropertyRows[RowIndex].bEnabled = false;
+		InParams.PropertyRows[RowIndex].bEnabled = false;
 	}
 }
diff --git a/Engine/Plugins/Experimental/DynamicMaterial/Source/DynamicMaterialEditor/Private/UI/PropertyGenerators/DMMaterialValuePropertyRowGenerator.cpp b/Engine/Plugins/Experimental/DynamicMaterial/Source/DynamicMaterialEditor/Private/UI/PropertyGenerators/DMMaterialValuePropertyRowGenerator.cpp
--- a/Engine/Plugins/Experimental/DynamicMaterial/Source/DynamicMaterialEditor/Private/UI/PropertyGenerators/DMMaterialValuePropertyRowGenerator.cpp
+++ b/Engine/Plugins/Experimental/DynamicMaterial/Source/DynamicMaterialEditor/Private/UI/PropertyGenerators/DMMaterialValuePropertyRowGenerator.cpp
@@ -19,20 +19,19 @@ const TSharedRef<FDMMaterialValuePropertyRowGenerator>& FDMMaterialValueProperty
 	return Generator;
 }
 
-void FDMMaterialValuePropertyRowGenerator::AddComponentProperties(const TSharedRef<SDMMaterialComponentEditor>& InComponentEditorWidget, UDMMaterialComponent* InComponent,
-	TArray<FDMPropertyHandle>& InOutPropertyRows, TSet<UDMMaterialComponent*>& InOutProcessedObjects)
+void FDMMaterialValuePropertyRowGenerator::AddComponentProperties(FDMComponentPropertyRowGeneratorParams& InParams)
 {
-	if (!IsValid(InComponent))
+	if (!IsValid(InParams.Object))
 	{
 		return;
 	}
 
-	if (InOutProcessedObjects.Contains(InComponent))
+	if (InParams.ProcessedObjects.Contains(InParams.Object))
 	{
 		return;
 	}
 
-	UDMMaterialValue* Value = Cast<UDMMaterialValue>(InComponent);
+	UDMMaterialValue* Value = Cast<UDMMaterialValue>(InParams.Object);
 
 	if (!Value)
 	{
@@ -45,17 +44,20 @@ void FDMMaterialValuePropertyRowGenerator::AddComponentProperties(const TSharedR
 		return;
 	}
 
-	InOutProcessedObjects.Add(InComponent);
+	InParams.ProcessedObjects.Add(InParams.Object);
 
-	if (TSharedPtr<SDMMaterialEditor> EditorWidget = InComponentEditorWidget->GetEditorWidget())
+	if (TSharedPtr<SDMMaterialEditor> EditorWidget = static_cast<const SDMMaterialComponentEditor*>(InParams.Owner)->GetEditorWidget())
 	{
-		if (UDynamicMaterialModelBase* MaterialModelBase = EditorWidget->GetMaterialModelBase())
+		if (UDynamicMaterialModelBase* PreviewMaterialModelBase = EditorWidget->GetPreviewMaterialModelBase())
 		{
-			if (UDynamicMaterialModelDynamic* MaterialModelDynamic = Cast<UDynamicMaterialModelDynamic>(MaterialModelBase))
+			if (UDynamicMaterialModelDynamic* MaterialModelDynamic = Cast<UDynamicMaterialModelDynamic>(PreviewMaterialModelBase))
 			{
 				if (UDMMaterialComponentDynamic* ComponentDynamic = MaterialModelDynamic->GetComponentDynamic(Value->GetFName()))
 				{
-					FDynamicMaterialEditorModule::Get().GeneratorComponentPropertyRows(InComponentEditorWidget, ComponentDynamic, InOutPropertyRows, InOutProcessedObjects);
+					FDMComponentPropertyRowGeneratorParams DynamicComponentParams = InParams;
+					DynamicComponentParams.Object = ComponentDynamic;
+
+					FDynamicMaterialEditorModule::Get().GeneratorComponentPropertyRows(DynamicComponentParams);
 				}
 
 				return;
@@ -65,7 +67,7 @@ void FDMMaterialValuePropertyRowGenerator::AddComponentProperties(const TSharedR
 
 	if (Value->AllowEditValue())
 	{
-		FDMPropertyHandle Handle = FDMWidgetLibrary::Get().GetPropertyHandle(&*InComponentEditorWidget, Value, UDMMaterialValue::ValueName);
+		FDMPropertyHandle Handle = FDMWidgetLibrary::Get().GetPropertyHandle(InParams.CreatePropertyHandleParams(UDMMaterialValue::ValueName));
 
 		Handle.ResetToDefaultOverride = FResetToDefaultOverride::Create(
 			FIsResetToDefaultVisible::CreateUObject(Value, &UDMMaterialValue::CanResetToDefault),
@@ -74,7 +76,7 @@ void FDMMaterialValuePropertyRowGenerator::AddComponentProperties(const TSharedR
 
 		Handle.bEnabled = true;
 
-		InOutPropertyRows.Add(Handle);
+		InParams.PropertyRows.Add(Handle);
 	}
 
 	const TArray<FName>& Properties = Value->GetEditableProperties();
@@ -86,9 +88,9 @@ void FDMMaterialValuePropertyRowGenerator::AddComponentProperties(const TSharedR
 			continue;
 		}
 			
-		if (InComponent->IsPropertyVisible(Property))
+		if (Value->IsPropertyVisible(Property))
 		{
-			AddPropertyEditRows(InComponentEditorWidget, InComponent, Property, InOutPropertyRows, InOutProcessedObjects);
+			AddPropertyEditRows(InParams, Property);
 		}
 	}
 }
diff --git a/Engine/Plugins/Experimental/DynamicMaterial/Source/DynamicMaterialEditor/Private/UI/PropertyGenerators/DMStagePropertyRowGenerator.cpp b/Engine/Plugins/Experimental/DynamicMaterial/Source/DynamicMaterialEditor/Private/UI/PropertyGenerators/DMStagePropertyRowGenerator.cpp
--- a/Engine/Plugins/Experimental/DynamicMaterial/Source/DynamicMaterialEditor/Private/UI/PropertyGenerators/DMStagePropertyRowGenerator.cpp
+++ b/Engine/Plugins/Experimental/DynamicMaterial/Source/DynamicMaterialEditor/Private/UI/PropertyGenerators/DMStagePropertyRowGenerator.cpp
@@ -2,24 +2,10 @@
 
 #include "UI/PropertyGenerators/DMStagePropertyRowGenerator.h"
 
-#include "Components/DMMaterialComponent.h"
-#include "Components/DMMaterialLayer.h"
 #include "Components/DMMaterialStage.h"
 #include "Components/DMMaterialStageSource.h"
-#include "DetailLayoutBuilder.h"
-#include "DMEDefs.h"
 #include "DynamicMaterialEditorModule.h"
 #include "Internationalization/Text.h"
-#include "Internationalization/Text.h"
-#include "Model/DynamicMaterialModel.h"
-#include "Model/DynamicMaterialModelBase.h"
-#include "UI/Menus/DMMaterialStageSourceMenus.h"
-#include "UI/Widgets/Editor/SDMMaterialComponentEditor.h"
-#include "UI/Widgets/Editor/SDMMaterialSlotEditor.h"
-#include "UI/Widgets/Editor/SlotEditor/SDMMaterialSlotLayerItem.h"
-#include "UI/Widgets/Editor/SlotEditor/SDMMaterialSlotLayerView.h"
-#include "UI/Widgets/Editor/SlotEditor/SDMMaterialStage.h"
-#include "UI/Widgets/SDMMaterialEditor.h"
 #include "Widgets/Input/SComboButton.h"
 
 #define LOCTEXT_NAMESPACE "DMStagePropertyRowGenerator"
@@ -30,20 +16,19 @@ const TSharedRef<FDMStagePropertyRowGenerator>& FDMStagePropertyRowGenerator::Ge
 	return Generator;
 }
 
-void FDMStagePropertyRowGenerator::AddComponentProperties(const TSharedRef<SDMMaterialComponentEditor>& InComponentEditorWidget, UDMMaterialComponent* InComponent,
-	TArray<FDMPropertyHandle>& InOutPropertyRows, TSet<UDMMaterialComponent*>& InOutProcessedObjects)
+void FDMStagePropertyRowGenerator::AddComponentProperties(FDMComponentPropertyRowGeneratorParams& InParams)
 {
-	if (!IsValid(InComponent))
+	if (!IsValid(InParams.Object))
 	{
 		return;
 	}
 
-	if (InOutProcessedObjects.Contains(InComponent))
+	if (InParams.ProcessedObjects.Contains(InParams.Object))
 	{
 		return;
 	}
 
-	UDMMaterialStage* Stage = Cast<UDMMaterialStage>(InComponent);
+	UDMMaterialStage* Stage = Cast<UDMMaterialStage>(InParams.Object);
 
 	if (!Stage)
 	{
@@ -57,8 +42,15 @@ void FDMStagePropertyRowGenerator::AddComponentProperties(const TSharedRef<SDMMa
 		return;
 	}
 
-	FDynamicMaterialEditorModule::GeneratorComponentPropertyRows(InComponentEditorWidget, Source, InOutPropertyRows, InOutProcessedObjects);
-	FDMComponentPropertyRowGenerator::AddComponentProperties(InComponentEditorWidget, Stage, InOutPropertyRows, InOutProcessedObjects);
+	FDMComponentPropertyRowGeneratorParams SourceParams = InParams;
+	SourceParams.Object = Source;
+
+	FDynamicMaterialEditorModule::GeneratorComponentPropertyRows(SourceParams);
+
+	FDMComponentPropertyRowGeneratorParams StageParams = InParams;
+	StageParams.Object = Stage;
+
+	FDMComponentPropertyRowGenerator::AddComponentProperties(StageParams);
 }
 
 #undef LOCTEXT_NAMESPACE
diff --git a/Engine/Plugins/Experimental/DynamicMaterial/Source/DynamicMaterialEditor/Private/UI/PropertyGenerators/DMTextureUVDynamicPropertyRowGenerator.cpp b/Engine/Plugins/Experimental/DynamicMaterial/Source/DynamicMaterialEditor/Private/UI/PropertyGenerators/DMTextureUVDynamicPropertyRowGenerator.cpp
--- a/Engine/Plugins/Experimental/DynamicMaterial/Source/DynamicMaterialEditor/Private/UI/PropertyGenerators/DMTextureUVDynamicPropertyRowGenerator.cpp
+++ b/Engine/Plugins/Experimental/DynamicMaterial/Source/DynamicMaterialEditor/Private/UI/PropertyGenerators/DMTextureUVDynamicPropertyRowGenerator.cpp
@@ -7,6 +7,7 @@
 #include "Components/DMTextureUV.h"
 #include "Components/DMTextureUVDynamic.h"
 #include "DMEDefs.h"
+#include "IDynamicMaterialEditorModule.h"
 #include "Modules/ModuleManager.h"
 #include "PropertyCustomizationHelpers.h"
 #include "PropertyEditorModule.h"
@@ -26,31 +27,29 @@ const TSharedRef<FDMTextureUVDynamicPropertyRowGenerator>& FDMTextureUVDynamicPr
 
 namespace UE::DynamicMaterialEditor::Private
 {
-	void AddTextureUVDynamicPropertyRow(const TSharedRef<SWidget>& InParentWidget, UDMMaterialComponent* InTextureUVDynamic,
-		FName InProperty, bool bInEnabled, TArray<FDMPropertyHandle>& InOutPropertyRows);
+	void AddTextureUVDynamicPropertyRow(FDMComponentPropertyRowGeneratorParams& InParams, UDMMaterialComponent* InComponent, 
+		FName InProperty, bool bInEnabled);
 
-	void AddTextureUVDynamicVisualizerRow(const TSharedRef<SWidget>& InParentWidget, UDMTextureUVDynamic* InTextureUVDynamic,
-		TArray<FDMPropertyHandle>& InOutPropertyRows);
+	void AddTextureUVDynamicVisualizerRow(FDMComponentPropertyRowGeneratorParams& InParams, UDMTextureUVDynamic* InTextureUVDynamic);
 
 	bool CanResetTextureUVDynamicDynamicPropertyToDefault(TSharedPtr<IPropertyHandle> InPropertyHandle);
 
 	void ResetTextureUVDynamicDynamicPropertyToDefault(TSharedPtr<IPropertyHandle> InPropertyHandle);
 }
 
-void FDMTextureUVDynamicPropertyRowGenerator::AddComponentProperties(const TSharedRef<SDMMaterialComponentEditor>& InComponentEditorWidget, 
-	UDMMaterialComponent* InComponent, TArray<FDMPropertyHandle>& InOutPropertyRows, TSet<UDMMaterialComponent*>& InOutProcessedObjects)
+void FDMTextureUVDynamicPropertyRowGenerator::AddComponentProperties(FDMComponentPropertyRowGeneratorParams& InParams)
 {
-	if (!IsValid(InComponent))
+	if (!IsValid(InParams.Object))
 	{
 		return;
 	}
 
-	if (InOutProcessedObjects.Contains(InComponent))
+	if (InParams.ProcessedObjects.Contains(InParams.Object))
 	{
 		return;
 	}
 
-	UDMTextureUVDynamic* TextureUVDynamic = Cast<UDMTextureUVDynamic>(InComponent);
+	UDMTextureUVDynamic* TextureUVDynamic = Cast<UDMTextureUVDynamic>(InParams.Object);
 
 	if (!TextureUVDynamic)
 	{
@@ -64,28 +63,27 @@ void FDMTextureUVDynamicPropertyRowGenerator::AddComponentProperties(const TShar
 		return;
 	}
 
-	InOutProcessedObjects.Add(InComponent);
+	InParams.ProcessedObjects.Add(InParams.Object);
 
 	using namespace UE::DynamicMaterialEditor::Private;
 
-	AddTextureUVDynamicPropertyRow(InComponentEditorWidget, TextureUVDynamic, UDMTextureUV::NAME_Offset, /* Enabled */ true, InOutPropertyRows);
-	AddTextureUVDynamicPropertyRow(InComponentEditorWidget, TextureUVDynamic, UDMTextureUV::NAME_Rotation, /* Enabled */ true, InOutPropertyRows);
-	AddTextureUVDynamicPropertyRow(InComponentEditorWidget, TextureUVDynamic, UDMTextureUV::NAME_Tiling, /* Enabled */ true, InOutPropertyRows);
-	AddTextureUVDynamicPropertyRow(InComponentEditorWidget, TextureUVDynamic, UDMTextureUV::NAME_Pivot, /* Enabled */ true, InOutPropertyRows);
-	AddTextureUVDynamicPropertyRow(InComponentEditorWidget, TextureUV, UDMTextureUV::NAME_bMirrorOnX, /* Enabled */ false, InOutPropertyRows);
-	AddTextureUVDynamicPropertyRow(InComponentEditorWidget, TextureUV, UDMTextureUV::NAME_bMirrorOnY, /* Enabled */ false, InOutPropertyRows);
-	AddTextureUVDynamicVisualizerRow(InComponentEditorWidget, TextureUVDynamic, InOutPropertyRows);
+	AddTextureUVDynamicPropertyRow(InParams, TextureUVDynamic, UDMTextureUV::NAME_Offset, /* Enabled */ true);
+	AddTextureUVDynamicPropertyRow(InParams, TextureUVDynamic, UDMTextureUV::NAME_Rotation, /* Enabled */ true);
+	AddTextureUVDynamicPropertyRow(InParams, TextureUVDynamic, UDMTextureUV::NAME_Tiling, /* Enabled */ true);
+	AddTextureUVDynamicPropertyRow(InParams, TextureUVDynamic, UDMTextureUV::NAME_Pivot, /* Enabled */ true);
+	AddTextureUVDynamicPropertyRow(InParams, TextureUV, UDMTextureUV::NAME_bMirrorOnX, /* Enabled */ false);
+	AddTextureUVDynamicPropertyRow(InParams, TextureUV, UDMTextureUV::NAME_bMirrorOnY, /* Enabled */ false);
+	AddTextureUVDynamicVisualizerRow(InParams, TextureUVDynamic);
 }
 
-void FDMTextureUVDynamicPropertyRowGenerator::AddPopoutComponentProperties(const TSharedRef<SWidget>& InParentWidget, UDMMaterialComponent* InComponent, 
-	TArray<FDMPropertyHandle>& InOutPropertyRows)
+void FDMTextureUVDynamicPropertyRowGenerator::AddPopoutComponentProperties(FDMComponentPropertyRowGeneratorParams& InParams)
 {
-	if (!IsValid(InComponent))
+	if (!IsValid(InParams.Object))
 	{
 		return;
 	}
 
-	UDMTextureUVDynamic* TextureUVDynamic = Cast<UDMTextureUVDynamic>(InComponent);
+	UDMTextureUVDynamic* TextureUVDynamic = Cast<UDMTextureUVDynamic>(InParams.Object);
 
 	if (!TextureUVDynamic)
 	{
@@ -101,12 +99,12 @@ void FDMTextureUVDynamicPropertyRowGenerator::AddPopoutComponentProperties(const
 
 	using namespace UE::DynamicMaterialEditor::Private;
 
-	AddTextureUVDynamicPropertyRow(InParentWidget, TextureUVDynamic, UDMTextureUV::NAME_Offset, /* Enabled */ true, InOutPropertyRows);
-	AddTextureUVDynamicPropertyRow(InParentWidget, TextureUVDynamic, UDMTextureUV::NAME_Rotation, /* Enabled */ true, InOutPropertyRows);
-	AddTextureUVDynamicPropertyRow(InParentWidget, TextureUVDynamic, UDMTextureUV::NAME_Tiling, /* Enabled */ true, InOutPropertyRows);
-	AddTextureUVDynamicPropertyRow(InParentWidget, TextureUVDynamic, UDMTextureUV::NAME_Pivot, /* Enabled */ true, InOutPropertyRows);
-	AddTextureUVDynamicPropertyRow(InParentWidget, TextureUV, UDMTextureUV::NAME_bMirrorOnX, /* Enabled */ false, InOutPropertyRows);
-	AddTextureUVDynamicPropertyRow(InParentWidget, TextureUV, UDMTextureUV::NAME_bMirrorOnY, /* Enabled */ false, InOutPropertyRows);
+	AddTextureUVDynamicPropertyRow(InParams, TextureUVDynamic, UDMTextureUV::NAME_Offset, /* Enabled */ true);
+	AddTextureUVDynamicPropertyRow(InParams, TextureUVDynamic, UDMTextureUV::NAME_Rotation, /* Enabled */ true);
+	AddTextureUVDynamicPropertyRow(InParams, TextureUVDynamic, UDMTextureUV::NAME_Tiling, /* Enabled */ true);
+	AddTextureUVDynamicPropertyRow(InParams, TextureUVDynamic, UDMTextureUV::NAME_Pivot, /* Enabled */ true);
+	AddTextureUVDynamicPropertyRow(InParams, TextureUV, UDMTextureUV::NAME_bMirrorOnX, /* Enabled */ false);
+	AddTextureUVDynamicPropertyRow(InParams, TextureUV, UDMTextureUV::NAME_bMirrorOnY, /* Enabled */ false);
 }
 
 bool FDMTextureUVDynamicPropertyRowGenerator::AllowKeyframeButton(UDMMaterialComponent* InComponent, FProperty* InProperty)
@@ -124,10 +122,13 @@ bool FDMTextureUVDynamicPropertyRowGenerator::AllowKeyframeButton(UDMMaterialCom
 	return FDMComponentPropertyRowGenerator::AllowKeyframeButton(InComponent, InProperty);
 }
 
-void UE::DynamicMaterialEditor::Private::AddTextureUVDynamicPropertyRow(const TSharedRef<SWidget>& InParentWidget, 
-	UDMMaterialComponent* InTextureUVDynamic, FName InProperty, bool bInEnabled, TArray<FDMPropertyHandle>& InOutPropertyRows)
+void UE::DynamicMaterialEditor::Private::AddTextureUVDynamicPropertyRow(FDMComponentPropertyRowGeneratorParams& InParams,
+	UDMMaterialComponent* InComponent, FName InProperty, bool bInEnabled)
 {
-	FDMPropertyHandle& NewHandle = InOutPropertyRows.Add_GetRef(FDMWidgetLibrary::Get().GetPropertyHandle(&*InParentWidget, InTextureUVDynamic, InProperty));
+	FDMComponentPropertyRowGeneratorParams ComponentParams = InParams;
+	ComponentParams.Object = InComponent;
+
+	FDMPropertyHandle& NewHandle = InParams.PropertyRows.Add_GetRef(FDMWidgetLibrary::Get().GetPropertyHandle(ComponentParams.CreatePropertyHandleParams(InProperty)));
 
 	NewHandle.ResetToDefaultOverride = FResetToDefaultOverride::Create(
 		FIsResetToDefaultVisible::CreateStatic(&UE::DynamicMaterialEditor::Private::CanResetTextureUVDynamicDynamicPropertyToDefault),
@@ -137,8 +138,8 @@ void UE::DynamicMaterialEditor::Private::AddTextureUVDynamicPropertyRow(const TS
 	NewHandle.bEnabled = bInEnabled;
 }
 
-void UE::DynamicMaterialEditor::Private::AddTextureUVDynamicVisualizerRow(const TSharedRef<SWidget>& InParentWidget, UDMTextureUVDynamic* InTextureUVDynamic,
-	TArray<FDMPropertyHandle>& InOutPropertyRows)
+void UE::DynamicMaterialEditor::Private::AddTextureUVDynamicVisualizerRow(FDMComponentPropertyRowGeneratorParams& InParams, 
+	UDMTextureUVDynamic* InTextureUVDynamic)
 {
 	UDMTextureUV* TextureUV = InTextureUVDynamic->GetParentTextureUV();
 
@@ -155,12 +156,7 @@ void UE::DynamicMaterialEditor::Private::AddTextureUVDynamicVisualizerRow(const
 		return;
 	}
 
-	if (InParentWidget->GetWidgetClass().GetWidgetType() != SDMMaterialComponentEditor::StaticWidgetClass().GetWidgetType())
-	{
-		return;
-	}
-
-	TSharedPtr<SDMMaterialEditor> EditorWidget = StaticCastSharedRef<SDMMaterialComponentEditor>(InParentWidget)->GetEditorWidget();
+	TSharedPtr<SDMMaterialEditor> EditorWidget = static_cast<const SDMMaterialComponentEditor*>(InParams.Owner)->GetEditorWidget();
 
 	if (!EditorWidget.IsValid())
 	{
@@ -174,7 +170,7 @@ void UE::DynamicMaterialEditor::Private::AddTextureUVDynamicVisualizerRow(const
 	VisualizerHandle.ValueWidget = SNew(SDMTextureUVVisualizerProperty, EditorWidget.ToSharedRef(), Stage).TextureUVDynamic(InTextureUVDynamic);
 	VisualizerHandle.CategoryOverrideName = TEXT("Texture UV");
 	VisualizerHandle.bEnabled = true;
-	InOutPropertyRows.Add(VisualizerHandle);
+	InParams.PropertyRows.Add(VisualizerHandle);
 }
 
 
@@ -226,7 +222,7 @@ bool UE::DynamicMaterialEditor::Private::CanResetTextureUVDynamicDynamicProperty
 		return !DefaultObject->GetPivot().Equals(PropertyObject->GetPivot());
 	}
 
-	if (PropertyName == NAME_Rotation)
+	if (PropertyName == UDMTextureUV::NAME_Rotation)
 	{
 		return DefaultObject->GetRotation() != PropertyObject->GetRotation();
 	}
diff --git a/Engine/Plugins/Experimental/DynamicMaterial/Source/DynamicMaterialEditor/Private/UI/PropertyGenerators/DMTextureUVDynamicPropertyRowGenerator.h b/Engine/Plugins/Experimental/DynamicMaterial/Source/DynamicMaterialEditor/Private/UI/PropertyGenerators/DMTextureUVDynamicPropertyRowGenerator.h
--- a/Engine/Plugins/Experimental/DynamicMaterial/Source/DynamicMaterialEditor/Private/UI/PropertyGenerators/DMTextureUVDynamicPropertyRowGenerator.h
+++ b/Engine/Plugins/Experimental/DynamicMaterial/Source/DynamicMaterialEditor/Private/UI/PropertyGenerators/DMTextureUVDynamicPropertyRowGenerator.h
@@ -4,8 +4,6 @@
 
 #include "UI/PropertyGenerators/DMComponentPropertyRowGenerator.h"
 
-class SWidget;
-
 class FDMTextureUVDynamicPropertyRowGenerator : public FDMComponentPropertyRowGenerator
 {
 public:
@@ -13,12 +11,10 @@ public:
 
 	virtual ~FDMTextureUVDynamicPropertyRowGenerator() override = default;
 
-	static void AddPopoutComponentProperties(const TSharedRef<SWidget>& InParentWidget, UDMMaterialComponent* InComponent,
-		TArray<FDMPropertyHandle>& InOutPropertyRows);
+	static void AddPopoutComponentProperties(FDMComponentPropertyRowGeneratorParams& InParams);
 
 	//~ Begin FDMComponentPropertyRowGenerator
-	virtual void AddComponentProperties(const TSharedRef<SDMMaterialComponentEditor>& InComponentEditorWidget, UDMMaterialComponent* InComponent,
-		TArray<FDMPropertyHandle>& InOutPropertyRows, TSet<UDMMaterialComponent*>& InOutProcessedObjects) override;
+	virtual void AddComponentProperties(FDMComponentPropertyRowGeneratorParams& InParams) override;
 
 	virtual bool AllowKeyframeButton(UDMMaterialComponent* InComponent, FProperty* InProperty) override;
 	//~ End FDMComponentPropertyRowGenerator
diff --git a/Engine/Plugins/Experimental/DynamicMaterial/Source/DynamicMaterialEditor/Private/UI/PropertyGenerators/DMTextureUVPropertyRowGenerator.cpp b/Engine/Plugins/Experimental/DynamicMaterial/Source/DynamicMaterialEditor/Private/UI/PropertyGenerators/DMTextureUVPropertyRowGenerator.cpp
--- a/Engine/Plugins/Experimental/DynamicMaterial/Source/DynamicMaterialEditor/Private/UI/PropertyGenerators/DMTextureUVPropertyRowGenerator.cpp
+++ b/Engine/Plugins/Experimental/DynamicMaterial/Source/DynamicMaterialEditor/Private/UI/PropertyGenerators/DMTextureUVPropertyRowGenerator.cpp
@@ -27,48 +27,48 @@ const TSharedRef<FDMTextureUVPropertyRowGenerator>& FDMTextureUVPropertyRowGener
 
 namespace UE::DynamicMaterialEditor::Private
 {
-	void AddTextureUVPropertyRow(const TSharedRef<SWidget>& InParentWidget, UDMTextureUV* InTextureUV, FName InProperty,
-		TArray<FDMPropertyHandle>& InOutPropertyRows);
+	void AddTextureUVPropertyRow(FDMComponentPropertyRowGeneratorParams& InParams, FName InProperty);
 
-	void AddTextureUVVisualizerRow(const TSharedRef<SWidget>& InParentWidget, UDMTextureUV* InTextureUV,
-		TArray<FDMPropertyHandle>& InOutPropertyRows);
+	void AddTextureUVVisualizerRow(FDMComponentPropertyRowGeneratorParams& InParams);
 
 	bool CanResetTextureUVPropertyToDefault(TSharedPtr<IPropertyHandle> InPropertyHandle);
 
 	void ResetTextureUVPropertyToDefault(TSharedPtr<IPropertyHandle> InPropertyHandle);
 }
 
-void FDMTextureUVPropertyRowGenerator::AddComponentProperties(const TSharedRef<SDMMaterialComponentEditor>& InComponentEditorWidget, 
-	UDMMaterialComponent* InComponent, TArray<FDMPropertyHandle>& InOutPropertyRows, TSet<UDMMaterialComponent*>& InOutProcessedObjects)
+void FDMTextureUVPropertyRowGenerator::AddComponentProperties(FDMComponentPropertyRowGeneratorParams& InParams)
 {
-	if (!IsValid(InComponent))
+	if (!IsValid(InParams.Object))
 	{
 		return;
 	}
 
-	if (InOutProcessedObjects.Contains(InComponent))
+	if (InParams.ProcessedObjects.Contains(InParams.Object))
 	{
 		return;
 	}
 
-	UDMTextureUV* TextureUV = Cast<UDMTextureUV>(InComponent);
+	UDMTextureUV* TextureUV = Cast<UDMTextureUV>(InParams.Object);
 
 	if (!TextureUV)
 	{
 		return;
 	}
 
-	InOutProcessedObjects.Add(InComponent);
+	InParams.ProcessedObjects.Add(InParams.Object);
 
-	if (TSharedPtr<SDMMaterialEditor> EditorWidget = InComponentEditorWidget->GetEditorWidget())
+	if (TSharedPtr<SDMMaterialEditor> EditorWidget = static_cast<const SDMMaterialComponentEditor*>(InParams.Owner)->GetEditorWidget())
 	{
-		if (UDynamicMaterialModelBase* MaterialModelBase = EditorWidget->GetMaterialModelBase())
+		if (UDynamicMaterialModelBase* PreviewMaterialModelBase = EditorWidget->GetPreviewMaterialModelBase())
 		{
-			if (UDynamicMaterialModelDynamic* MaterialModelDynamic = Cast<UDynamicMaterialModelDynamic>(MaterialModelBase))
+			if (UDynamicMaterialModelDynamic* MaterialModelDynamic = Cast<UDynamicMaterialModelDynamic>(PreviewMaterialModelBase))
 			{
 				if (UDMMaterialComponentDynamic* ComponentDynamic = MaterialModelDynamic->GetComponentDynamic(TextureUV->GetFName()))
 				{
-					FDynamicMaterialEditorModule::Get().GeneratorComponentPropertyRows(InComponentEditorWidget, ComponentDynamic, InOutPropertyRows, InOutProcessedObjects);
+					FDMComponentPropertyRowGeneratorParams DynamicComponentParams = InParams;
+					DynamicComponentParams.Object = ComponentDynamic;
+
+					FDynamicMaterialEditorModule::Get().GeneratorComponentPropertyRows(DynamicComponentParams);
 				}
 
 				return;
@@ -78,24 +78,23 @@ void FDMTextureUVPropertyRowGenerator::AddComponentProperties(const TSharedRef<S
 
 	using namespace UE::DynamicMaterialEditor::Private;
 
-	AddTextureUVPropertyRow(InComponentEditorWidget, TextureUV, UDMTextureUV::NAME_Offset, InOutPropertyRows);
-	AddTextureUVPropertyRow(InComponentEditorWidget, TextureUV, UDMTextureUV::NAME_Rotation, InOutPropertyRows);
-	AddTextureUVPropertyRow(InComponentEditorWidget, TextureUV, UDMTextureUV::NAME_Tiling, InOutPropertyRows);
-	AddTextureUVPropertyRow(InComponentEditorWidget, TextureUV, UDMTextureUV::NAME_Pivot, InOutPropertyRows);
-	AddTextureUVPropertyRow(InComponentEditorWidget, TextureUV, UDMTextureUV::NAME_bMirrorOnX, InOutPropertyRows);
-	AddTextureUVPropertyRow(InComponentEditorWidget, TextureUV, UDMTextureUV::NAME_bMirrorOnY, InOutPropertyRows);
-	AddTextureUVVisualizerRow(InComponentEditorWidget, TextureUV, InOutPropertyRows);
+	AddTextureUVPropertyRow(InParams, UDMTextureUV::NAME_Offset);
+	AddTextureUVPropertyRow(InParams, UDMTextureUV::NAME_Rotation);
+	AddTextureUVPropertyRow(InParams, UDMTextureUV::NAME_Tiling);
+	AddTextureUVPropertyRow(InParams, UDMTextureUV::NAME_Pivot);
+	AddTextureUVPropertyRow(InParams, UDMTextureUV::NAME_bMirrorOnX);
+	AddTextureUVPropertyRow(InParams, UDMTextureUV::NAME_bMirrorOnY);
+	AddTextureUVVisualizerRow(InParams);
 }
 
-void FDMTextureUVPropertyRowGenerator::AddPopoutComponentProperties(const TSharedRef<SWidget>& InParentWidget, UDMMaterialComponent* InComponent, 
-	TArray<FDMPropertyHandle>& InOutPropertyRows)
+void FDMTextureUVPropertyRowGenerator::AddPopoutComponentProperties(FDMComponentPropertyRowGeneratorParams& InParams)
 {
-	if (!IsValid(InComponent))
+	if (!IsValid(InParams.Object))
 	{
 		return;
 	}
 
-	UDMTextureUV* TextureUV = Cast<UDMTextureUV>(InComponent);
+	UDMTextureUV* TextureUV = Cast<UDMTextureUV>(InParams.Object);
 
 	if (!TextureUV)
 	{
@@ -104,12 +103,12 @@ void FDMTextureUVPropertyRowGenerator::AddPopoutComponentProperties(const TShare
 
 	using namespace UE::DynamicMaterialEditor::Private;
 
-	AddTextureUVPropertyRow(InParentWidget, TextureUV, UDMTextureUV::NAME_Offset, InOutPropertyRows);
-	AddTextureUVPropertyRow(InParentWidget, TextureUV, UDMTextureUV::NAME_Rotation, InOutPropertyRows);
-	AddTextureUVPropertyRow(InParentWidget, TextureUV, UDMTextureUV::NAME_Tiling, InOutPropertyRows);
-	AddTextureUVPropertyRow(InParentWidget, TextureUV, UDMTextureUV::NAME_Pivot, InOutPropertyRows);
-	AddTextureUVPropertyRow(InParentWidget, TextureUV, UDMTextureUV::NAME_bMirrorOnX, InOutPropertyRows);
-	AddTextureUVPropertyRow(InParentWidget, TextureUV, UDMTextureUV::NAME_bMirrorOnY, InOutPropertyRows);
+	AddTextureUVPropertyRow(InParams, UDMTextureUV::NAME_Offset);
+	AddTextureUVPropertyRow(InParams, UDMTextureUV::NAME_Rotation);
+	AddTextureUVPropertyRow(InParams, UDMTextureUV::NAME_Tiling);
+	AddTextureUVPropertyRow(InParams, UDMTextureUV::NAME_Pivot);
+	AddTextureUVPropertyRow(InParams, UDMTextureUV::NAME_bMirrorOnX);
+	AddTextureUVPropertyRow(InParams, UDMTextureUV::NAME_bMirrorOnY);
 }
 
 bool FDMTextureUVPropertyRowGenerator::AllowKeyframeButton(UDMMaterialComponent* InComponent, FProperty* InProperty)
@@ -127,10 +126,9 @@ bool FDMTextureUVPropertyRowGenerator::AllowKeyframeButton(UDMMaterialComponent*
 	return FDMComponentPropertyRowGenerator::AllowKeyframeButton(InComponent, InProperty);
 }
 
-void UE::DynamicMaterialEditor::Private::AddTextureUVPropertyRow(const TSharedRef<SWidget>& InParentWidget, UDMTextureUV* InTextureUV,
-	FName InProperty, TArray<FDMPropertyHandle>& InOutPropertyRows)
+void UE::DynamicMaterialEditor::Private::AddTextureUVPropertyRow(FDMComponentPropertyRowGeneratorParams& InParams, FName InProperty)
 {
-	FDMPropertyHandle& NewHandle = InOutPropertyRows.Add_GetRef(FDMWidgetLibrary::Get().GetPropertyHandle(&*InParentWidget, InTextureUV, InProperty));
+	FDMPropertyHandle& NewHandle = InParams.PropertyRows.Add_GetRef(FDMWidgetLibrary::Get().GetPropertyHandle(InParams.CreatePropertyHandleParams(InProperty)));
 
 	NewHandle.ResetToDefaultOverride = FResetToDefaultOverride::Create(
 		FIsResetToDefaultVisible::CreateStatic(&UE::DynamicMaterialEditor::Private::CanResetTextureUVPropertyToDefault),
@@ -140,23 +138,19 @@ void UE::DynamicMaterialEditor::Private::AddTextureUVPropertyRow(const TSharedRe
 	NewHandle.bEnabled = true;
 }
 
-void UE::DynamicMaterialEditor::Private::AddTextureUVVisualizerRow(const TSharedRef<SWidget>& InParentWidget, UDMTextureUV* InTextureUV,
-	TArray<FDMPropertyHandle>& InOutPropertyRows)
+void UE::DynamicMaterialEditor::Private::AddTextureUVVisualizerRow(FDMComponentPropertyRowGeneratorParams& InParams)
 {
+	UDMTextureUV* TextureUV = CastChecked<UDMTextureUV>(InParams.Object);
+
 	// Make sure we don't get a substage
-	UDMMaterialStage* Stage = InTextureUV->GetTypedParent<UDMMaterialStage>(/* Allow Subclasses */ false);
+	UDMMaterialStage* Stage = TextureUV->GetTypedParent<UDMMaterialStage>(/* Allow Subclasses */ false);
 
 	if (!Stage)
 	{
 		return;
 	}
 
-	if (InParentWidget->GetWidgetClass().GetWidgetType() != SDMMaterialComponentEditor::StaticWidgetClass().GetWidgetType())
-	{
-		return;
-	}
-
-	TSharedPtr<SDMMaterialEditor> EditorWidget = StaticCastSharedRef<SDMMaterialComponentEditor>(InParentWidget)->GetEditorWidget();
+	TSharedPtr<SDMMaterialEditor> EditorWidget = static_cast<const SDMMaterialComponentEditor*>(InParams.Owner)->GetEditorWidget();
 
 	if (!EditorWidget.IsValid())
 	{
@@ -166,11 +160,11 @@ void UE::DynamicMaterialEditor::Private::AddTextureUVVisualizerRow(const TShared
 	FDMPropertyHandle VisualizerHandle;
 	VisualizerHandle.NameOverride = LOCTEXT("Visualizer", "UV Visualizer");
 	VisualizerHandle.NameToolTipOverride = LOCTEXT("VisualizerToolTip", "A graphical Texture UV editor.\n\n- Offset Mode: Change the Texture UV offset.\n- Pivot Mode: Change the Texture UV pivot, rotation and tiling.\n\nControl+click to reset values to default.");
-	VisualizerHandle.ValueName = FName(*InTextureUV->GetComponentPath());
-	VisualizerHandle.ValueWidget = SNew(SDMTextureUVVisualizerProperty, EditorWidget.ToSharedRef(), Stage).TextureUV(InTextureUV);
+	VisualizerHandle.ValueName = FName(*TextureUV->GetComponentPath());
+	VisualizerHandle.ValueWidget = SNew(SDMTextureUVVisualizerProperty, EditorWidget.ToSharedRef(), Stage).TextureUV(TextureUV);
 	VisualizerHandle.CategoryOverrideName = TEXT("Texture UV");
 	VisualizerHandle.bEnabled = true;
-	InOutPropertyRows.Add(VisualizerHandle);
+	InParams.PropertyRows.Add(VisualizerHandle);
 }
 
 
@@ -237,7 +231,7 @@ bool UE::DynamicMaterialEditor::Private::CanResetTextureUVPropertyToDefault(TSha
 		return !DefaultObject->GetPivot().Equals(PropertyObject->GetPivot());
 	}
 
-	if (PropertyName == NAME_Rotation)
+	if (PropertyName == UDMTextureUV::NAME_Rotation)
 	{
 		return DefaultObject->GetRotation() != PropertyObject->GetRotation();
 	}
diff --git a/Engine/Plugins/Experimental/DynamicMaterial/Source/DynamicMaterialEditor/Private/UI/PropertyGenerators/DMTextureUVPropertyRowGenerator.h b/Engine/Plugins/Experimental/DynamicMaterial/Source/DynamicMaterialEditor/Private/UI/PropertyGenerators/DMTextureUVPropertyRowGenerator.h
--- a/Engine/Plugins/Experimental/DynamicMaterial/Source/DynamicMaterialEditor/Private/UI/PropertyGenerators/DMTextureUVPropertyRowGenerator.h
+++ b/Engine/Plugins/Experimental/DynamicMaterial/Source/DynamicMaterialEditor/Private/UI/PropertyGenerators/DMTextureUVPropertyRowGenerator.h
@@ -13,12 +13,10 @@ public:
 
 	virtual ~FDMTextureUVPropertyRowGenerator() override = default;
 
-	static void AddPopoutComponentProperties(const TSharedRef<SWidget>& InParentWidget, UDMMaterialComponent* InComponent,
-		TArray<FDMPropertyHandle>& InOutPropertyRows);
+	static void AddPopoutComponentProperties(FDMComponentPropertyRowGeneratorParams& InParams);
 
 	//~ Begin FDMComponentPropertyRowGenerator
-	virtual void AddComponentProperties(const TSharedRef<SDMMaterialComponentEditor>& InComponentEditorWidget, UDMMaterialComponent* InComponent,
-		TArray<FDMPropertyHandle>& InOutPropertyRows, TSet<UDMMaterialComponent*>& InOutProcessedObjects) override;
+	virtual void AddComponentProperties(FDMComponentPropertyRowGeneratorParams& InParams) override;
 
 	virtual bool AllowKeyframeButton(UDMMaterialComponent* InComponent, FProperty* InProperty) override;
 	//~ End FDMComponentPropertyRowGenerator
diff --git a/Engine/Plugins/Experimental/DynamicMaterial/Source/DynamicMaterialEditor/Private/UI/PropertyGenerators/DMThroughputPropertyRowGenerator.cpp b/Engine/Plugins/Experimental/DynamicMaterial/Source/DynamicMaterialEditor/Private/UI/PropertyGenerators/DMThroughputPropertyRowGenerator.cpp
--- a/Engine/Plugins/Experimental/DynamicMaterial/Source/DynamicMaterialEditor/Private/UI/PropertyGenerators/DMThroughputPropertyRowGenerator.cpp
+++ b/Engine/Plugins/Experimental/DynamicMaterial/Source/DynamicMaterialEditor/Private/UI/PropertyGenerators/DMThroughputPropertyRowGenerator.cpp
@@ -14,36 +14,34 @@ const TSharedRef<FDMThroughputPropertyRowGenerator>& FDMThroughputPropertyRowGen
 	return Generator;
 }
 
-void FDMThroughputPropertyRowGenerator::AddComponentProperties(const TSharedRef<SDMMaterialComponentEditor>& InComponentEditorWidget,
-	UDMMaterialComponent* InComponent, TArray<FDMPropertyHandle>& InOutPropertyRows,
-	TSet<UDMMaterialComponent*>& InOutProcessedObjects)
+void FDMThroughputPropertyRowGenerator::AddComponentProperties(FDMComponentPropertyRowGeneratorParams& InParams)
 {
-	if (!IsValid(InComponent))
+	if (!IsValid(InParams.Object))
 	{
 		return;
 	}
 
-	if (InOutProcessedObjects.Contains(InComponent))
+	if (InParams.ProcessedObjects.Contains(InParams.Object))
 	{
 		return;
 	}
 
-	UDMMaterialStageThroughput* Throughput = Cast<UDMMaterialStageThroughput>(InComponent);
+	UDMMaterialStageThroughput* Throughput = Cast<UDMMaterialStageThroughput>(InParams.Object);
 
 	if (!Throughput)
 	{
 		return;
 	}
 
-	InOutProcessedObjects.Add(Throughput);
+	InParams.ProcessedObjects.Add(Throughput);
 
 	const TArray<FName>& ThroughputProperties = Throughput->GetEditableProperties();
 
 	for (const FName& ThroughputProperty : ThroughputProperties)
 	{
-		if (InComponent->IsPropertyVisible(ThroughputProperty))
+		if (Throughput->IsPropertyVisible(ThroughputProperty))
 		{
-			AddPropertyEditRows(InComponentEditorWidget, InComponent, ThroughputProperty, InOutPropertyRows, InOutProcessedObjects);
+			AddPropertyEditRows(InParams, ThroughputProperty);
 		}
 	}
 
@@ -54,7 +52,7 @@ void FDMThroughputPropertyRowGenerator::AddComponentProperties(const TSharedRef<
 		return;
 	}
 
-	InOutProcessedObjects.Add(Stage);
+	InParams.ProcessedObjects.Add(Stage);
 
 	static const FName StageInputsName = GET_MEMBER_NAME_CHECKED(UDMMaterialStage, Inputs);
 
@@ -70,7 +68,7 @@ void FDMThroughputPropertyRowGenerator::AddComponentProperties(const TSharedRef<
 				continue;
 			}
 
-			const int32 StartRow = InOutPropertyRows.Num();
+			const int32 StartRow = InParams.PropertyRows.Num();
 
 			for (const FDMMaterialStageConnectorChannel& Channel : InputMap[InputIdx].Channels)
 			{
@@ -78,27 +76,30 @@ void FDMThroughputPropertyRowGenerator::AddComponentProperties(const TSharedRef<
 
 				if (Inputs.IsValidIndex(StageInputIdx))
 				{
-					FDynamicMaterialEditorModule::GeneratorComponentPropertyRows(InComponentEditorWidget, Inputs[StageInputIdx], InOutPropertyRows, InOutProcessedObjects);
+					FDMComponentPropertyRowGeneratorParams InputParams = InParams;
+					InputParams.Object = Inputs[StageInputIdx];
+
+					FDynamicMaterialEditorModule::GeneratorComponentPropertyRows(InputParams);
 				}
 			}
 
-			for (int32 PropertyRowIdx = StartRow; PropertyRowIdx < InOutPropertyRows.Num(); ++PropertyRowIdx)
+			for (int32 PropertyRowIdx = StartRow; PropertyRowIdx < InParams.PropertyRows.Num(); ++PropertyRowIdx)
 			{
-				FDMPropertyHandle& PropertyRow = InOutPropertyRows[PropertyRowIdx];
+				FDMPropertyHandle& PropertyRow = InParams.PropertyRows[PropertyRowIdx];
 
 				if (PropertyRow.NameOverride.IsSet())
 				{
 					continue;
 				}
 
-				if (!PropertyRow.PropertyHandle.IsValid()
-					|| !PropertyRow.PropertyHandle->GetProperty())
+				if (!PropertyRow.PreviewHandle.PropertyHandle.IsValid()
+					|| !PropertyRow.PreviewHandle.PropertyHandle->GetProperty())
 				{
 					continue;
 				}
 
 				TArray<UObject*> Outers;
-				PropertyRow.PropertyHandle->GetOuterObjects(Outers);
+				PropertyRow.PreviewHandle.PropertyHandle->GetOuterObjects(Outers);
 
 				if (Outers.IsEmpty() || !Outers[0]->IsA<UDMMaterialValue>())
 				{
@@ -112,9 +113,9 @@ void FDMThroughputPropertyRowGenerator::AddComponentProperties(const TSharedRef<
 					{
 						PropertyRow.NameOverride = FText::FromName(PropertyRow.ValueName);
 					}
-					else if (PropertyRow.PropertyHandle.IsValid())
+					else if (PropertyRow.PreviewHandle.PropertyHandle.IsValid())
 					{
-						PropertyRow.NameOverride = PropertyRow.PropertyHandle->GetPropertyDisplayName();
+						PropertyRow.NameOverride = PropertyRow.PreviewHandle.PropertyHandle->GetPropertyDisplayName();
 					}
 					else
 					{
@@ -125,14 +126,16 @@ void FDMThroughputPropertyRowGenerator::AddComponentProperties(const TSharedRef<
 		}
 	}
 
+	FDMComponentPropertyRowGeneratorParams StageParams = InParams;
+	StageParams.Object = Stage;
+
 	const TArray<FName>& StageProperties = Stage->GetEditableProperties();
 
 	for (const FName& StageProperty : StageProperties)
 	{
 		if (StageProperty != StageInputsName)
 		{
-			AddPropertyEditRows(InComponentEditorWidget, Stage, StageProperty, InOutPropertyRows, InOutProcessedObjects);
-			continue;
+			AddPropertyEditRows(StageParams, StageProperty);
 		}
 	}
 }
diff --git a/Engine/Plugins/Experimental/DynamicMaterial/Source/DynamicMaterialEditor/Private/UI/Utils/DMWidgetLibrary.cpp b/Engine/Plugins/Experimental/DynamicMaterial/Source/DynamicMaterialEditor/Private/UI/Utils/DMWidgetLibrary.cpp
--- a/Engine/Plugins/Experimental/DynamicMaterial/Source/DynamicMaterialEditor/Private/UI/Utils/DMWidgetLibrary.cpp
+++ b/Engine/Plugins/Experimental/DynamicMaterial/Source/DynamicMaterialEditor/Private/UI/Utils/DMWidgetLibrary.cpp
@@ -5,9 +5,10 @@
 #include "Components/MaterialValues/DMMaterialValueFloat.h"
 #include "IDetailTreeNode.h"
 #include "IPropertyRowGenerator.h"
+#include "Model/DynamicMaterialModelBase.h"
 #include "Modules/ModuleManager.h"
 #include "PropertyEditorModule.h"
-#include "UI/Widgets/Editor/SDMMaterialComponentEditor.h"
+#include "Utils/DMPrivate.h"
 
 IDMWidgetLibrary& IDMWidgetLibrary::Get()
 {
@@ -41,61 +42,77 @@ void FDMWidgetLibrary::SetExpansionState(UObject* InOwner, FName InName, bool bI
 	ExpansionStates.FindOrAdd(ExpansionItem) = bInIsExpanded;
 }
 
-FDMPropertyHandle FDMWidgetLibrary::GetPropertyHandle(const SWidget* InOwningWidget, UObject* InObject, FName InPropertyName)
+FDMPropertyHandle FDMWidgetLibrary::GetPropertyHandle(const FDMPropertyHandleGenerateParams& InParams)
 {
-	TArray<FDMPropertyHandle>& PropertyHandles = PropertyHandleMap.FindOrAdd(InOwningWidget);
+	TArray<FDMPropertyHandle>& PropertyHandles = PropertyHandleMap.FindOrAdd(InParams.Widget);
 
 	for (const FDMPropertyHandle& ExistingHandle : PropertyHandles)
 	{
-		if (ExistingHandle.PropertyHandle && ExistingHandle.PropertyHandle->GetProperty()->GetFName() == InPropertyName)
+		if (ExistingHandle.PreviewHandle.PropertyHandle && ExistingHandle.PreviewHandle.PropertyHandle->GetProperty()->GetFName() == InParams.PropertyName)
 		{
 			TArray<UObject*> Outers;
-			ExistingHandle.PropertyHandle->GetOuterObjects(Outers);
+			ExistingHandle.PreviewHandle.PropertyHandle->GetOuterObjects(Outers);
 
-			if (Outers.IsEmpty() == false && Outers[0] == InObject)
+			if (Outers.IsEmpty() == false && Outers[0] == InParams.Object)
 			{
 				return ExistingHandle;
 			}
 		}
 	}
 
-	if (TSharedPtr<IPropertyRowGenerator> PropertyRowGenerator = SearchForGenerator(PropertyHandles, InObject))
+	if (const FDMPropertyHandle* ParentHandle = SearchForHandle(PropertyHandles, InParams.Object))
 	{
-		FDMPropertyHandle PropertyHandle;
-		PropertyHandle.PropertyRowGenerator = PropertyRowGenerator;
+		FDMPropertyHandle NewChildHandle;
+		NewChildHandle.PreviewHandle.PropertyRowGenerator = ParentHandle->PreviewHandle.PropertyRowGenerator;
+		NewChildHandle.OriginalHandle.PropertyRowGenerator = ParentHandle->OriginalHandle.PropertyRowGenerator;
 
-		if (TSharedPtr<IDetailTreeNode> DetailTreeNode = SearchGeneratorForNode(PropertyRowGenerator.ToSharedRef(), InPropertyName))
+		if (NewChildHandle.PreviewHandle.PropertyRowGenerator.IsValid())
 		{
-			PropertyHandle.DetailTreeNode = DetailTreeNode;
-			PropertyHandle.PropertyHandle = DetailTreeNode->CreatePropertyHandle();
+			if (TSharedPtr<IDetailTreeNode> DetailTreeNode = SearchGeneratorForNode(NewChildHandle.PreviewHandle.PropertyRowGenerator.ToSharedRef(), InParams.PropertyName))
+			{
+				NewChildHandle.PreviewHandle.DetailTreeNode = DetailTreeNode;
+				NewChildHandle.PreviewHandle.PropertyHandle = DetailTreeNode->CreatePropertyHandle();
+			}
+		}
 
-			AddPropertyMetaData(InObject, InPropertyName, PropertyHandle);
+		if (NewChildHandle.OriginalHandle.PropertyRowGenerator.IsValid())
+		{
+			if (TSharedPtr<IDetailTreeNode> DetailTreeNode = SearchGeneratorForNode(NewChildHandle.OriginalHandle.PropertyRowGenerator.ToSharedRef(), InParams.PropertyName))
+			{
+				NewChildHandle.OriginalHandle.DetailTreeNode = DetailTreeNode;
+				NewChildHandle.OriginalHandle.PropertyHandle = DetailTreeNode->CreatePropertyHandle();
+			}
+		}
 
-			return PropertyHandle;
+		if (NewChildHandle.PreviewHandle.PropertyHandle.IsValid())
+		{
+			AddPropertyMetaData(InParams.Object, InParams.PropertyName, NewChildHandle);
 		}
 
-		return PropertyHandle;
+		PropertyHandles.Add(NewChildHandle);
+
+		return NewChildHandle;
 	}
 
-	FDMPropertyHandle NewHandle = CreatePropertyHandle(InOwningWidget, InObject, InPropertyName);
+	FDMPropertyHandle NewHandle = CreatePropertyHandle(InParams);
 
-	if (!NewHandle.PropertyHandle.IsValid() && NewHandle.DetailTreeNode.IsValid())
+	if (!NewHandle.PreviewHandle.PropertyHandle.IsValid() && NewHandle.PreviewHandle.DetailTreeNode.IsValid())
 	{
-		NewHandle.PropertyHandle = NewHandle.DetailTreeNode->CreatePropertyHandle();
+		NewHandle.PreviewHandle.PropertyHandle = NewHandle.PreviewHandle.DetailTreeNode->CreatePropertyHandle();
 	}
 
-	AddPropertyMetaData(InObject, InPropertyName, NewHandle);
+	if (!NewHandle.OriginalHandle.PropertyHandle.IsValid() && NewHandle.OriginalHandle.DetailTreeNode.IsValid())
+	{
+		NewHandle.OriginalHandle.PropertyHandle = NewHandle.OriginalHandle.DetailTreeNode->CreatePropertyHandle();
+	}
+
+	AddPropertyMetaData(InParams.Object, InParams.PropertyName, NewHandle);
 
 	PropertyHandles.Add(NewHandle);
 
 	return NewHandle;
 }
 
-FDMPropertyHandle FDMWidgetLibrary::GetPropertyHandle(const SDMMaterialComponentEditor* InOwningWidget, UObject* InObject, FName InPropertyName)
-{
-	return GetPropertyHandle(static_cast<const SWidget*>(InOwningWidget), InObject, InPropertyName);
-}
-
 void FDMWidgetLibrary::ClearPropertyHandles(const SWidget* InOwningWidget)
 {
 	PropertyHandleMap.Remove(InOwningWidget);
@@ -157,28 +174,59 @@ void FDMWidgetLibrary::ClearData()
 	PropertyHandleMap.Empty();
 }
 
-FDMPropertyHandle FDMWidgetLibrary::CreatePropertyHandle(const void* InOwningWidget, UObject* InObject, FName InPropertyName)
+FDMPropertyHandle FDMWidgetLibrary::CreatePropertyHandle(const FDMPropertyHandleGenerateParams& InParams)
 {
 	FDMPropertyHandle PropertyHandle;
 
+	if (!InParams.Object)
+	{
+		return PropertyHandle;
+	}
+
 	FPropertyEditorModule& PropertyEditor = FModuleManager::Get().LoadModuleChecked<FPropertyEditorModule>(TEXT("PropertyEditor"));
 
-	FPropertyRowGeneratorArgs RowGeneratorArgs;
+	FPropertyRowGeneratorArgs PreviewRowGeneratorArgs;
+	PreviewRowGeneratorArgs.NotifyHook = InParams.NotifyHook;
+
+	PropertyHandle.PreviewHandle.PropertyRowGenerator = PropertyEditor.CreatePropertyRowGenerator(PreviewRowGeneratorArgs);
+	PropertyHandle.PreviewHandle.PropertyRowGenerator->SetObjects({InParams.Object});
 
-	if (UDMMaterialComponent* Component = Cast<UDMMaterialComponent>(InObject))
+	if (const TSharedPtr<IDetailTreeNode> FoundTreeNode = SearchGeneratorForNode(
+		PropertyHandle.PreviewHandle.PropertyRowGenerator.ToSharedRef(), InParams.PropertyName))
 	{
-		RowGeneratorArgs.NotifyHook = Component;
+		PropertyHandle.PreviewHandle.DetailTreeNode = FoundTreeNode;
+		PropertyHandle.PreviewHandle.PropertyHandle = FoundTreeNode->CreatePropertyHandle();
 	}
 
-	PropertyHandle.PropertyRowGenerator = PropertyEditor.CreatePropertyRowGenerator(RowGeneratorArgs);
-	PropertyHandle.PropertyRowGenerator->SetObjects({InObject});
-
-	if (const TSharedPtr<IDetailTreeNode> FoundTreeNode = SearchGeneratorForNode(
-		PropertyHandle.PropertyRowGenerator.ToSharedRef(), InPropertyName))
+	if (InParams.PreviewMaterialModelBase && InParams.OriginalMaterialModelBase && InParams.Object)
 	{
-		PropertyHandle.DetailTreeNode = FoundTreeNode;
-		PropertyHandle.PropertyHandle = FoundTreeNode->CreatePropertyHandle();
-		return PropertyHandle;
+		UObject* OriginalObject = nullptr;
+
+		if (InParams.Object == InParams.PreviewMaterialModelBase)
+		{
+			OriginalObject = InParams.OriginalMaterialModelBase;
+		}
+		else if (InParams.Object->IsIn(InParams.PreviewMaterialModelBase))
+		{
+			const FString RelativePath = InParams.Object->GetPathName(InParams.PreviewMaterialModelBase);
+			OriginalObject = UE::DynamicMaterialEditor::Private::FindSubobject(InParams.OriginalMaterialModelBase, RelativePath);
+		}
+
+		if (OriginalObject)
+		{
+			// No notify hook for the original
+			FPropertyRowGeneratorArgs OriginalRowGeneratorArgs;
+
+			PropertyHandle.OriginalHandle.PropertyRowGenerator = PropertyEditor.CreatePropertyRowGenerator(OriginalRowGeneratorArgs);
+			PropertyHandle.OriginalHandle.PropertyRowGenerator->SetObjects({OriginalObject});
+
+			if (const TSharedPtr<IDetailTreeNode> FoundTreeNode = SearchGeneratorForNode(
+				PropertyHandle.OriginalHandle.PropertyRowGenerator.ToSharedRef(), InParams.PropertyName))
+			{
+				PropertyHandle.OriginalHandle.DetailTreeNode = FoundTreeNode;
+				PropertyHandle.OriginalHandle.PropertyHandle = FoundTreeNode->CreatePropertyHandle();
+			}
+		}
 	}
 
 	return PropertyHandle;
@@ -224,7 +272,7 @@ TSharedPtr<IDetailTreeNode> FDMWidgetLibrary::SearchGeneratorForNode(const TShar
 	return SearchNodesForProperty(InGenerator->GetRootTreeNodes(), InPropertyName);
 }
 
-TSharedPtr<IPropertyRowGenerator> FDMWidgetLibrary::SearchForGenerator(const TArray<FDMPropertyHandle>& InPropertyHandles, UObject* InObject)
+const FDMPropertyHandle* FDMWidgetLibrary::SearchForHandle(const TArray<FDMPropertyHandle>& InPropertyHandles, UObject* InObject)
 {
 	if (!InObject)
 	{
@@ -233,13 +281,13 @@ TSharedPtr<IPropertyRowGenerator> FDMWidgetLibrary::SearchForGenerator(const TAr
 
 	for (const FDMPropertyHandle& PropertyHandle : InPropertyHandles)
 	{
-		if (PropertyHandle.PropertyRowGenerator.IsValid())
+		if (PropertyHandle.PreviewHandle.PropertyRowGenerator.IsValid())
 		{
-			for (const TWeakObjectPtr<UObject>& WeakObject : PropertyHandle.PropertyRowGenerator->GetSelectedObjects())
+			for (const TWeakObjectPtr<UObject>& WeakObject : PropertyHandle.PreviewHandle.PropertyRowGenerator->GetSelectedObjects())
 			{
 				if (WeakObject.Get() == InObject)
 				{
-					return PropertyHandle.PropertyRowGenerator;
+					return &PropertyHandle;
 				}
 			}
 		}
@@ -252,14 +300,13 @@ void FDMWidgetLibrary::AddPropertyMetaData(UObject* InObject, FName InPropertyNa
 {
 	FProperty* Property = nullptr;
 
-	if (InPropertyHandle.PropertyHandle.IsValid())
+	if (InPropertyHandle.PreviewHandle.PropertyHandle.IsValid())
 	{
-		TSharedRef<IPropertyHandle> PropertyHandleRef = InPropertyHandle.PropertyHandle.ToSharedRef();
-
+		TSharedRef<IPropertyHandle> PropertyHandleRef = InPropertyHandle.PreviewHandle.PropertyHandle.ToSharedRef();
 		InPropertyHandle.Priority = GetPriority(PropertyHandleRef);
 		InPropertyHandle.bKeyframeable = IsKeyframeable(PropertyHandleRef);
 
-		Property = InPropertyHandle.PropertyHandle->GetProperty();
+		Property = InPropertyHandle.PreviewHandle.PropertyHandle->GetProperty();
 
 		if (UDMMaterialValueFloat* FloatValue = Cast<UDMMaterialValueFloat>(InObject))
 		{
@@ -270,10 +317,10 @@ void FDMWidgetLibrary::AddPropertyMetaData(UObject* InObject, FName InPropertyNa
 				const FName ClampMin = FName("ClampMin");
 				const FName ClampMax = FName("ClampMax");
 
-				InPropertyHandle.PropertyHandle->SetInstanceMetaData(UIMin, FString::SanitizeFloat(FloatValue->GetValueRange().Min));
-				InPropertyHandle.PropertyHandle->SetInstanceMetaData(ClampMin, FString::SanitizeFloat(FloatValue->GetValueRange().Min));
-				InPropertyHandle.PropertyHandle->SetInstanceMetaData(UIMax, FString::SanitizeFloat(FloatValue->GetValueRange().Max));
-				InPropertyHandle.PropertyHandle->SetInstanceMetaData(ClampMax, FString::SanitizeFloat(FloatValue->GetValueRange().Max));
+				InPropertyHandle.PreviewHandle.PropertyHandle->SetInstanceMetaData(UIMin, FString::SanitizeFloat(FloatValue->GetValueRange().Min));
+				InPropertyHandle.PreviewHandle.PropertyHandle->SetInstanceMetaData(ClampMin, FString::SanitizeFloat(FloatValue->GetValueRange().Min));
+				InPropertyHandle.PreviewHandle.PropertyHandle->SetInstanceMetaData(UIMax, FString::SanitizeFloat(FloatValue->GetValueRange().Max));
+				InPropertyHandle.PreviewHandle.PropertyHandle->SetInstanceMetaData(ClampMax, FString::SanitizeFloat(FloatValue->GetValueRange().Max));
 			}
 		}
 	}
diff --git a/Engine/Plugins/Experimental/DynamicMaterial/Source/DynamicMaterialEditor/Private/UI/Utils/DMWidgetLibrary.h b/Engine/Plugins/Experimental/DynamicMaterial/Source/DynamicMaterialEditor/Private/UI/Utils/DMWidgetLibrary.h
--- a/Engine/Plugins/Experimental/DynamicMaterial/Source/DynamicMaterialEditor/Private/UI/Utils/DMWidgetLibrary.h
+++ b/Engine/Plugins/Experimental/DynamicMaterial/Source/DynamicMaterialEditor/Private/UI/Utils/DMWidgetLibrary.h
@@ -35,8 +35,7 @@ public:
 	void ClearData();
 
 	//~ Begin IDMWidgetLibrary
-	virtual FDMPropertyHandle GetPropertyHandle(const SWidget* InOwningWidget, UObject* InObject, FName InPropertyName) override;
-	virtual FDMPropertyHandle GetPropertyHandle(const SDMMaterialComponentEditor* InOwningWidget, UObject* InObject, FName InPropertyName) override;
+	virtual FDMPropertyHandle GetPropertyHandle(const FDMPropertyHandleGenerateParams& InParams) override;
 	//~ End IDMWidgetLibrary
 
 private:
@@ -64,13 +63,13 @@ private:
 
 	TMap<const SWidget*, TArray<FDMPropertyHandle>> PropertyHandleMap;
 
-	FDMPropertyHandle CreatePropertyHandle(const void* InOwningWidget, UObject* InObject, FName InPropertyName);
+	FDMPropertyHandle CreatePropertyHandle(const FDMPropertyHandleGenerateParams& InParams);
 
 	static TSharedPtr<IDetailTreeNode> SearchNodesForProperty(const TArray<TSharedRef<IDetailTreeNode>>& InNodes, FName InPropertyName);
 
 	static TSharedPtr<IDetailTreeNode> SearchGeneratorForNode(const TSharedRef<IPropertyRowGenerator>& InGenerator, FName InPropertyName);
 
-	static TSharedPtr<IPropertyRowGenerator> SearchForGenerator(const TArray<FDMPropertyHandle>& InPropertyHandles, UObject* InObject);
+	static const FDMPropertyHandle* SearchForHandle(const TArray<FDMPropertyHandle>& InPropertyHandles, UObject* InObject);
 
 	static void AddPropertyMetaData(UObject* InObject, FName InPropertyName, FDMPropertyHandle& InPropertyHandle);
 
diff --git a/Engine/Plugins/Experimental/DynamicMaterial/Source/DynamicMaterialEditor/Private/UI/Widgets/Editor/EditorLayouts/SDMMaterialEditor_Left.cpp b/Engine/Plugins/Experimental/DynamicMaterial/Source/DynamicMaterialEditor/Private/UI/Widgets/Editor/EditorLayouts/SDMMaterialEditor_Left.cpp
--- a/Engine/Plugins/Experimental/DynamicMaterial/Source/DynamicMaterialEditor/Private/UI/Widgets/Editor/EditorLayouts/SDMMaterialEditor_Left.cpp
+++ b/Engine/Plugins/Experimental/DynamicMaterial/Source/DynamicMaterialEditor/Private/UI/Widgets/Editor/EditorLayouts/SDMMaterialEditor_Left.cpp
@@ -16,8 +16,9 @@ void SDMMaterialEditor_Left::Construct(const FArguments& InArgs, const TSharedRe
 {
 	SDMMaterialEditor_LeftBase::Construct(
 		SDMMaterialEditor_LeftBase::FArguments()
-		.MaterialModelBase(InArgs._MaterialModelBase)
-		.MaterialProperty(InArgs._MaterialProperty),
+			.MaterialModelBase(InArgs._MaterialModelBase)
+			.MaterialProperty(InArgs._MaterialProperty)
+			.PreviewMaterialModelBase(InArgs._PreviewMaterialModelBase),
 		InDesignerWidget
 	);
 }
diff --git a/Engine/Plugins/Experimental/DynamicMaterial/Source/DynamicMaterialEditor/Private/UI/Widgets/Editor/EditorLayouts/SDMMaterialEditor_Left.h b/Engine/Plugins/Experimental/DynamicMaterial/Source/DynamicMaterialEditor/Private/UI/Widgets/Editor/EditorLayouts/SDMMaterialEditor_Left.h
--- a/Engine/Plugins/Experimental/DynamicMaterial/Source/DynamicMaterialEditor/Private/UI/Widgets/Editor/EditorLayouts/SDMMaterialEditor_Left.h
+++ b/Engine/Plugins/Experimental/DynamicMaterial/Source/DynamicMaterialEditor/Private/UI/Widgets/Editor/EditorLayouts/SDMMaterialEditor_Left.h
@@ -9,9 +9,11 @@ class SDMMaterialEditor_Left : public SDMMaterialEditor_LeftBase
 	SLATE_BEGIN_ARGS(SDMMaterialEditor_Left)
 		: _MaterialModelBase(nullptr)
 		, _MaterialProperty(TOptional<FDMObjectMaterialProperty>())
+		, _PreviewMaterialModelBase(nullptr)
 		{}
 		SLATE_ARGUMENT(UDynamicMaterialModelBase*, MaterialModelBase)
 		SLATE_ARGUMENT(TOptional<FDMObjectMaterialProperty>, MaterialProperty)
+		SLATE_ARGUMENT(UDynamicMaterialModelBase*, PreviewMaterialModelBase)
 	SLATE_END_ARGS()
 
 public:
diff --git a/Engine/Plugins/Experimental/DynamicMaterial/Source/DynamicMaterialEditor/Private/UI/Widgets/Editor/EditorLayouts/SDMMaterialEditor_LeftBase.cpp b/Engine/Plugins/Experimental/DynamicMaterial/Source/DynamicMaterialEditor/Private/UI/Widgets/Editor/EditorLayouts/SDMMaterialEditor_LeftBase.cpp
--- a/Engine/Plugins/Experimental/DynamicMaterial/Source/DynamicMaterialEditor/Private/UI/Widgets/Editor/EditorLayouts/SDMMaterialEditor_LeftBase.cpp
+++ b/Engine/Plugins/Experimental/DynamicMaterial/Source/DynamicMaterialEditor/Private/UI/Widgets/Editor/EditorLayouts/SDMMaterialEditor_LeftBase.cpp
@@ -19,8 +19,9 @@ void SDMMaterialEditor_LeftBase::Construct(const FArguments& InArgs, const TShar
 {
 	SDMMaterialEditor::Construct(
 		SDMMaterialEditor::FArguments()
-		.MaterialModelBase(InArgs._MaterialModelBase)
-		.MaterialProperty(InArgs._MaterialProperty),
+			.MaterialModelBase(InArgs._MaterialModelBase)
+			.MaterialProperty(InArgs._MaterialProperty)
+			.PreviewMaterialModelBase(InArgs._PreviewMaterialModelBase),
 		InDesignerWidget
 	);
 }
diff --git a/Engine/Plugins/Experimental/DynamicMaterial/Source/DynamicMaterialEditor/Private/UI/Widgets/Editor/EditorLayouts/SDMMaterialEditor_LeftBase.h b/Engine/Plugins/Experimental/DynamicMaterial/Source/DynamicMaterialEditor/Private/UI/Widgets/Editor/EditorLayouts/SDMMaterialEditor_LeftBase.h
--- a/Engine/Plugins/Experimental/DynamicMaterial/Source/DynamicMaterialEditor/Private/UI/Widgets/Editor/EditorLayouts/SDMMaterialEditor_LeftBase.h
+++ b/Engine/Plugins/Experimental/DynamicMaterial/Source/DynamicMaterialEditor/Private/UI/Widgets/Editor/EditorLayouts/SDMMaterialEditor_LeftBase.h
@@ -9,9 +9,11 @@ class SDMMaterialEditor_LeftBase : public SDMMaterialEditor
 	SLATE_BEGIN_ARGS(SDMMaterialEditor_LeftBase)
 		: _MaterialModelBase(nullptr)
 		, _MaterialProperty(TOptional<FDMObjectMaterialProperty>())
+		, _PreviewMaterialModelBase(nullptr)
 		{}
 		SLATE_ARGUMENT(UDynamicMaterialModelBase*, MaterialModelBase)
 		SLATE_ARGUMENT(TOptional<FDMObjectMaterialProperty>, MaterialProperty)
+		SLATE_ARGUMENT(UDynamicMaterialModelBase*, PreviewMaterialModelBase)
 	SLATE_END_ARGS()
 
 public:
diff --git a/Engine/Plugins/Experimental/DynamicMaterial/Source/DynamicMaterialEditor/Private/UI/Widgets/Editor/EditorLayouts/SDMMaterialEditor_LeftSlim.cpp b/Engine/Plugins/Experimental/DynamicMaterial/Source/DynamicMaterialEditor/Private/UI/Widgets/Editor/EditorLayouts/SDMMaterialEditor_LeftSlim.cpp
--- a/Engine/Plugins/Experimental/DynamicMaterial/Source/DynamicMaterialEditor/Private/UI/Widgets/Editor/EditorLayouts/SDMMaterialEditor_LeftSlim.cpp
+++ b/Engine/Plugins/Experimental/DynamicMaterial/Source/DynamicMaterialEditor/Private/UI/Widgets/Editor/EditorLayouts/SDMMaterialEditor_LeftSlim.cpp
@@ -15,7 +15,8 @@ void SDMMaterialEditor_LeftSlim::Construct(const FArguments& InArgs, const TShar
 	SDMMaterialEditor_LeftBase::Construct(
 		SDMMaterialEditor_LeftBase::FArguments()
 			.MaterialModelBase(InArgs._MaterialModelBase)
-			.MaterialProperty(InArgs._MaterialProperty),
+			.MaterialProperty(InArgs._MaterialProperty)
+			.PreviewMaterialModelBase(InArgs._PreviewMaterialModelBase),
 		InDesignerWidget
 	);
 }
diff --git a/Engine/Plugins/Experimental/DynamicMaterial/Source/DynamicMaterialEditor/Private/UI/Widgets/Editor/SDMMaterialComponentEditor.cpp b/Engine/Plugins/Experimental/DynamicMaterial/Source/DynamicMaterialEditor/Private/UI/Widgets/Editor/SDMMaterialComponentEditor.cpp
--- a/Engine/Plugins/Experimental/DynamicMaterial/Source/DynamicMaterialEditor/Private/UI/Widgets/Editor/SDMMaterialComponentEditor.cpp
+++ b/Engine/Plugins/Experimental/DynamicMaterial/Source/DynamicMaterialEditor/Private/UI/Widgets/Editor/SDMMaterialComponentEditor.cpp
@@ -4,12 +4,15 @@
 
 #include "Components/DMMaterialComponent.h"
 #include "Components/DMMaterialLayer.h"
+#include "Components/DMMaterialValue.h"
 #include "CustomDetailsViewArgs.h"
 #include "DynamicMaterialEditorModule.h"
+#include "DynamicMaterialEditorSettings.h"
 #include "DynamicMaterialModule.h"
 #include "ICustomDetailsView.h"
 #include "Items/ICustomDetailsViewCustomCategoryItem.h"
 #include "Items/ICustomDetailsViewItem.h"
+#include "Model/DynamicMaterialModelBase.h"
 #include "UI/Utils/DMWidgetLibrary.h"
 #include "UI/Widgets/SDMMaterialEditor.h"
 #include "Widgets/Input/SComboButton.h"
@@ -55,6 +58,53 @@ UDMMaterialComponent* SDMMaterialComponentEditor::GetComponent() const
 	return Cast<UDMMaterialComponent>(ObjectWeak.Get());
 }
 
+UDMMaterialComponent* SDMMaterialComponentEditor::GetOriginalComponent(UDMMaterialComponent* InPreviewComponent) const
+{
+	if (!InPreviewComponent)
+	{
+		return nullptr;
+	}
+
+	TSharedPtr<SDMMaterialEditor> EditorWidget = EditorWidgetWeak.Pin();
+
+	if (!EditorWidget.IsValid())
+	{
+		return nullptr;
+	}
+
+	return EditorWidget->GetOriginalComponent(InPreviewComponent);
+}
+
+void SDMMaterialComponentEditor::NotifyPreChange(FProperty* InPropertyAboutToChange)
+{
+	SDMObjectEditorWidgetBase::NotifyPreChange(InPropertyAboutToChange);
+
+	if (UDMMaterialComponent* Component = GetComponent())
+	{
+		Component->NotifyPreChange(InPropertyAboutToChange);
+
+		if (UDMMaterialComponent* OriginalComponent = GetOriginalComponent(Component))
+		{
+			OriginalComponent->NotifyPreChange(InPropertyAboutToChange);
+		}
+	}
+}
+
+void SDMMaterialComponentEditor::NotifyPostChange(const FPropertyChangedEvent& InPropertyChangedEvent, FProperty* InPropertyThatChanged)
+{
+	SDMObjectEditorWidgetBase::NotifyPostChange(InPropertyChangedEvent, InPropertyThatChanged);
+
+	if (UDMMaterialComponent* Component = GetComponent())
+	{
+		Component->NotifyPostChange(InPropertyChangedEvent, InPropertyThatChanged);
+
+		if (UDMMaterialComponent* OriginalComponent = GetOriginalComponent(Component))
+		{			
+			OriginalComponent->NotifyPostChange(InPropertyChangedEvent, InPropertyThatChanged);
+		}
+	}
+}
+
 void SDMMaterialComponentEditor::OnComponentUpdated(UDMMaterialComponent* InComponent, UDMMaterialComponent* InSource, EDMUpdateType InUpdateType)
 {
 	if (EnumHasAnyFlags(InUpdateType, EDMUpdateType::Structure | EDMUpdateType::RefreshDetailView))
@@ -64,6 +114,29 @@ void SDMMaterialComponentEditor::OnComponentUpdated(UDMMaterialComponent* InComp
 			EditorWidget->EditComponent(GetComponent(), /* Force refresh */ true);
 		}
 	}
+
+	if (UDynamicMaterialEditorSettings* Settings = UDynamicMaterialEditorSettings::Get())
+	{
+		if (Settings->bAutomaticallyCopyParametersToSourceMaterial)
+		{
+			if (InComponent)
+			{
+				if (UDMMaterialComponent* OriginalComponent = GetOriginalComponent(InComponent))
+				{
+					IDMParameterContainer::CopyParametersBetween(InComponent, OriginalComponent);
+					return;
+				}
+			}
+		}
+	}
+
+	if (TSharedPtr<SDMMaterialEditor> EditorWidget = GetEditorWidget())
+	{
+		if (UDynamicMaterialModelBase* MaterialModelBase = EditorWidget->GetOriginalMaterialModelBase())
+		{
+			MaterialModelBase->MarkPreviewModified();
+		}
+	}	
 }
 
 TSharedRef<ICustomDetailsViewItem> SDMMaterialComponentEditor::GetDefaultCategory(const TSharedRef<ICustomDetailsView>& InDetailsView,
@@ -100,17 +173,114 @@ TSharedRef<ICustomDetailsViewItem> SDMMaterialComponentEditor::GetDefaultCategor
 
 TArray<FDMPropertyHandle> SDMMaterialComponentEditor::GetPropertyRows()
 {
+	TSharedPtr<SDMMaterialEditor> EditorWidget = GetEditorWidget();
+
+	if (!EditorWidget)
+	{
+		return {};
+	}
+
 	TArray<FDMPropertyHandle> PropertyRows;
-	TSet<UDMMaterialComponent*> ProcessedObjects;
+	TSet<UObject*> ProcessedObjects;
 
-	FDynamicMaterialEditorModule::GeneratorComponentPropertyRows(
-		SharedThis(this),
-		GetComponent(),
-		PropertyRows,
-		ProcessedObjects
-	);
+	FDMComponentPropertyRowGeneratorParams Params(PropertyRows, ProcessedObjects);
+	Params.NotifyHook = this;
+	Params.Owner = this;
+	Params.Object = GetComponent();
+	Params.PreviewMaterialModelBase = EditorWidget->GetPreviewMaterialModelBase();
+	Params.OriginalMaterialModelBase = EditorWidget->GetOriginalMaterialModelBase();
+
+	FDynamicMaterialEditorModule::GeneratorComponentPropertyRows(Params);
+
+	for (const FDMPropertyHandle& PropertyRow : PropertyRows)
+	{
+		if (PropertyRow.PreviewHandle.PropertyHandle.IsValid())
+		{
+			TArray<UObject*> Outers;
+			PropertyRow.PreviewHandle.PropertyHandle->GetOuterObjects(Outers);
+
+			for (UObject* Outer : Outers)
+			{
+				if (UDMMaterialComponent* Component = Cast<UDMMaterialComponent>(Outer))
+				{
+					if (!Component->GetOnUpdate().IsBoundToObject(this))
+					{
+						Component->GetOnUpdate().AddSP(this, &SDMMaterialComponentEditor::OnComponentUpdated);
+					}
+				}
+			}
+		}
+	}
 
 	return PropertyRows;
 }
 
+void SDMMaterialComponentEditor::AddDetailTreeRowExtensionWidgets(const TSharedRef<ICustomDetailsView>& InDetailsView,
+	const FDMPropertyHandle& InPropertyRow, const TSharedRef<ICustomDetailsViewItem>& InItem, const TSharedRef<IPropertyHandle>& InPropertyHandle)
+{
+	TArray<FPropertyRowExtensionButton> ExtensionButtons;
+
+	// Reset to Default
+	if (InPropertyRow.ResetToDefaultOverride.IsSet())
+	{
+		FPropertyRowExtensionButton ResetToDefaultButton;
+
+		if (InItem->CreateResetToDefaultButton(ResetToDefaultButton))
+		{
+			ExtensionButtons.Add(MoveTemp(ResetToDefaultButton));
+		}
+	}
+
+	// Global Extensions
+	FPropertyEditorModule& PropertyEditorModule = FModuleManager::Get().GetModuleChecked<FPropertyEditorModule>("PropertyEditor");
+	FOnGenerateGlobalRowExtensionArgs RowExtensionArgs;
+	RowExtensionArgs.OwnerTreeNode = InPropertyRow.OriginalHandle.DetailTreeNode;
+	RowExtensionArgs.PropertyHandle = InPropertyRow.OriginalHandle.PropertyHandle;
+	PropertyEditorModule.GetGlobalRowExtensionDelegate().Broadcast(RowExtensionArgs, ExtensionButtons);
+
+	bool bValidPropertyHandles = false;
+
+	if (InPropertyRow.PreviewHandle.PropertyHandle.IsValid()
+		&& InPropertyRow.OriginalHandle.PropertyHandle.IsValid())
+	{
+		TArray<UObject*> OuterPreview;
+		TArray<UObject*> OuterOriginal;
+
+		InPropertyRow.PreviewHandle.PropertyHandle->GetOuterObjects(OuterPreview);
+		InPropertyRow.OriginalHandle.PropertyHandle->GetOuterObjects(OuterOriginal);
+
+		if (OuterPreview.Num() == 1 && OuterOriginal.Num() == 1
+			&& IsValid(OuterPreview[0]) && IsValid(OuterOriginal[0])
+			&& OuterPreview[0]->GetClass() == OuterOriginal[0]->GetClass())
+		{
+			bValidPropertyHandles = true;
+
+			// Sequencer relies on getting the Keyframe Handler via the Details View of the IDetailTreeNode, but is null since there's no Details View here
+			// instead add it manually here
+			if (InPropertyRow.bKeyframeable)
+			{
+				if (TOptional<FPropertyRowExtensionButton> CreateKeyButton = CreateKeyframeButton(
+					InPropertyRow.PreviewHandle.PropertyHandle,
+					InPropertyRow.OriginalHandle.PropertyHandle))
+				{
+					ExtensionButtons.Add(*CreateKeyButton);
+				}
+			}
+		}
+	}
+
+	if (!bValidPropertyHandles)
+	{
+		ExtensionButtons.Add(CreateNeedsApplyButton());
+	}		
+
+	if (ExtensionButtons.IsEmpty())
+	{
+		InItem->SetOverrideWidget(ECustomDetailsViewWidgetType::Extensions, SNullWidget::NullWidget);
+		return;
+	}
+
+	InItem->SetOverrideWidget(ECustomDetailsViewWidgetType::Extensions, InItem->CreateExtensionButtonWidget(ExtensionButtons));
+}
+
 #undef LOCTEXT_NAMESPACE
diff --git a/Engine/Plugins/Experimental/DynamicMaterial/Source/DynamicMaterialEditor/Private/UI/Widgets/Editor/SDMMaterialComponentEditor.h b/Engine/Plugins/Experimental/DynamicMaterial/Source/DynamicMaterialEditor/Private/UI/Widgets/Editor/SDMMaterialComponentEditor.h
--- a/Engine/Plugins/Experimental/DynamicMaterial/Source/DynamicMaterialEditor/Private/UI/Widgets/Editor/SDMMaterialComponentEditor.h
+++ b/Engine/Plugins/Experimental/DynamicMaterial/Source/DynamicMaterialEditor/Private/UI/Widgets/Editor/SDMMaterialComponentEditor.h
@@ -4,8 +4,9 @@
 
 #include "UI/Widgets/Editor/SDMObjectEditorWidgetBase.h"
 
-enum class EDMUpdateType : uint8;
 class UDMMaterialComponent;
+enum class EDMUpdateType : uint8;
+struct FPropertyRowExtensionButton;
 
 /** Extends the object editor to provide component-specific events and properties. */
 class SDMMaterialComponentEditor : public SDMObjectEditorWidgetBase
@@ -22,12 +23,21 @@ public:
 
 	UDMMaterialComponent* GetComponent() const;
 
+	//~ Begin FNotifyHook
+	virtual void NotifyPreChange(FProperty* InPropertyAboutToChange) override;
+	virtual void NotifyPostChange(const FPropertyChangedEvent& InPropertyChangedEvent, FProperty* InPropertyThatChanged) override;
+	//~ End FNotifyHook
+
 protected:
 	void OnComponentUpdated(UDMMaterialComponent* InComponent, UDMMaterialComponent* InSource, EDMUpdateType InUpdateType);
 
+	UDMMaterialComponent* GetOriginalComponent(UDMMaterialComponent* InPreviewComponent) const;
+
 	//~ Begin SDMObjectEditorWidgetBase
 	virtual TSharedRef<ICustomDetailsViewItem> GetDefaultCategory(const TSharedRef<ICustomDetailsView>& InDetailsView,
 		const FCustomDetailsViewItemId& InRootId) override;
 	virtual TArray<FDMPropertyHandle> GetPropertyRows() override;
+	virtual void AddDetailTreeRowExtensionWidgets(const TSharedRef<ICustomDetailsView>& InDetailsView, const FDMPropertyHandle& InPropertyRow,
+		const TSharedRef<ICustomDetailsViewItem>& InItem, const TSharedRef<IPropertyHandle>& InPropertyHandle) override;
 	//~ End SDMObjectEditorWidgetBase
 };
diff --git a/Engine/Plugins/Experimental/DynamicMaterial/Source/DynamicMaterialEditor/Private/UI/Widgets/Editor/SDMMaterialGlobalSettingsEditor.cpp b/Engine/Plugins/Experimental/DynamicMaterial/Source/DynamicMaterialEditor/Private/UI/Widgets/Editor/SDMMaterialGlobalSettingsEditor.cpp
--- a/Engine/Plugins/Experimental/DynamicMaterial/Source/DynamicMaterialEditor/Private/UI/Widgets/Editor/SDMMaterialGlobalSettingsEditor.cpp
+++ b/Engine/Plugins/Experimental/DynamicMaterial/Source/DynamicMaterialEditor/Private/UI/Widgets/Editor/SDMMaterialGlobalSettingsEditor.cpp
@@ -2,11 +2,20 @@
 
 #include "UI/Widgets/Editor/SDMMaterialGlobalSettingsEditor.h"
 
+#include "Components/DMMaterialComponent.h"
+#include "Components/DMMaterialValue.h"
 #include "DMEDefs.h"
+#include "DynamicMaterialEditorSettings.h"
+#include "ICustomDetailsView.h"
+#include "IDynamicMaterialEditorModule.h"
+#include "Items/ICustomDetailsViewItem.h"
+#include "Model/DynamicMaterialModel.h"
 #include "Model/DynamicMaterialModelBase.h"
 #include "UI/PropertyGenerators/DMMaterialModelPropertyRowGenerator.h"
 #include "UI/Widgets/SDMMaterialEditor.h"
 
+#define LOCTEXT_NAMESPACE "SDMMaterialGlobalSettingsEditor"
+
 void SDMMaterialGlobalSettingsEditor::PrivateRegisterAttributes(FSlateAttributeDescriptor::FInitializer&)
 {
 }
@@ -27,15 +36,221 @@ UDynamicMaterialModelBase* SDMMaterialGlobalSettingsEditor::GetMaterialModelBase
 	return Cast<UDynamicMaterialModelBase>(ObjectWeak.Get());
 }
 
+UDynamicMaterialModelBase* SDMMaterialGlobalSettingsEditor::GetOriginalMaterialModelBase() const
+{
+	TSharedPtr<SDMMaterialEditor> EditorWidget = EditorWidgetWeak.Pin();
+
+	if (!EditorWidget.IsValid())
+	{
+		return nullptr;
+	}
+
+	return EditorWidget->GetOriginalMaterialModelBase();
+}
+
+void SDMMaterialGlobalSettingsEditor::NotifyPreChange(FProperty* InPropertyAboutToChange)
+{
+	SDMObjectEditorWidgetBase::NotifyPreChange(InPropertyAboutToChange);
+
+	UDynamicMaterialModel* PreviewMaterialModel = Cast<UDynamicMaterialModel>(GetMaterialModelBase());
+
+	if (!PreviewMaterialModel)
+	{
+		return;
+	}
+
+	// We can't know which component it is, so call this on all of them.
+	TArray<UDMMaterialValue*> Values = {
+		PreviewMaterialModel->GetGlobalParameterValue(UDynamicMaterialModel::GlobalOffsetValueName),
+		PreviewMaterialModel->GetGlobalParameterValue(UDynamicMaterialModel::GlobalTilingValueName),
+		PreviewMaterialModel->GetGlobalParameterValue(UDynamicMaterialModel::GlobalRotationValueName)
+	};
+
+	for (UDMMaterialValue* Value : Values)
+	{
+		Value->NotifyPreChange(InPropertyAboutToChange);
+
+		if (UDMMaterialValue* OriginalValue = Cast<UDMMaterialValue>(GetOriginalComponent(Value)))
+		{
+			OriginalValue->NotifyPreChange(InPropertyAboutToChange);
+		}
+	}
+}
+
+void SDMMaterialGlobalSettingsEditor::NotifyPostChange(const FPropertyChangedEvent& InPropertyChangedEvent, FProperty* InPropertyThatChanged)
+{
+	SDMObjectEditorWidgetBase::NotifyPostChange(InPropertyChangedEvent, InPropertyThatChanged);
+
+	UDynamicMaterialModel* PreviewMaterialModel = Cast<UDynamicMaterialModel>(GetMaterialModelBase());
+
+	if (!PreviewMaterialModel)
+	{
+		return;
+	}
+
+	// We called pre-update on all of them, so call post change on all of them (or things break).
+	TArray<UDMMaterialValue*> Values = {
+		PreviewMaterialModel->GetGlobalParameterValue(UDynamicMaterialModel::GlobalOffsetValueName),
+		PreviewMaterialModel->GetGlobalParameterValue(UDynamicMaterialModel::GlobalTilingValueName),
+		PreviewMaterialModel->GetGlobalParameterValue(UDynamicMaterialModel::GlobalRotationValueName)
+	};
+
+	for (UDMMaterialValue* Value : Values)
+	{
+		Value->NotifyPostChange(InPropertyChangedEvent, InPropertyThatChanged);
+
+		if (UDMMaterialValue* OriginalValue = Cast<UDMMaterialValue>(GetOriginalComponent(Value)))
+		{
+			OriginalValue->NotifyPostChange(InPropertyChangedEvent, InPropertyThatChanged);
+		}
+	}
+}
+
+UDMMaterialComponent* SDMMaterialGlobalSettingsEditor::GetOriginalComponent(UDMMaterialComponent* InPreviewComponent) const
+{
+	if (!InPreviewComponent)
+	{
+		return nullptr;
+	}
+
+	TSharedPtr<SDMMaterialEditor> EditorWidget = EditorWidgetWeak.Pin();
+
+	if (!EditorWidget.IsValid())
+	{
+		return nullptr;
+	}
+
+	return EditorWidget->GetOriginalComponent(InPreviewComponent);
+}
+
+void SDMMaterialGlobalSettingsEditor::OnComponentUpdated(UDMMaterialComponent* InComponent, UDMMaterialComponent* InSource,
+	EDMUpdateType InUpdateType)
+{
+	if (UDynamicMaterialEditorSettings* Settings = UDynamicMaterialEditorSettings::Get())
+	{
+		if (Settings->bAutomaticallyCopyParametersToSourceMaterial)
+		{
+			if (InComponent)
+			{
+				if (UDMMaterialComponent* OriginalComponent = GetOriginalComponent(InComponent))
+				{
+					IDMParameterContainer::CopyParametersBetween(InComponent, OriginalComponent);
+					return;
+				}
+			}
+		}
+	}
+
+	if (TSharedPtr<SDMMaterialEditor> EditorWidget = GetEditorWidget())
+	{
+		if (UDynamicMaterialModelBase* MaterialModelBase = EditorWidget->GetOriginalMaterialModelBase())
+		{
+			MaterialModelBase->MarkPreviewModified();
+		}
+	}
+}
+
 TArray<FDMPropertyHandle> SDMMaterialGlobalSettingsEditor::GetPropertyRows()
 {
 	TArray<FDMPropertyHandle> PropertyRows;
+	TSet<UObject*> ProcessedObjects;
 
-	FDMMaterialModelPropertyRowGenerator::AddMaterialModelProperties(
-		SharedThis(this),
-		GetMaterialModelBase(),
-		PropertyRows
-	);
+	FDMComponentPropertyRowGeneratorParams Params(PropertyRows, ProcessedObjects);
+	Params.Owner = this;
+	Params.NotifyHook = this;
+	Params.Object = GetMaterialModelBase();
+	Params.PreviewMaterialModelBase = GetMaterialModelBase();
+	Params.OriginalMaterialModelBase = GetOriginalMaterialModelBase();
+
+	FDMMaterialModelPropertyRowGenerator::AddMaterialModelProperties(Params);
 
 	return PropertyRows;
 }
+
+void SDMMaterialGlobalSettingsEditor::AddDetailTreeRowExtensionWidgets(const TSharedRef<ICustomDetailsView>& InDetailsView,
+	const FDMPropertyHandle& InPropertyRow, const TSharedRef<ICustomDetailsViewItem>& InItem, const TSharedRef<IPropertyHandle>& InPropertyHandle)
+{
+	if (!InPropertyRow.PreviewHandle.PropertyHandle.IsValid())
+	{
+		return;
+	}
+
+	TArray<UObject*> OuterPreview;
+	InPropertyRow.PreviewHandle.PropertyHandle->GetOuterObjects(OuterPreview);
+
+	if (OuterPreview.Num() != 1 || !IsValid(OuterPreview[0]))
+	{
+		return;
+	}
+
+	TArray<FPropertyRowExtensionButton> ExtensionButtons;
+
+	// Reset to Default
+	if (InPropertyRow.ResetToDefaultOverride.IsSet())
+	{
+		FPropertyRowExtensionButton ResetToDefaultButton;
+
+		if (InItem->CreateResetToDefaultButton(ResetToDefaultButton))
+		{
+			ExtensionButtons.Add(MoveTemp(ResetToDefaultButton));
+		}
+	}
+
+	// Global Extensions
+	if (!OuterPreview[0]->IsA<UDynamicMaterialModelBase>())
+	{
+		FPropertyEditorModule& PropertyEditorModule = FModuleManager::Get().GetModuleChecked<FPropertyEditorModule>("PropertyEditor");
+		FOnGenerateGlobalRowExtensionArgs RowExtensionArgs;
+		RowExtensionArgs.OwnerTreeNode = InPropertyRow.OriginalHandle.DetailTreeNode;
+		RowExtensionArgs.PropertyHandle = InPropertyRow.OriginalHandle.PropertyHandle;
+		PropertyEditorModule.GetGlobalRowExtensionDelegate().Broadcast(RowExtensionArgs, ExtensionButtons);
+	}
+
+	bool bValidPropertyHandles = false;
+	TArray<UObject*> OuterOriginal;
+
+	if (InPropertyRow.OriginalHandle.PropertyHandle.IsValid())
+	{
+		InPropertyRow.OriginalHandle.PropertyHandle->GetOuterObjects(OuterOriginal);
+	}
+
+	if (OuterOriginal.Num() == 1 && IsValid(OuterOriginal[0]) && OuterPreview[0]->GetClass() == OuterOriginal[0]->GetClass())
+	{
+		bValidPropertyHandles = true;
+
+		// Sequencer relies on getting the Keyframe Handler via the Details View of the IDetailTreeNode, but is null since there's no Details View here
+		// instead add it manually here
+		if (InPropertyRow.bKeyframeable)
+		{
+			if (TOptional<FPropertyRowExtensionButton> CreateKeyButton = CreateKeyframeButton(
+						InPropertyRow.PreviewHandle.PropertyHandle,
+						InPropertyRow.OriginalHandle.PropertyHandle))
+			{
+				ExtensionButtons.Add(*CreateKeyButton);
+			}
+
+			if (UDMMaterialComponent* PreviewComponent = Cast<UDMMaterialComponent>(OuterPreview[0]))
+			{
+				if (!PreviewComponent->GetOnUpdate().IsBoundToObject(this))
+				{
+					PreviewComponent->GetOnUpdate().AddSP(this, &SDMMaterialGlobalSettingsEditor::OnComponentUpdated);
+				}
+			}
+		}
+	}
+
+	if (!bValidPropertyHandles)
+	{
+		ExtensionButtons.Add(CreateNeedsApplyButton());
+	}
+
+	if (ExtensionButtons.IsEmpty())
+	{
+		InItem->SetOverrideWidget(ECustomDetailsViewWidgetType::Extensions, SNullWidget::NullWidget);
+		return;
+	}
+
+	InItem->SetOverrideWidget(ECustomDetailsViewWidgetType::Extensions, InItem->CreateExtensionButtonWidget(ExtensionButtons));
+}
+
+#undef LOCTEXT_NAMESPACE
diff --git a/Engine/Plugins/Experimental/DynamicMaterial/Source/DynamicMaterialEditor/Private/UI/Widgets/Editor/SDMMaterialGlobalSettingsEditor.h b/Engine/Plugins/Experimental/DynamicMaterial/Source/DynamicMaterialEditor/Private/UI/Widgets/Editor/SDMMaterialGlobalSettingsEditor.h
--- a/Engine/Plugins/Experimental/DynamicMaterial/Source/DynamicMaterialEditor/Private/UI/Widgets/Editor/SDMMaterialGlobalSettingsEditor.h
+++ b/Engine/Plugins/Experimental/DynamicMaterial/Source/DynamicMaterialEditor/Private/UI/Widgets/Editor/SDMMaterialGlobalSettingsEditor.h
@@ -4,7 +4,9 @@
 
 #include "UI/Widgets/Editor/SDMObjectEditorWidgetBase.h"
 
+class UDMMaterialComponent;
 class UDynamicMaterialModelBase;
+enum class EDMUpdateType : uint8;
 
 class SDMMaterialGlobalSettingsEditor : public SDMObjectEditorWidgetBase
 {
@@ -21,8 +23,21 @@ public:
 
 	UDynamicMaterialModelBase* GetMaterialModelBase() const;
 
+	UDynamicMaterialModelBase* GetOriginalMaterialModelBase() const;
+
+	//~ Begin FNotifyHook
+	virtual void NotifyPreChange(FProperty* InPropertyAboutToChange) override;
+	virtual void NotifyPostChange(const FPropertyChangedEvent& InPropertyChangedEvent, FProperty* InPropertyThatChanged) override;
+	//~ End FNotifyHook
+
 protected:
+	UDMMaterialComponent* GetOriginalComponent(UDMMaterialComponent* InPreviewComponent) const;
+
+	void OnComponentUpdated(UDMMaterialComponent* InComponent, UDMMaterialComponent* InSource, EDMUpdateType InUpdateType);
+
 	//~ Begin SDMObjectEditorWidgetBase
 	virtual TArray<FDMPropertyHandle> GetPropertyRows() override;
+	virtual void AddDetailTreeRowExtensionWidgets(const TSharedRef<ICustomDetailsView>& InDetailsView, const FDMPropertyHandle& InPropertyRow,
+		const TSharedRef<ICustomDetailsViewItem>& InItem, const TSharedRef<IPropertyHandle>& InPropertyHandle) override;
 	//~ End SDMObjectEditorWidgetBase
 };
