Material Designer: Apply button

diff --git a/Engine/Plugins/Experimental/DynamicMaterial/Source/DynamicMaterialEditor/Private/UI/Widgets/Editor/SDMMaterialProperties.cpp b/Engine/Plugins/Experimental/DynamicMaterial/Source/DynamicMaterialEditor/Private/UI/Widgets/Editor/SDMMaterialProperties.cpp
--- a/Engine/Plugins/Experimental/DynamicMaterial/Source/DynamicMaterialEditor/Private/UI/Widgets/Editor/SDMMaterialProperties.cpp
+++ b/Engine/Plugins/Experimental/DynamicMaterial/Source/DynamicMaterialEditor/Private/UI/Widgets/Editor/SDMMaterialProperties.cpp
@@ -79,6 +79,167 @@ void SDMMaterialProperties::Validate()
 	}
 }
 
+void SDMMaterialProperties::NotifyPreChange(FProperty* InPropertyAboutToChange)
+{
+	TSharedPtr<SDMMaterialEditor> EditorWidget = EditorWidgetWeak.Pin();
+
+	if (!EditorWidget.IsValid())
+	{
+		return;
+	}
+
+	UDynamicMaterialModelBase* MaterialModelBase = EditorWidget->GetPreviewMaterialModelBase();
+
+	if (!MaterialModelBase)
+	{
+		return;
+	}
+
+	UDynamicMaterialModel* MaterialModel = MaterialModelBase->ResolveMaterialModel();
+
+	if (!MaterialModel)
+	{
+		return;
+	}
+
+	UDynamicMaterialModelEditorOnlyData* EditorOnlyData = UDynamicMaterialModelEditorOnlyData::Get(MaterialModel);
+
+	if (!EditorOnlyData)
+	{
+		return;
+	}
+
+	// Impossible to know which property changed, so notify all of them.
+	UE::DynamicMaterial::ForEachMaterialPropertyType(
+		[EditorOnlyData, InPropertyAboutToChange](EDMMaterialPropertyType InMaterialProperty)
+		{
+			if (UE::DynamicMaterialEditor::Private::IsCustomMaterialProperty(InMaterialProperty))
+			{
+				return EDMIterationResult::Continue;
+			}
+
+			if (UDMMaterialProperty* MaterialProperty = EditorOnlyData->GetMaterialProperty(InMaterialProperty))
+			{
+				UDMMaterialSlot* Slot = EditorOnlyData->GetSlotForEnabledMaterialProperty(InMaterialProperty);
+
+				const bool bIsActive = Slot && MaterialProperty->IsEnabled() && MaterialProperty->IsValidForModel(*EditorOnlyData);
+
+				if (bIsActive)
+				{
+					if (UDMMaterialValue* Value = Cast<UDMMaterialValue>(MaterialProperty->GetComponent(UDynamicMaterialModelEditorOnlyData::AlphaValueName)))
+					{
+						Value->NotifyPreChange(InPropertyAboutToChange);
+					}
+				}
+			}
+
+			return EDMIterationResult::Continue;
+		}
+	);
+}
+
+void SDMMaterialProperties::NotifyPostChange(const FPropertyChangedEvent& InPropertyChangedEvent, FProperty* InPropertyThatChanged)
+{
+	TSharedPtr<SDMMaterialEditor> EditorWidget = EditorWidgetWeak.Pin();
+
+	if (!EditorWidget.IsValid())
+	{
+		return;
+	}
+
+	UDynamicMaterialModelBase* MaterialModelBase = EditorWidget->GetPreviewMaterialModelBase();
+
+	if (!MaterialModelBase)
+	{
+		return;
+	}
+
+	UDynamicMaterialModel* MaterialModel = MaterialModelBase->ResolveMaterialModel();
+
+	if (!MaterialModel)
+	{
+		return;
+	}
+
+	UDynamicMaterialModelEditorOnlyData* EditorOnlyData = UDynamicMaterialModelEditorOnlyData::Get(MaterialModel);
+
+	if (!EditorOnlyData)
+	{
+		return;
+	}
+
+	// All of them were notified of the pre-change, so notify them of the post change.
+	UE::DynamicMaterial::ForEachMaterialPropertyType(
+		[EditorOnlyData, &InPropertyChangedEvent, InPropertyThatChanged](EDMMaterialPropertyType InMaterialProperty)
+		{
+			if (UE::DynamicMaterialEditor::Private::IsCustomMaterialProperty(InMaterialProperty))
+			{
+				return EDMIterationResult::Continue;
+			}
+
+			if (UDMMaterialProperty* MaterialProperty = EditorOnlyData->GetMaterialProperty(InMaterialProperty))
+			{
+				UDMMaterialSlot* Slot = EditorOnlyData->GetSlotForEnabledMaterialProperty(InMaterialProperty);
+
+				const bool bIsActive = Slot && MaterialProperty->IsEnabled() && MaterialProperty->IsValidForModel(*EditorOnlyData);
+
+				if (bIsActive)
+				{
+					if (UDMMaterialValue* Value = Cast<UDMMaterialValue>(MaterialProperty->GetComponent(UDynamicMaterialModelEditorOnlyData::AlphaValueName)))
+					{
+						Value->NotifyPostChange(InPropertyChangedEvent, InPropertyThatChanged);
+					}
+				}
+			}
+
+			return EDMIterationResult::Continue;
+		}
+	);
+}
+
+void SDMMaterialProperties::OnComponentUpdated(UDMMaterialComponent* InComponent, UDMMaterialComponent* InSource, EDMUpdateType InUpdateType)
+{
+	if (UDynamicMaterialEditorSettings* Settings = UDynamicMaterialEditorSettings::Get())
+	{
+		if (Settings->bAutomaticallyCopyParametersToSourceMaterial)
+		{
+			if (InComponent)
+			{
+				if (UDMMaterialComponent* OriginalComponent = GetOriginalComponent(InComponent))
+				{
+					IDMParameterContainer::CopyParametersBetween(InComponent, OriginalComponent);
+					return;
+				}
+			}
+		}
+	}
+
+	if (TSharedPtr<SDMMaterialEditor> EditorWidget = EditorWidgetWeak.Pin())
+	{
+		if (UDynamicMaterialModelBase* MaterialModelBase = EditorWidget->GetOriginalMaterialModelBase())
+		{
+			MaterialModelBase->MarkPreviewModified();
+		}
+	}
+}
+
+UDMMaterialComponent* SDMMaterialProperties::GetOriginalComponent(UDMMaterialComponent* InPreviewComponent) const
+{
+	if (!InPreviewComponent)
+	{
+		return nullptr;
+	}
+
+	TSharedPtr<SDMMaterialEditor> EditorWidget = EditorWidgetWeak.Pin();
+
+	if (!EditorWidget.IsValid())
+	{
+		return nullptr;
+	}
+
+	return EditorWidget->GetOriginalComponent(InPreviewComponent);
+}
+
 TSharedRef<SWidget> SDMMaterialProperties::CreateSlot_Content()
 {
 	TSharedPtr<SDMMaterialEditor> EditorWidget = EditorWidgetWeak.Pin();
@@ -87,7 +248,8 @@ TSharedRef<SWidget> SDMMaterialProperties::CreateSlot_Content()
 	{
 		return SNullWidget::NullWidget;
 	}
-	UDynamicMaterialModelBase* MaterialModelBase = EditorWidget->GetMaterialModelBase();
+
+	UDynamicMaterialModelBase* MaterialModelBase = EditorWidget->GetPreviewMaterialModelBase();
 
 	if (!MaterialModelBase)
 	{
@@ -245,7 +407,7 @@ void SDMMaterialProperties::AddProperty(const TSharedRef<ICustomDetailsView>& In
 		return;
 	}
 
-	UDynamicMaterialModel* MaterialModel = EditorWidget->GetMaterialModel();
+	UDynamicMaterialModel* MaterialModel = EditorWidget->GetPreviewMaterialModel();
 
 	if (!MaterialModel)
 	{
@@ -275,7 +437,7 @@ TSharedRef<SWidget> SDMMaterialProperties::CreatePropertyRow(UDMMaterialProperty
 {
 	// There are all ensured to be valid by the caller of this method
 	TSharedPtr<SDMMaterialEditor> EditorWidget = EditorWidgetWeak.Pin();
-	UDynamicMaterialModelBase* MaterialModelBase = EditorWidget->GetMaterialModelBase();
+	UDynamicMaterialModelBase* MaterialModelBase = EditorWidget->GetPreviewMaterialModelBase();
 	UDynamicMaterialModelEditorOnlyData* EditorOnlyData = UDynamicMaterialModelEditorOnlyData::Get(MaterialModelBase);
 	UDynamicMaterialEditorSettings* Settings = UDynamicMaterialEditorSettings::Get();
 	const EDMMaterialPropertyType MaterialProperty = InProperty->GetMaterialProperty();
@@ -422,7 +584,7 @@ bool SDMMaterialProperties::GetPropertyEnabledEnabled(EDMMaterialPropertyType In
 		return false;
 	}
 
-	UDynamicMaterialModel* MaterialModel = EditorWidget->GetMaterialModel();
+	UDynamicMaterialModel* MaterialModel = EditorWidget->GetPreviewMaterialModel();
 
 	if (!MaterialModel)
 	{
@@ -455,7 +617,7 @@ ECheckBoxState SDMMaterialProperties::GetPropertyEnabledState(EDMMaterialPropert
 		return ECheckBoxState::Unchecked;
 	}
 
-	UDynamicMaterialModel* MaterialModel = EditorWidget->GetMaterialModel();
+	UDynamicMaterialModel* MaterialModel = EditorWidget->GetPreviewMaterialModel();
 
 	if (!MaterialModel)
 	{
@@ -490,7 +652,7 @@ void SDMMaterialProperties::OnPropertyEnabledStateChanged(ECheckBoxState InState
 		return;
 	}
 
-	UDynamicMaterialModel* MaterialModel = EditorWidget->GetMaterialModel();
+	UDynamicMaterialModel* MaterialModel = EditorWidget->GetPreviewMaterialModel();
 
 	if (!MaterialModel)
 	{
@@ -531,7 +693,7 @@ void SDMMaterialProperties::OnPropertyEnabledStateChanged(ECheckBoxState InState
 	EditorWidget->EditProperties();
 }
 
-FReply SDMMaterialProperties::OnPropertyClicked(const FGeometry& InGeometry, const FPointerEvent& InPointerEvent, 
+FReply SDMMaterialProperties::OnPropertyClicked(const FGeometry& InGeometry, const FPointerEvent& InPointerEvent,
 	EDMMaterialPropertyType InMaterialProperty)
 {
 	TSharedPtr<SDMMaterialEditor> EditorWidget = EditorWidgetWeak.Pin();
@@ -541,7 +703,7 @@ FReply SDMMaterialProperties::OnPropertyClicked(const FGeometry& InGeometry, con
 		return FReply::Handled();
 	}
 
-	UDynamicMaterialModel* MaterialModel = EditorWidget->GetMaterialModel();
+	UDynamicMaterialModel* MaterialModel = EditorWidget->GetPreviewMaterialModel();
 
 	if (!MaterialModel)
 	{
@@ -581,9 +743,9 @@ TSharedRef<SWidget> SDMMaterialProperties::CreateGlobalSlider(UDMMaterialPropert
 		return SNullWidget::NullWidget;
 	}
 
-	UDynamicMaterialModelBase* MaterialModelBase = EditorWidget->GetMaterialModelBase();
+	UDynamicMaterialModelBase* PreviewMaterialModelBase = EditorWidget->GetPreviewMaterialModelBase();
 
-	if (!MaterialModelBase)
+	if (!PreviewMaterialModelBase)
 	{
 		return SNullWidget::NullWidget;
 	}
@@ -595,7 +757,7 @@ TSharedRef<SWidget> SDMMaterialProperties::CreateGlobalSlider(UDMMaterialPropert
 		return SNullWidget::NullWidget;
 	}
 
-	if (UDynamicMaterialModelDynamic* MaterialModelDynamic = Cast<UDynamicMaterialModelDynamic>(MaterialModelBase))
+	if (UDynamicMaterialModelDynamic* MaterialModelDynamic = Cast<UDynamicMaterialModelDynamic>(PreviewMaterialModelBase))
 	{
 		AlphaObject = MaterialModelDynamic->GetComponentDynamic(AlphaObject->GetFName());
 
@@ -607,11 +769,14 @@ TSharedRef<SWidget> SDMMaterialProperties::CreateGlobalSlider(UDMMaterialPropert
 
 	TSharedPtr<IDetailKeyframeHandler> KeyframeHandler = nullptr;
 
-	if (UWorld* World = MaterialModelBase->GetWorld())
+	if (UDynamicMaterialModelBase* OriginalMaterialModelBase = EditorWidget->GetOriginalMaterialModelBase())
 	{
-		if (const UDMWorldSubsystem* const WorldSubsystem = World->GetSubsystem<UDMWorldSubsystem>())
+		if (UWorld* World = OriginalMaterialModelBase->GetWorld())
 		{
-			KeyframeHandler = WorldSubsystem->GetKeyframeHandler();
+			if (const UDMWorldSubsystem* const WorldSubsystem = World->GetSubsystem<UDMWorldSubsystem>())
+			{
+				KeyframeHandler = WorldSubsystem->GetKeyframeHandler();
+			}
 		}
 	}
 
@@ -624,9 +789,17 @@ TSharedRef<SWidget> SDMMaterialProperties::CreateGlobalSlider(UDMMaterialPropert
 	TSharedRef<ICustomDetailsView> DetailsView = ICustomDetailsViewModule::Get().CreateCustomDetailsView(Args);
 	FCustomDetailsViewItemId RootId = DetailsView->GetRootItem()->GetItemId();
 
-	FDMPropertyHandle PropertyHandle = FDMWidgetLibrary::Get().GetPropertyHandle(this, AlphaObject, UDMMaterialValue::ValueName);
+	FDMPropertyHandleGenerateParams Params;
+	Params.Widget = this;
+	Params.NotifyHook = this;
+	Params.Object = AlphaObject;
+	Params.PreviewMaterialModelBase = EditorWidget->GetPreviewMaterialModelBase();
+	Params.OriginalMaterialModelBase = EditorWidget->GetOriginalMaterialModelBase();
+	Params.PropertyName = UDMMaterialValue::ValueName;
 
-	TSharedRef<ICustomDetailsViewItem> AlphaItem = DetailsView->CreateDetailTreeItem(DetailsView->GetRootItem(), PropertyHandle.DetailTreeNode.ToSharedRef());
+	FDMPropertyHandle PropertyHandle = FDMWidgetLibrary::Get().GetPropertyHandle(Params);
+
+	TSharedRef<ICustomDetailsViewItem> AlphaItem = DetailsView->CreateDetailTreeItem(DetailsView->GetRootItem(), PropertyHandle.PreviewHandle.DetailTreeNode.ToSharedRef());
 
 	if (UDMMaterialValue* AlphaValue = Cast<UDMMaterialValue>(AlphaObject))
 	{
@@ -634,6 +807,11 @@ TSharedRef<SWidget> SDMMaterialProperties::CreateGlobalSlider(UDMMaterialPropert
 			FIsResetToDefaultVisible::CreateUObject(AlphaValue, &UDMMaterialValue::CanResetToDefault),
 			FResetToDefaultHandler::CreateUObject(AlphaValue, &UDMMaterialValue::ResetToDefault)
 		));
+
+		if (!AlphaValue->GetOnUpdate().IsBoundToObject(this))
+		{
+			AlphaValue->GetOnUpdate().AddSP(this, &SDMMaterialProperties::OnComponentUpdated);
+		}
 	}
 	else if (UDMMaterialValueDynamic* AlphaValueDynamic = Cast<UDMMaterialValueDynamic>(AlphaObject))
 	{
@@ -641,6 +819,11 @@ TSharedRef<SWidget> SDMMaterialProperties::CreateGlobalSlider(UDMMaterialPropert
 			FIsResetToDefaultVisible::CreateUObject(AlphaValueDynamic, &UDMMaterialValueDynamic::CanResetToDefault),
 			FResetToDefaultHandler::CreateUObject(AlphaValueDynamic, &UDMMaterialValueDynamic::ResetToDefault)
 		));
+
+		if (!AlphaValueDynamic->GetOnUpdate().IsBoundToObject(this))
+		{
+			AlphaValueDynamic->GetOnUpdate().AddSP(this, &SDMMaterialProperties::OnComponentUpdated);
+		}
 	}
 
 	AlphaItem->MakeWidget(nullptr, SharedThis(this));
@@ -661,7 +844,6 @@ TSharedRef<SWidget> SDMMaterialProperties::CreateGlobalSlider(UDMMaterialPropert
 		}
 	}
 
-
 	const FText PropertyGlobalSliderToolTipFormat = LOCTEXT("PropertyGlobalSliderToolTipFormat", "Change the global {0} value.");
 
 	const FText PropertyGlobalSliderToolTip = FText::Format(
@@ -717,7 +899,7 @@ void SDMMaterialProperties::OnExpansionStateChanged(const TSharedRef<ICustomDeta
 		return;
 	}
 
-	UDynamicMaterialModelBase* MaterialModelBase = EditorWidget->GetMaterialModelBase();
+	UDynamicMaterialModelBase* MaterialModelBase = EditorWidget->GetPreviewMaterialModelBase();
 
 	if (!MaterialModelBase)
 	{
@@ -767,7 +949,7 @@ bool SDMMaterialProperties::OnAssetDraggedOver(TArrayView<FAssetData> InAssets,
 		return false;
 	}
 
-	UDynamicMaterialModel* MaterialModel = EditorWidget->GetMaterialModel();
+	UDynamicMaterialModel* MaterialModel = EditorWidget->GetPreviewMaterialModel();
 
 	if (!MaterialModel)
 	{
@@ -844,7 +1026,7 @@ void SDMMaterialProperties::HandleDrop_Texture(UTexture* InTexture, EDMMaterialP
 		return;
 	}
 
-	UDynamicMaterialModel* MaterialModel = EditorWidget->GetMaterialModel();
+	UDynamicMaterialModel* MaterialModel = EditorWidget->GetPreviewMaterialModel();
 
 	if (!MaterialModel)
 	{
diff --git a/Engine/Plugins/Experimental/DynamicMaterial/Source/DynamicMaterialEditor/Private/UI/Widgets/Editor/SDMMaterialProperties.h b/Engine/Plugins/Experimental/DynamicMaterial/Source/DynamicMaterialEditor/Private/UI/Widgets/Editor/SDMMaterialProperties.h
--- a/Engine/Plugins/Experimental/DynamicMaterial/Source/DynamicMaterialEditor/Private/UI/Widgets/Editor/SDMMaterialProperties.h
+++ b/Engine/Plugins/Experimental/DynamicMaterial/Source/DynamicMaterialEditor/Private/UI/Widgets/Editor/SDMMaterialProperties.h
@@ -2,6 +2,7 @@
 
 #pragma once
 
+#include "Misc/NotifyHook.h"
 #include "Widgets/SCompoundWidget.h"
 
 #include "Containers/Array.h"
@@ -13,14 +14,16 @@ class SBox;
 class SDMMaterialComponentPreview;
 class SDMMaterialEditor;
 class SVerticalBox;
+class UDMMaterialComponent;
 class UDMMaterialProperty;
 class UDMMaterialSlot;
 class UTexture;
 enum class ECheckBoxState : uint8;
 enum class EDMMaterialPropertyType : uint8;
+enum class EDMUpdateType : uint8;
 struct FPropertyChangedEvent;
 
-class SDMMaterialProperties : public SCompoundWidget
+class SDMMaterialProperties : public SCompoundWidget, public FNotifyHook
 {
 	SLATE_DECLARE_WIDGET(SDMMaterialProperties, SCompoundWidget)
 
@@ -34,6 +37,11 @@ public:
 
 	void Validate();
 
+	//~ Begin FNotifyHook
+	virtual void NotifyPreChange(FProperty* InPropertyAboutToChange) override;
+	virtual void NotifyPostChange(const FPropertyChangedEvent& InPropertyChangedEvent, FProperty* InPropertyThatChanged) override;
+	//~ End FNotifyHook
+
 protected:
 	TWeakPtr<SDMMaterialEditor> EditorWidgetWeak;
 
@@ -75,4 +83,8 @@ protected:
 	void OnAssetsDropped(const FDragDropEvent& InDragDropEvent, TArrayView<FAssetData> InAssets, EDMMaterialPropertyType InMaterialProperty);
 
 	void HandleDrop_Texture(UTexture* InTexture, EDMMaterialPropertyType InMaterialProperty);
+
+	void OnComponentUpdated(UDMMaterialComponent* InComponent, UDMMaterialComponent* InSource, EDMUpdateType InUpdateType);
+
+	UDMMaterialComponent* GetOriginalComponent(UDMMaterialComponent* InPreviewComponent) const;
 };
diff --git a/Engine/Plugins/Experimental/DynamicMaterial/Source/DynamicMaterialEditor/Private/UI/Widgets/Editor/SDMMaterialSlotEditor.cpp b/Engine/Plugins/Experimental/DynamicMaterial/Source/DynamicMaterialEditor/Private/UI/Widgets/Editor/SDMMaterialSlotEditor.cpp
--- a/Engine/Plugins/Experimental/DynamicMaterial/Source/DynamicMaterialEditor/Private/UI/Widgets/Editor/SDMMaterialSlotEditor.cpp
+++ b/Engine/Plugins/Experimental/DynamicMaterial/Source/DynamicMaterialEditor/Private/UI/Widgets/Editor/SDMMaterialSlotEditor.cpp
@@ -21,6 +21,7 @@
 #include "DetailLayoutBuilder.h"
 #include "DMTextureSet.h"
 #include "DMWorldSubsystem.h"
+#include "DynamicMaterialEditorSettings.h"
 #include "DynamicMaterialEditorStyle.h"
 #include "DynamicMaterialModule.h"
 #include "Engine/Texture.h"
@@ -81,7 +82,7 @@ void SDMMaterialSlotEditor::Construct(const FArguments& InArgs, const TSharedRef
 
 	SetCanTick(false);
 
-	bIsDynamic = !Cast<UDynamicMaterialModel>(InEditorWidget->GetMaterialModelBase());
+	bIsDynamic = !Cast<UDynamicMaterialModel>(InEditorWidget->GetPreviewMaterialModelBase());
 
 	ContentSlot = TDMWidgetSlot<SWidget>(SharedThis(this), 0, SNullWidget::NullWidget);
 
@@ -498,6 +499,22 @@ void SDMMaterialSlotEditor::TriggerEffectSelectionChange(const TSharedRef<SDMMat
 	OnEffectSelectionChanged.Broadcast(InEffectView, InEffect);
 }
 
+void SDMMaterialSlotEditor::NotifyPreChange(FProperty* PropertyAboutToChange)
+{
+	if (UDMMaterialValue* OpacityValue = LayerOpacityValueWeak.Get())
+	{
+		OpacityValue->NotifyPreChange(PropertyAboutToChange);
+	}
+}
+
+void SDMMaterialSlotEditor::NotifyPostChange(const FPropertyChangedEvent& PropertyChangedEvent, FProperty* PropertyThatChanged)
+{
+	if (UDMMaterialValue* OpacityValue = LayerOpacityValueWeak.Get())
+	{
+		OpacityValue->NotifyPostChange(PropertyChangedEvent, PropertyThatChanged);
+	}
+}
+
 TSharedRef<SWidget> SDMMaterialSlotEditor::CreateSlot_Container()
 {
 	SVerticalBox::FSlot* SettingsSlotPtr = nullptr;
@@ -678,9 +695,9 @@ TSharedRef<SWidget> SDMMaterialSlotEditor::CreateSlot_LayerOpacity()
 						{
 							if (TSharedPtr<SDMMaterialEditor> EditorWidget = EditorWidgetWeak.Pin())
 							{
-								if (UDynamicMaterialModelBase* MaterialModelBase = EditorWidget->GetMaterialModelBase())
+								if (UDynamicMaterialModelBase* OriginalMaterialModelBase = EditorWidget->GetOriginalMaterialModelBase())
 								{
-									World = MaterialModelBase->GetWorld();
+									World = OriginalMaterialModelBase->GetWorld();
 								}
 							}
 						}
@@ -693,7 +710,22 @@ TSharedRef<SWidget> SDMMaterialSlotEditor::CreateSlot_LayerOpacity()
 							}
 						}
 
-						FDMPropertyHandle PropertyHandle = FDMWidgetLibrary::Get().GetPropertyHandle(this, OpacityValue, UDMMaterialValue::ValueName);
+						if (!OpacityValue->GetOnUpdate().IsBoundToObject(this))
+						{
+							OpacityValue->GetOnUpdate().AddSP(this, &SDMMaterialSlotEditor::OnOpacityUpdated);
+						}
+
+						TSharedPtr<SDMMaterialEditor> EditorWidget = GetEditorWidget();
+
+						FDMPropertyHandleGenerateParams Params;
+						Params.Widget = this;
+						Params.NotifyHook = this;
+						Params.Object = OpacityValue;
+						Params.PreviewMaterialModelBase = EditorWidget.IsValid() ? EditorWidget->GetPreviewMaterialModelBase() : nullptr;
+						Params.OriginalMaterialModelBase = EditorWidget.IsValid() ? EditorWidget->GetOriginalMaterialModelBase() : nullptr;
+						Params.PropertyName = UDMMaterialValue::ValueName;
+
+						FDMPropertyHandle PropertyHandle = FDMWidgetLibrary::Get().GetPropertyHandle(Params);
 
 						FCustomDetailsViewArgs Args;
 						Args.KeyframeHandler = KeyframeHandler;
@@ -704,7 +736,7 @@ TSharedRef<SWidget> SDMMaterialSlotEditor::CreateSlot_LayerOpacity()
 						TSharedRef<ICustomDetailsView> DetailsView = ICustomDetailsViewModule::Get().CreateCustomDetailsView(Args);
 						FCustomDetailsViewItemId RootId = DetailsView->GetRootItem()->GetItemId();
 
-						LayerOpacityItem = DetailsView->CreateDetailTreeItem(DetailsView->GetRootItem(), PropertyHandle.DetailTreeNode.ToSharedRef());
+						LayerOpacityItem = DetailsView->CreateDetailTreeItem(DetailsView->GetRootItem(), PropertyHandle.PreviewHandle.DetailTreeNode.ToSharedRef());
 
 						LayerOpacityItem->SetResetToDefaultOverride(FResetToDefaultOverride::Create(
 							FIsResetToDefaultVisible::CreateUObject(OpacityValue, &UDMMaterialValue::CanResetToDefault),
@@ -954,6 +986,26 @@ FReply SDMMaterialSlotEditor::OnLayerRowButtonsRemoveClicked()
 	return FReply::Handled();
 }
 
+void SDMMaterialSlotEditor::OnOpacityUpdated(UDMMaterialComponent* InComponent, UDMMaterialComponent* InSource, EDMUpdateType InUpdateType)
+{
+	if (UDynamicMaterialEditorSettings* Settings = UDynamicMaterialEditorSettings::Get())
+	{
+		if (Settings->bAutomaticallyCopyParametersToSourceMaterial)
+		{
+			if (InComponent)
+			{
+				if (TSharedPtr<SDMMaterialEditor> EditorWidget = GetEditorWidget())
+				{
+					if (UDMMaterialComponent* OriginalComponent = EditorWidget->GetOriginalComponent(InComponent))
+					{
+						IDMParameterContainer::CopyParametersBetween(InComponent, OriginalComponent);
+					}
+				}
+			}
+		}
+	}
+}
+
 bool SDMMaterialSlotEditor::OnAreAssetsAcceptableForDrop(TArrayView<FAssetData> InAssets)
 {
 	TSharedPtr<SDMMaterialEditor> EditorWidget = GetEditorWidget();
@@ -963,9 +1015,9 @@ bool SDMMaterialSlotEditor::OnAreAssetsAcceptableForDrop(TArrayView<FAssetData>
 		return false;
 	}
 
-	UDynamicMaterialModelBase* MaterialModelBase = EditorWidget->GetMaterialModelBase();
+	UDynamicMaterialModelBase* PreviewMaterialModelBase = EditorWidget->GetPreviewMaterialModelBase();
 
-	if (MaterialModelBase->IsA<UDynamicMaterialModelDynamic>())
+	if (PreviewMaterialModelBase->IsA<UDynamicMaterialModelDynamic>())
 	{
 		return false;
 	}
@@ -1023,9 +1075,9 @@ void SDMMaterialSlotEditor::OnAssetsDropped(const FDragDropEvent& InDragDropEven
 		return;
 	}
 
-	UDynamicMaterialModelBase* MaterialModelBase = EditorWidget->GetMaterialModelBase();
+	UDynamicMaterialModelBase* PreviewMaterialModelBase = EditorWidget->GetPreviewMaterialModelBase();
 
-	if (MaterialModelBase->IsA<UDynamicMaterialModelDynamic>())
+	if (PreviewMaterialModelBase->IsA<UDynamicMaterialModelDynamic>())
 	{
 		return;
 	}
diff --git a/Engine/Plugins/Experimental/DynamicMaterial/Source/DynamicMaterialEditor/Private/UI/Widgets/Editor/SDMMaterialSlotEditor.h b/Engine/Plugins/Experimental/DynamicMaterial/Source/DynamicMaterialEditor/Private/UI/Widgets/Editor/SDMMaterialSlotEditor.h
--- a/Engine/Plugins/Experimental/DynamicMaterial/Source/DynamicMaterialEditor/Private/UI/Widgets/Editor/SDMMaterialSlotEditor.h
+++ b/Engine/Plugins/Experimental/DynamicMaterial/Source/DynamicMaterialEditor/Private/UI/Widgets/Editor/SDMMaterialSlotEditor.h
@@ -2,6 +2,7 @@
 
 #pragma once
 
+#include "Misc/NotifyHook.h"
 #include "Widgets/SCompoundWidget.h"
 
 #include "UI/Utils/DMWidgetSlot.h"
@@ -14,6 +15,7 @@ class SDMMaterialSlotLayerItem;
 class SDMMaterialSlotLayerView;
 class SHorizontalBox;
 class SVerticalBox;
+class UDMMaterialComponent;
 class UDMMaterialEffect;
 class UDMMaterialLayerObject;
 class UDMMaterialSlot;
@@ -22,9 +24,10 @@ class UDMMaterialValueFloat1;
 class UDMTextureSet;
 class UMaterialFunctionInterface;
 class UTexture;
+enum class EDMUpdateType : uint8;
 struct FDMMaterialLayerReference;
 
-class SDMMaterialSlotEditor : public SCompoundWidget
+class SDMMaterialSlotEditor : public SCompoundWidget, public FNotifyHook
 {
 	SLATE_DECLARE_WIDGET(SDMMaterialSlotEditor, SCompoundWidget)
 
@@ -103,6 +106,11 @@ public:
 	FOnEffectSelectionChanged::RegistrationType& GetOnEffectSelectionChanged() { return OnEffectSelectionChanged; }
 	void TriggerEffectSelectionChange(const TSharedRef<SDMMaterialSlotLayerEffectView>& InEffectView, UDMMaterialEffect* InEffect);
 
+	//~ Begin FNotifyHook
+	virtual void NotifyPreChange(FProperty* PropertyAboutToChange) override;
+	virtual void NotifyPostChange(const FPropertyChangedEvent& PropertyChangedEvent, FProperty* PropertyThatChanged) override;
+	//~ End FNotifyHook
+
 protected:
 	TWeakPtr<SDMMaterialEditor> EditorWidgetWeak;
 	TWeakObjectPtr<UDMMaterialSlot> MaterialSlotWeak;
@@ -146,6 +154,8 @@ protected:
 	bool GetLayerRowsButtonsCanRemove() const;
 	FReply OnLayerRowButtonsRemoveClicked();
 
+	void OnOpacityUpdated(UDMMaterialComponent* InComponent, UDMMaterialComponent* InSource, EDMUpdateType InUpdateType);
+
 	/** Drag and drop. */
 	bool OnAreAssetsAcceptableForDrop(TArrayView<FAssetData> InAssets);
 
diff --git a/Engine/Plugins/Experimental/DynamicMaterial/Source/DynamicMaterialEditor/Private/UI/Widgets/Editor/SDMObjectEditorWidgetBase.cpp b/Engine/Plugins/Experimental/DynamicMaterial/Source/DynamicMaterialEditor/Private/UI/Widgets/Editor/SDMObjectEditorWidgetBase.cpp
--- a/Engine/Plugins/Experimental/DynamicMaterial/Source/DynamicMaterialEditor/Private/UI/Widgets/Editor/SDMObjectEditorWidgetBase.cpp
+++ b/Engine/Plugins/Experimental/DynamicMaterial/Source/DynamicMaterialEditor/Private/UI/Widgets/Editor/SDMObjectEditorWidgetBase.cpp
@@ -4,14 +4,17 @@
 
 #include "CustomDetailsViewArgs.h"
 #include "CustomDetailsViewModule.h"
+#include "CustomDetailsViewSequencer.h"
 #include "DetailLayoutBuilder.h"
 #include "DMWorldSubsystem.h"
 #include "Engine/World.h"
 #include "ICustomDetailsView.h"
+#include "IDetailKeyframeHandler.h"
 #include "Items/ICustomDetailsViewCustomCategoryItem.h"
 #include "Items/ICustomDetailsViewCustomItem.h"
 #include "Items/ICustomDetailsViewItem.h"
 #include "Model/DynamicMaterialModelBase.h"
+#include "ScopedTransaction.h"
 #include "UI/Utils/DMWidgetLibrary.h"
 #include "UI/Widgets/SDMMaterialEditor.h"
 #include "Widgets/Layout/SBox.h"
@@ -40,9 +43,9 @@ void SDMObjectEditorWidgetBase::Construct(const FArguments& InArgs, const TShare
 
 	UObject* WorldContext = InObject;
 
-	if (!WorldContext)
+	if (!WorldContext || !WorldContext->GetWorld())
 	{
-		WorldContext = InEditorWidget->GetMaterialModelBase();
+		WorldContext = InEditorWidget->GetOriginalMaterialModelBase();
 	}
 
 	if (WorldContext)
@@ -67,6 +70,14 @@ void SDMObjectEditorWidgetBase::Validate()
 	}
 }
 
+void SDMObjectEditorWidgetBase::NotifyPreChange(FProperty* InPropertyAboutToChange)
+{	
+}
+
+void SDMObjectEditorWidgetBase::NotifyPostChange(const FPropertyChangedEvent& InPropertyChangedEvent, FProperty* InPropertyThatChanged)
+{
+}
+
 TSharedRef<SWidget> SDMObjectEditorWidgetBase::CreateWidget()
 {
 	FDMWidgetLibrary::Get().ClearPropertyHandles(this);
@@ -89,7 +100,7 @@ TSharedRef<SWidget> SDMObjectEditorWidgetBase::CreateWidget()
 	{
 		const bool bHasValidCustomWidget = PropertyRow.ValueWidget.IsValid() && !PropertyRow.ValueName.IsNone() && PropertyRow.NameOverride.IsSet();
 
-		if (!PropertyRow.DetailTreeNode && !bHasValidCustomWidget)
+		if (!PropertyRow.PreviewHandle.DetailTreeNode && !bHasValidCustomWidget)
 		{
 			continue;
 		}
@@ -117,7 +128,7 @@ TSharedRef<SWidget> SDMObjectEditorWidgetBase::CreateWidget()
 		{
 			AddCustomRow(DetailsView, CategoryItem, Position, PropertyRow);
 		}
-		else if (PropertyRow.DetailTreeNode)
+		else if (PropertyRow.PreviewHandle.DetailTreeNode)
 		{
 			AddDetailTreeRow(DetailsView, CategoryItem, Position, PropertyRow);
 		}		
@@ -155,10 +166,10 @@ TSharedRef<ICustomDetailsViewItem> SDMObjectEditorWidgetBase::GetCategoryForRow(
 {
 	FName CategoryName = InPropertyRow.CategoryOverrideName;
 
-	if (CategoryName.IsNone() && InPropertyRow.PropertyHandle.IsValid())
+	if (CategoryName.IsNone() && InPropertyRow.PreviewHandle.PropertyHandle.IsValid())
 	{
 		// Sub category (possibly)
-		if (TSharedPtr<IPropertyHandle> SubCategoryProperty = InPropertyRow.PropertyHandle->GetParentHandle())
+		if (TSharedPtr<IPropertyHandle> SubCategoryProperty = InPropertyRow.PreviewHandle.PropertyHandle->GetParentHandle())
 		{
 			if (SubCategoryProperty->IsCategoryHandle())
 			{
@@ -205,7 +216,7 @@ TSharedRef<ICustomDetailsViewItem> SDMObjectEditorWidgetBase::GetCategoryForRow(
 void SDMObjectEditorWidgetBase::AddDetailTreeRow(const TSharedRef<ICustomDetailsView>& InDetailsView,
 	const TSharedRef<ICustomDetailsViewItem>& InParent, ECustomDetailsTreeInsertPosition InPosition, const FDMPropertyHandle& InPropertyRow)
 {
-	TSharedRef<ICustomDetailsViewItem> Item = InDetailsView->CreateDetailTreeItem(InParent, InPropertyRow.DetailTreeNode.ToSharedRef());
+	TSharedRef<ICustomDetailsViewItem> Item = InDetailsView->CreateDetailTreeItem(InParent, InPropertyRow.PreviewHandle.DetailTreeNode.ToSharedRef());
 
 	if (InPropertyRow.NameOverride.IsSet())
 	{
@@ -241,11 +252,19 @@ void SDMObjectEditorWidgetBase::AddDetailTreeRow(const TSharedRef<ICustomDetails
 		Item->SetValueWidgetWidthOverride(InPropertyRow.MaxWidth);
 	}
 
+	Item->SetCreateChildItemDelegate(FOnCustomDetailsViewGenerateChildItem::CreateSP(this, &SDMObjectEditorWidgetBase::CreateChildItem, InPropertyRow));
+
+	if (InPropertyRow.PreviewHandle.PropertyHandle.IsValid())
+	{
+		AddDetailTreeRowExtensionWidgets(InDetailsView, InPropertyRow, Item, InPropertyRow.PreviewHandle.PropertyHandle.ToSharedRef());
+	}
+
 	InDetailsView->ExtendTree(InParent->GetItemId(), InPosition, Item);
 }
 
 void SDMObjectEditorWidgetBase::AddCustomRow(const TSharedRef<ICustomDetailsView>& InDetailsView, 
-	const TSharedRef<ICustomDetailsViewItem>& InParent, ECustomDetailsTreeInsertPosition InPosition, const FDMPropertyHandle& InPropertyRow)
+	const TSharedRef<ICustomDetailsViewItem>& InParent, ECustomDetailsTreeInsertPosition InPosition, 
+	const FDMPropertyHandle& InPropertyRow)
 {
 	TSharedPtr<ICustomDetailsViewCustomItem> Item = InDetailsView->CreateCustomItem(
 		InParent,
@@ -274,6 +293,11 @@ void SDMObjectEditorWidgetBase::AddCustomRow(const TSharedRef<ICustomDetailsView
 		Item->AsItem()->SetValueWidgetWidthOverride(InPropertyRow.MaxWidth);
 	}
 
+	if (InPropertyRow.PreviewHandle.PropertyHandle.IsValid())
+	{
+		AddDetailTreeRowExtensionWidgets(InDetailsView, InPropertyRow, Item->AsItem(), InPropertyRow.PreviewHandle.PropertyHandle.ToSharedRef());
+	}
+
 	InDetailsView->ExtendTree(InParent->GetItemId(), InPosition, Item->AsItem());
 }
 
@@ -294,4 +318,261 @@ void SDMObjectEditorWidgetBase::OnExpansionStateChanged(const TSharedRef<ICustom
 	FDMWidgetLibrary::Get().SetExpansionState(ObjectWeak.Get(), *ItemId.GetItemName(), bInExpansionState);
 }
 
+void SDMObjectEditorWidgetBase::AddDetailTreeRowExtensionWidgets(const TSharedRef<ICustomDetailsView>& InDetailsView,
+	const FDMPropertyHandle& InPropertyRow, const TSharedRef<ICustomDetailsViewItem>& InItem, const TSharedRef<IPropertyHandle>& InPropertyHandle)
+{
+}
+
+TOptional<FPropertyRowExtensionButton> SDMObjectEditorWidgetBase::CreateKeyframeButton(TSharedPtr<IPropertyHandle> InPreviewPropertyHandle,
+	TSharedPtr<IPropertyHandle> InOriginalPropertyHandle)
+{
+	if (!InPreviewPropertyHandle.IsValid() || !InOriginalPropertyHandle.IsValid()
+		|| !InPreviewPropertyHandle->IsValidHandle() || !InOriginalPropertyHandle->IsValidHandle())
+	{
+		return {};
+	}
+
+	TArray<FPropertyRowExtensionButton> ExtensionButtons;
+	FCustomDetailsViewSequencerUtils::CreateSequencerExtensionButton(KeyframeHandler, InPreviewPropertyHandle, ExtensionButtons);
+
+	if (ExtensionButtons.IsEmpty())
+	{
+		return {};
+	}
+
+	ExtensionButtons[0].Icon = TAttribute<FSlateIcon>::CreateSP(
+		this,
+		&SDMObjectEditorWidgetBase::GetCreateKeyIcon,
+		InPreviewPropertyHandle.ToWeakPtr(),
+		InOriginalPropertyHandle.ToWeakPtr()
+	);
+
+	ExtensionButtons[0].UIAction.CanExecuteAction.BindSP(
+		this,
+		&SDMObjectEditorWidgetBase::CanCreateKeyFrame,
+		InPreviewPropertyHandle.ToWeakPtr(),
+		InOriginalPropertyHandle.ToWeakPtr()
+	);
+
+	ExtensionButtons[0].UIAction.IsActionVisibleDelegate.BindSP(
+		this,
+		&SDMObjectEditorWidgetBase::CanCreateKeyFrame,
+		InPreviewPropertyHandle.ToWeakPtr(),
+		InOriginalPropertyHandle.ToWeakPtr()
+	);
+
+	ExtensionButtons[0].UIAction.ExecuteAction.BindSP(
+		this,
+		&SDMObjectEditorWidgetBase::CreateKeyFrame,
+		InPreviewPropertyHandle.ToWeakPtr(),
+		InOriginalPropertyHandle.ToWeakPtr()
+	);
+
+	return ExtensionButtons[0];
+}
+
+FSlateIcon SDMObjectEditorWidgetBase::GetCreateKeyIcon(TWeakPtr<IPropertyHandle> InPreviewPropertyHandleWeak,
+	TWeakPtr<IPropertyHandle> InOriginalPropertyHandleWeak) const
+{
+	if (!CanCreateKeyFrame(InPreviewPropertyHandleWeak, InOriginalPropertyHandleWeak))
+	{
+		return FSlateIcon();
+	}
+
+	EPropertyKeyedStatus KeyedStatus = EPropertyKeyedStatus::NotKeyed;
+
+	if (KeyframeHandler.IsValid())
+	{
+		KeyedStatus = KeyframeHandler->GetPropertyKeyedStatus(*InOriginalPropertyHandleWeak.Pin());
+	}
+
+	FName FoundIcon;
+
+	switch (KeyedStatus)
+	{
+		default:
+		case EPropertyKeyedStatus::NotKeyed:
+			FoundIcon = TEXT("Sequencer.KeyedStatus.NotKeyed");
+			break;
+
+		case EPropertyKeyedStatus::KeyedInFrame:
+			FoundIcon = TEXT("Sequencer.KeyedStatus.Keyed");
+			break;
+
+		case EPropertyKeyedStatus::KeyedInOtherFrame:
+			FoundIcon = TEXT("Sequencer.KeyedStatus.Animated");
+			break;
+
+		case EPropertyKeyedStatus::PartiallyKeyed:
+			FoundIcon = TEXT("Sequencer.KeyedStatus.PartialKey");
+			break;
+	}
+
+	return FSlateIcon(FAppStyle::GetAppStyleSetName(), FoundIcon);
+}
+
+bool SDMObjectEditorWidgetBase::CanCreateKeyFrame(TWeakPtr<IPropertyHandle> InPreviewPropertyHandleWeak,
+	TWeakPtr<IPropertyHandle> InOriginalPropertyHandleWeak) const
+{
+	if (!KeyframeHandler.IsValid())
+	{
+		return false;
+	}
+
+	TSharedPtr<IPropertyHandle> PreviewPropertyHandle = InPreviewPropertyHandleWeak.Pin();
+
+	if (!PreviewPropertyHandle.IsValid())
+	{
+		return false;
+	}
+
+	TSharedPtr<IPropertyHandle> OriginalPropertyHandle = InOriginalPropertyHandleWeak.Pin();
+
+	if (!OriginalPropertyHandle.IsValid())
+	{
+		return false;
+	}
+
+	return PreviewPropertyHandle->IsValidHandle() && OriginalPropertyHandle->IsValidHandle();
+}
+
+void SDMObjectEditorWidgetBase::CreateKeyFrame(TWeakPtr<IPropertyHandle> InPreviewPropertyHandleWeak,
+	TWeakPtr<IPropertyHandle> InOriginalPropertyHandleWeak)
+{
+	TSharedPtr<IPropertyHandle> PreviewPropertyHandle = InPreviewPropertyHandleWeak.Pin();
+	TSharedPtr<IPropertyHandle> OriginalPropertyHandle = InOriginalPropertyHandleWeak.Pin();
+
+	if (!PreviewPropertyHandle.IsValid() || !OriginalPropertyHandle.IsValid()
+		|| !PreviewPropertyHandle->IsValidHandle() || !OriginalPropertyHandle->IsValidHandle())
+	{
+		return;
+	}
+
+	TArray<UObject*> OriginalObjects;
+	OriginalPropertyHandle->GetOuterObjects(OriginalObjects);
+
+	if (OriginalObjects.IsEmpty())
+	{
+		return;
+	}
+
+	FString Value;
+	PreviewPropertyHandle->GetValueAsFormattedString(Value);
+
+	if (Value.IsEmpty())
+	{
+		return;
+	}
+
+	FScopedTransaction Transaction(LOCTEXT("CreateKeyFrame", "Create Key Frame"));
+
+	for (UObject* Object : OriginalObjects)
+	{
+		if (IsValid(Object))
+		{
+			Object->Modify();
+		}
+	}
+
+	OriginalPropertyHandle->SetValueFromFormattedString(Value);
+	KeyframeHandler->OnKeyPropertyClicked(*OriginalPropertyHandle);
+}
+
+FPropertyRowExtensionButton SDMObjectEditorWidgetBase::CreateNeedsApplyButton() const
+{
+	FPropertyRowExtensionButton NeedsApplyButton;
+	NeedsApplyButton.Label = LOCTEXT("NeedsApply", "Needs Apply");
+	NeedsApplyButton.Icon = FSlateIcon(FAppStyle::Get().GetStyleSetName(), TEXT("Icons.ErrorWithColor"));
+	NeedsApplyButton.ToolTip = LOCTEXT("NeedsApplyTooltip", "This property only exists in the preview material. Apply changes to add this to the source material.");
+
+	return NeedsApplyButton;
+}
+
+TSharedPtr<ICustomDetailsViewItem> SDMObjectEditorWidgetBase::CreateChildItem(const TSharedRef<ICustomDetailsView>& InDetailsView,
+	const TSharedPtr<ICustomDetailsViewItem>& InParent, const TSharedRef<IDetailTreeNode>& InChildNode, FDMPropertyHandle InOriginalRow)
+{
+	if (!InParent.IsValid())
+	{
+		return nullptr;
+	}
+
+	TSharedPtr<IPropertyHandle> PropertyHandle = InChildNode->CreatePropertyHandle();
+
+	if (!PropertyHandle)
+	{
+		return nullptr;
+	}
+
+	// Only customise struct properties
+	if (PropertyHandle->GetPropertyPath().Find(TEXT("->")) == INDEX_NONE)
+	{
+		return nullptr;
+	}
+
+	TSharedPtr<SDMMaterialEditor> EditorWidget = GetEditorWidget();
+
+	if (!EditorWidget.IsValid())
+	{
+		return nullptr;
+	}
+
+	UDynamicMaterialModelBase* PreviewMaterialModelBase = EditorWidget->GetPreviewMaterialModelBase();
+
+	if (!PreviewMaterialModelBase)
+	{
+		return nullptr;
+	}
+
+	TArray<UObject*> Objects;
+	PropertyHandle->GetOuterObjects(Objects);
+
+	if (Objects.Num() != 1 || !IsValid(Objects[0]))
+	{
+		return nullptr;
+	}
+
+	if (!Objects[0]->IsIn(PreviewMaterialModelBase))
+	{
+		return nullptr;
+	}
+
+	return CreateChildItem_Impl(InDetailsView, InParent.ToSharedRef(), InChildNode, PropertyHandle.ToSharedRef(), InOriginalRow);
+}
+
+TSharedPtr<ICustomDetailsViewItem> SDMObjectEditorWidgetBase::CreateChildItem_Impl(const TSharedRef<ICustomDetailsView>& InDetailsView,
+	const TSharedRef<ICustomDetailsViewItem>& InParent, const TSharedRef<IDetailTreeNode>& InChildNode, const TSharedRef<IPropertyHandle>& InPropertyHandle,
+	const FDMPropertyHandle& InOriginalRow)
+{
+	TSharedRef<ICustomDetailsViewItem> Item = InDetailsView->CreateDetailTreeItem(InParent, InChildNode);
+
+	if (!InOriginalRow.bEnabled)
+	{
+		Item->SetEnabledOverride(false);
+
+		// Disable the expansion widgets (SNullWidget is treated as removing the override).
+		Item->SetOverrideWidget(ECustomDetailsViewWidgetType::Extensions, SNew(SBox));
+	}
+
+	if (!InOriginalRow.bKeyframeable)
+	{
+		Item->SetKeyframeEnabled(false);
+	}
+
+	if (InOriginalRow.ResetToDefaultOverride.IsSet())
+	{
+		Item->SetResetToDefaultOverride(InOriginalRow.ResetToDefaultOverride.GetValue());
+	}
+
+	if (InOriginalRow.MaxWidth.IsSet())
+	{
+		Item->SetValueWidgetWidthOverride(InOriginalRow.MaxWidth);
+	}
+
+	Item->SetCreateChildItemDelegate(FOnCustomDetailsViewGenerateChildItem::CreateSP(this, &SDMObjectEditorWidgetBase::CreateChildItem, InOriginalRow));
+
+	AddDetailTreeRowExtensionWidgets(InDetailsView, InOriginalRow, Item, InPropertyHandle);
+
+	return Item;
+}
+
 #undef LOCTEXT_NAMESPACE
diff --git a/Engine/Plugins/Experimental/DynamicMaterial/Source/DynamicMaterialEditor/Private/UI/Widgets/Editor/SDMObjectEditorWidgetBase.h b/Engine/Plugins/Experimental/DynamicMaterial/Source/DynamicMaterialEditor/Private/UI/Widgets/Editor/SDMObjectEditorWidgetBase.h
--- a/Engine/Plugins/Experimental/DynamicMaterial/Source/DynamicMaterialEditor/Private/UI/Widgets/Editor/SDMObjectEditorWidgetBase.h
+++ b/Engine/Plugins/Experimental/DynamicMaterial/Source/DynamicMaterialEditor/Private/UI/Widgets/Editor/SDMObjectEditorWidgetBase.h
@@ -5,6 +5,7 @@
 #include "Widgets/SCompoundWidget.h"
 
 #include "Containers/Array.h"
+#include "Misc/NotifyHook.h"
 #include "Templates/SharedPointer.h"
 #include "UI/Utils/DMWidgetSlot.h"
 #include "UObject/WeakObjectPtr.h"
@@ -14,13 +15,17 @@ class FName;
 class ICustomDetailsView;
 class ICustomDetailsViewItem;
 class IDetailKeyframeHandler;
+class IDetailTreeNode;
+class IPropertyHandle;
 class SDMMaterialEditor;
 class UObject;
 enum class ECustomDetailsTreeInsertPosition : uint8;
 struct FDMPropertyHandle;
+struct FPropertyRowExtensionButton;
+struct FSlateIcon;
 
 /** Base class for the object editor. Provides the methods and layout for producing a Custom Details View. */
-class SDMObjectEditorWidgetBase : public SCompoundWidget
+class SDMObjectEditorWidgetBase : public SCompoundWidget, public FNotifyHook
 {
 	SLATE_DECLARE_WIDGET(SDMObjectEditorWidgetBase, SCompoundWidget)
 
@@ -38,6 +43,11 @@ public:
 
 	void Validate();
 
+	//~ Begin FNotifyHook
+	virtual void NotifyPreChange(FProperty* InPropertyAboutToChange) override;
+	virtual void NotifyPostChange(const FPropertyChangedEvent& InPropertyChangedEvent, FProperty* InPropertyThatChanged) override;
+	//~ End FNotifyHook
+
 protected:
 	static constexpr const TCHAR* DefaultCategoryName = TEXT("General");
 
@@ -68,4 +78,23 @@ protected:
 	void OnExpansionStateChanged(const TSharedRef<ICustomDetailsViewItem>& InItem, bool bInExpansionState);
 
 	virtual TArray<FDMPropertyHandle> GetPropertyRows() = 0;
+
+	virtual void AddDetailTreeRowExtensionWidgets(const TSharedRef<ICustomDetailsView>& InDetailsView, const FDMPropertyHandle& InPropertyRow,
+		const TSharedRef<ICustomDetailsViewItem>& InItem, const TSharedRef<IPropertyHandle>& InPropertyHandle);
+
+	TOptional<FPropertyRowExtensionButton> CreateKeyframeButton(TSharedPtr<IPropertyHandle> InPreviewPropertyHandle, TSharedPtr<IPropertyHandle> InOriginalPropertyHandle);
+
+	FSlateIcon GetCreateKeyIcon(TWeakPtr<IPropertyHandle> InPreviewPropertyHandleWeak, TWeakPtr<IPropertyHandle> InOriginalPropertyHandleWeak) const;
+
+	bool CanCreateKeyFrame(TWeakPtr<IPropertyHandle> InPreviewPropertyHandleWeak, TWeakPtr<IPropertyHandle> InOriginalPropertyHandleWeak) const;
+
+	void CreateKeyFrame(TWeakPtr<IPropertyHandle> InPreviewPropertyHandleWeak, TWeakPtr<IPropertyHandle> InOriginalPropertyHandleWeak);
+
+	FPropertyRowExtensionButton CreateNeedsApplyButton() const;
+
+	TSharedPtr<ICustomDetailsViewItem> CreateChildItem(const TSharedRef<ICustomDetailsView>& InDetailsView, const TSharedPtr<ICustomDetailsViewItem>& InParent,
+		const TSharedRef<IDetailTreeNode>& InChildNode, FDMPropertyHandle InOriginalRow);
+
+	virtual TSharedPtr<ICustomDetailsViewItem> CreateChildItem_Impl(const TSharedRef<ICustomDetailsView>& InDetailsView, const TSharedRef<ICustomDetailsViewItem>& InParent,
+		const TSharedRef<IDetailTreeNode>& InChildNode, const TSharedRef<IPropertyHandle>& InPropertyHandle, const FDMPropertyHandle& InOriginalRow);
 };
diff --git a/Engine/Plugins/Experimental/DynamicMaterial/Source/DynamicMaterialEditor/Private/UI/Widgets/Editor/SlotEditor/SDMMaterialSlotLayerEffectItem.cpp b/Engine/Plugins/Experimental/DynamicMaterial/Source/DynamicMaterialEditor/Private/UI/Widgets/Editor/SlotEditor/SDMMaterialSlotLayerEffectItem.cpp
--- a/Engine/Plugins/Experimental/DynamicMaterial/Source/DynamicMaterialEditor/Private/UI/Widgets/Editor/SlotEditor/SDMMaterialSlotLayerEffectItem.cpp
+++ b/Engine/Plugins/Experimental/DynamicMaterial/Source/DynamicMaterialEditor/Private/UI/Widgets/Editor/SlotEditor/SDMMaterialSlotLayerEffectItem.cpp
@@ -175,9 +175,9 @@ bool SDMMaterialSlotLayerEffectItem::CanModifyMaterialModel() const
 				{
 					if (TSharedPtr<SDMMaterialEditor> EditorWidget = SlotEditorWidget->GetEditorWidget())
 					{
-						if (UDynamicMaterialModelBase* MaterialModelBase = EditorWidget->GetMaterialModelBase())
+						if (UDynamicMaterialModelBase* PreviewMaterialModelBase = EditorWidget->GetPreviewMaterialModelBase())
 						{
-							return MaterialModelBase->IsA<UDynamicMaterialModel>();
+							return PreviewMaterialModelBase->IsA<UDynamicMaterialModel>();
 						}
 					}
 				}
diff --git a/Engine/Plugins/Experimental/DynamicMaterial/Source/DynamicMaterialEditor/Private/UI/Widgets/SDMMaterialDesigner.cpp b/Engine/Plugins/Experimental/DynamicMaterial/Source/DynamicMaterialEditor/Private/UI/Widgets/SDMMaterialDesigner.cpp
--- a/Engine/Plugins/Experimental/DynamicMaterial/Source/DynamicMaterialEditor/Private/UI/Widgets/SDMMaterialDesigner.cpp
+++ b/Engine/Plugins/Experimental/DynamicMaterial/Source/DynamicMaterialEditor/Private/UI/Widgets/SDMMaterialDesigner.cpp
@@ -146,7 +146,7 @@ void SDMMaterialDesigner::OnActorSelected(AActor* InActor)
 	}
 }
 
-UDynamicMaterialModelBase* SDMMaterialDesigner::GetMaterialModelBase() const
+UDynamicMaterialModelBase* SDMMaterialDesigner::GetOriginalMaterialModelBase() const
 {
 	if (!Content.IsValid())
 	{
@@ -155,7 +155,7 @@ UDynamicMaterialModelBase* SDMMaterialDesigner::GetMaterialModelBase() const
 
 	if (Content->GetWidgetClass().GetWidgetType() == SDMMaterialEditor::StaticWidgetClass().GetWidgetType())
 	{
-		return StaticCastSharedPtr<SDMMaterialEditor>(Content)->GetMaterialModelBase();
+		return StaticCastSharedPtr<SDMMaterialEditor>(Content)->GetOriginalMaterialModelBase();
 	}
 
 	if (Content->GetWidgetClass().GetWidgetType() == SDMMaterialWizard::StaticWidgetClass().GetWidgetType())
@@ -284,7 +284,7 @@ void SDMMaterialDesigner::SetEditorView(UDynamicMaterialModelBase* InMaterialMod
 		Layout = Settings->Layout;
 	}
 
-	SetEditorLayout(Layout, InMaterialModelBase);
+	SetEditorLayout(Layout, InMaterialModelBase, nullptr);
 }
 
 void SDMMaterialDesigner::SetEditorView(const FDMObjectMaterialProperty& InObjectMaterialProperty)
@@ -296,7 +296,7 @@ void SDMMaterialDesigner::SetEditorView(const FDMObjectMaterialProperty& InObjec
 		Layout = Settings->Layout;
 	}
 
-	SetEditorLayout(Layout, InObjectMaterialProperty);
+	SetEditorLayout(Layout, InObjectMaterialProperty, nullptr);
 }
 
 void SDMMaterialDesigner::SetWidget(const TSharedRef<SWidget>& InWidget, bool bInIncludeAssetDropTarget)
@@ -502,7 +502,8 @@ void SDMMaterialDesigner::OnLayoutChanged()
 	// Already assured
 	TSharedPtr<SDMMaterialEditor> CurrentEditor = StaticCastSharedRef<SDMMaterialEditor>(Content.ToSharedRef());
 
-	UDynamicMaterialModelBase* MaterialModelBase = CurrentEditor->GetMaterialModelBase();
+	UDynamicMaterialModelBase* OriginalMaterialModelBase = CurrentEditor->GetOriginalMaterialModelBase();
+	UDynamicMaterialModelBase* PreviewMaterialModelBase = CurrentEditor->GetPreviewMaterialModelBase();
 	const FDMObjectMaterialProperty* MaterialObjectProperty = CurrentEditor->GetMaterialObjectProperty();
 	EDMMaterialEditorMode EditorMode = CurrentEditor->GetEditMode();
 	EDMMaterialPropertyType SelectedProperty = CurrentEditor->GetSelectedPropertyType();
@@ -518,14 +519,14 @@ void SDMMaterialDesigner::OnLayoutChanged()
 
 	if (MaterialObjectProperty)
 	{
-		if (!SetEditorLayout(Settings->Layout, *MaterialObjectProperty))
+		if (!SetEditorLayout(Settings->Layout, *MaterialObjectProperty, PreviewMaterialModelBase))
 		{
 			return;
 		}
 	}
 	else
 	{
-		if (!SetEditorLayout(Settings->Layout, MaterialModelBase))
+		if (!SetEditorLayout(Settings->Layout, OriginalMaterialModelBase, PreviewMaterialModelBase))
 		{
 			return;
 		}
@@ -556,7 +557,8 @@ void SDMMaterialDesigner::OnLayoutChanged()
 	}
 }
 
-bool SDMMaterialDesigner::SetEditorLayout(EDMMaterialEditorLayout InLayout, UDynamicMaterialModelBase* InMaterialModelBase)
+bool SDMMaterialDesigner::SetEditorLayout(EDMMaterialEditorLayout InLayout, UDynamicMaterialModelBase* InMaterialModelBase,
+	UDynamicMaterialModelBase* InCurrentPreviewMaterial)
 {
 	TSharedPtr<SDMMaterialEditor> NewEditor;
 
@@ -564,17 +566,20 @@ bool SDMMaterialDesigner::SetEditorLayout(EDMMaterialEditorLayout InLayout, UDyn
 	{
 		case EDMMaterialEditorLayout::Left:
 			NewEditor = SNew(SDMMaterialEditor_Left, SharedThis(this))
-				.MaterialModelBase(InMaterialModelBase);
+				.MaterialModelBase(InMaterialModelBase)
+				.PreviewMaterialModelBase(InCurrentPreviewMaterial);
 			break;
 
 		case EDMMaterialEditorLayout::Top:
 			NewEditor = SNew(SDMMaterialEditor_TopVertical, SharedThis(this))
-				.MaterialModelBase(InMaterialModelBase);
+				.MaterialModelBase(InMaterialModelBase)
+				.PreviewMaterialModelBase(InCurrentPreviewMaterial);
 			break;
 
 		case EDMMaterialEditorLayout::TopSlim:
 			NewEditor = SNew(SDMMaterialEditor_TopSlim, SharedThis(this))
-				.MaterialModelBase(InMaterialModelBase);
+				.MaterialModelBase(InMaterialModelBase)
+				.PreviewMaterialModelBase(InCurrentPreviewMaterial);
 			break;
 
 		default:
@@ -591,7 +596,8 @@ bool SDMMaterialDesigner::SetEditorLayout(EDMMaterialEditorLayout InLayout, UDyn
 	return true;
 }
 
-bool SDMMaterialDesigner::SetEditorLayout(EDMMaterialEditorLayout InLayout, const FDMObjectMaterialProperty& InObjectMaterialProperty)
+bool SDMMaterialDesigner::SetEditorLayout(EDMMaterialEditorLayout InLayout, const FDMObjectMaterialProperty& InObjectMaterialProperty,
+	UDynamicMaterialModelBase* InCurrentPreviewMaterial)
 {
 	TSharedPtr<SDMMaterialEditor> NewEditor;
 
@@ -599,17 +605,20 @@ bool SDMMaterialDesigner::SetEditorLayout(EDMMaterialEditorLayout InLayout, cons
 	{
 		case EDMMaterialEditorLayout::Left:
 			NewEditor = SNew(SDMMaterialEditor_Left, SharedThis(this))
-				.MaterialProperty(InObjectMaterialProperty);
+				.MaterialProperty(InObjectMaterialProperty)
+				.PreviewMaterialModelBase(InCurrentPreviewMaterial);
 			break;
 
 		case EDMMaterialEditorLayout::Top:
 			NewEditor = SNew(SDMMaterialEditor_TopVertical, SharedThis(this))
-				.MaterialProperty(InObjectMaterialProperty);
+				.MaterialProperty(InObjectMaterialProperty)
+				.PreviewMaterialModelBase(InCurrentPreviewMaterial);
 			break;
 
 		case EDMMaterialEditorLayout::TopSlim:
 			NewEditor = SNew(SDMMaterialEditor_TopSlim, SharedThis(this))
-				.MaterialProperty(InObjectMaterialProperty);
+				.MaterialProperty(InObjectMaterialProperty)
+				.PreviewMaterialModelBase(InCurrentPreviewMaterial);
 			break;
 
 		default:
diff --git a/Engine/Plugins/Experimental/DynamicMaterial/Source/DynamicMaterialEditor/Private/UI/Widgets/SDMMaterialDesigner.h b/Engine/Plugins/Experimental/DynamicMaterial/Source/DynamicMaterialEditor/Private/UI/Widgets/SDMMaterialDesigner.h
--- a/Engine/Plugins/Experimental/DynamicMaterial/Source/DynamicMaterialEditor/Private/UI/Widgets/SDMMaterialDesigner.h
+++ b/Engine/Plugins/Experimental/DynamicMaterial/Source/DynamicMaterialEditor/Private/UI/Widgets/SDMMaterialDesigner.h
@@ -57,7 +57,7 @@ public:
 	void OnActorSelected(AActor* InActor);
 
 	/** Getters */
-	UDynamicMaterialModelBase* GetMaterialModelBase() const;
+	UDynamicMaterialModelBase* GetOriginalMaterialModelBase() const;
 
 	//~ Begin SWidget
 	virtual void Tick(const FGeometry& InAllottedGeometry, const double InCurrentTime, const float InDeltaTime) override;
@@ -104,7 +104,9 @@ protected:
 
 	void OnLayoutChanged();
 
-	bool SetEditorLayout(EDMMaterialEditorLayout InLayout, UDynamicMaterialModelBase* InMaterialModelBase);
+	bool SetEditorLayout(EDMMaterialEditorLayout InLayout, UDynamicMaterialModelBase* InMaterialModelBase, 
+		UDynamicMaterialModelBase* InCurrentPreviewMaterial);
 
-	bool SetEditorLayout(EDMMaterialEditorLayout InLayout, const FDMObjectMaterialProperty& InObjectMaterialProperty);
+	bool SetEditorLayout(EDMMaterialEditorLayout InLayout, const FDMObjectMaterialProperty& InObjectMaterialProperty,
+		UDynamicMaterialModelBase* InCurrentPreviewMaterial);
 };
diff --git a/Engine/Plugins/Experimental/DynamicMaterial/Source/DynamicMaterialEditor/Private/UI/Widgets/SDMMaterialEditor.cpp b/Engine/Plugins/Experimental/DynamicMaterial/Source/DynamicMaterialEditor/Private/UI/Widgets/SDMMaterialEditor.cpp
--- a/Engine/Plugins/Experimental/DynamicMaterial/Source/DynamicMaterialEditor/Private/UI/Widgets/SDMMaterialEditor.cpp
+++ b/Engine/Plugins/Experimental/DynamicMaterial/Source/DynamicMaterialEditor/Private/UI/Widgets/SDMMaterialEditor.cpp
@@ -16,6 +16,8 @@
 #include "Framework/Application/IInputProcessor.h"
 #include "Framework/Application/SlateApplication.h"
 #include "Framework/Commands/GenericCommands.h"
+#include "Material/DynamicMaterialInstance.h"
+#include "Material/DynamicMaterialInstanceFactory.h"
 #include "Materials/Material.h"
 #include "Materials/MaterialFunctionInterface.h"
 #include "Misc/MessageDialog.h"
@@ -23,6 +25,7 @@
 #include "Model/DynamicMaterialModelBase.h"
 #include "Model/DynamicMaterialModelDynamic.h"
 #include "Model/DynamicMaterialModelEditorOnlyData.h"
+#include "PackageTools.h"
 #include "Styling/SlateIconFinder.h"
 #include "UI/Utils/DMEditorSelectionContext.h"
 #include "UI/Utils/DMPreviewMaterialManager.h"
@@ -35,6 +38,7 @@
 #include "UI/Widgets/Editor/SDMStatusBar.h"
 #include "UI/Widgets/Editor/SDMToolBar.h"
 #include "UI/Widgets/SDMMaterialDesigner.h"
+#include "UObject/Linker.h"
 #include "Utils/DMMaterialModelFunctionLibrary.h"
 #include "Utils/DMPrivate.h"
 #include "Utils/DMTextureSetFunctionLibrary.h"
@@ -169,11 +173,11 @@ void SDMMaterialEditor::Construct(const FArguments& InArgs, const TSharedRef<SDM
 
 	if (InArgs._MaterialProperty.IsSet())
 	{
-		SetObjectMaterialProperty(InArgs._MaterialProperty.GetValue());
+		SetObjectMaterialProperty(InArgs._MaterialProperty.GetValue(), InArgs._PreviewMaterialModelBase);
 	}
 	else if (IsValid(InArgs._MaterialModelBase))
 	{
-		SetMaterialModelBase(InArgs._MaterialModelBase);
+		SetMaterialModelBase(InArgs._MaterialModelBase, InArgs._PreviewMaterialModelBase);
 	}
 	else
 	{
@@ -196,18 +200,45 @@ TSharedPtr<SDMMaterialDesigner> SDMMaterialEditor::GetDesignerWidget() const
 	return DesignerWidgetWeak.Pin();
 }
 
-UDynamicMaterialModelBase* SDMMaterialEditor::GetMaterialModelBase() const
+UDynamicMaterialModelBase* SDMMaterialEditor::GetOriginalMaterialModelBase() const
 {
-	return MaterialModelBaseWeak.Get();
+	return OriginalMaterialModelBaseWeak.Get();
 }
 
-void SDMMaterialEditor::SetMaterialModelBase(UDynamicMaterialModelBase* InMaterialModelBase)
+UDynamicMaterialModelBase* SDMMaterialEditor::GetPreviewMaterialModelBase() const
 {
-	MaterialModelBaseWeak = InMaterialModelBase;
+	return PreviewMaterialModelBase.Get();
+}
+
+void SDMMaterialEditor::SetMaterialModelBase(UDynamicMaterialModelBase* InMaterialModelBase, UDynamicMaterialModelBase* InPreviewMaterialModelBase)
+{
+	OriginalMaterialModelBaseWeak = InMaterialModelBase;
 
-	if (UDynamicMaterialModelDynamic* MaterialModelDynamic = Cast<UDynamicMaterialModelDynamic>(InMaterialModelBase))
+	if (!InMaterialModelBase)
+	{
+		PreviewMaterialModelBase = nullptr;
+	}
+	else if (InPreviewMaterialModelBase)
+	{
+		PreviewMaterialModelBase = InPreviewMaterialModelBase;
+		PreviewMaterialModelBase->MarkOriginalUpdated();
+	}
+	else
 	{
-		MaterialModelDynamic->EnsureComponents();
+		PreviewMaterialModelBase = UDMMaterialModelFunctionLibrary::CreatePreviewModel(InMaterialModelBase);
+
+		UDynamicMaterialInstanceFactory* Factory = NewObject<UDynamicMaterialInstanceFactory>();
+
+		UDynamicMaterialInstance* NewInstance = Cast<UDynamicMaterialInstance>(Factory->FactoryCreateNew(
+			UDynamicMaterialInstance::StaticClass(),
+			GetTransientPackage(),
+			NAME_None,
+			RF_Transactional,
+			PreviewMaterialModelBase,
+			nullptr
+		));
+
+		PreviewMaterialInstance = NewInstance;
 	}
 
 	EditGlobalSettings();
@@ -217,9 +248,9 @@ void SDMMaterialEditor::SetMaterialModelBase(UDynamicMaterialModelBase* InMateri
 	BindEditorOnlyDataUpdate(InMaterialModelBase);
 }
 
-UDynamicMaterialModel* SDMMaterialEditor::GetMaterialModel() const
+UDynamicMaterialModel* SDMMaterialEditor::GetPreviewMaterialModel() const
 {
-	if (UDynamicMaterialModelBase* MaterialModelBase = MaterialModelBaseWeak.Get())
+	if (UDynamicMaterialModelBase* MaterialModelBase = GetPreviewMaterialModelBase())
 	{
 		return MaterialModelBase->ResolveMaterialModel();
 	}
@@ -229,7 +260,7 @@ UDynamicMaterialModel* SDMMaterialEditor::GetMaterialModel() const
 
 bool SDMMaterialEditor::IsDynamicModel() const
 {
-	return !!Cast<UDynamicMaterialModelDynamic>(MaterialModelBaseWeak.Get());
+	return !!Cast<UDynamicMaterialModelDynamic>(GetPreviewMaterialModelBase());
 }
 
 const FDMObjectMaterialProperty* SDMMaterialEditor::GetMaterialObjectProperty() const
@@ -242,7 +273,7 @@ const FDMObjectMaterialProperty* SDMMaterialEditor::GetMaterialObjectProperty()
 	return nullptr;
 }
 
-void SDMMaterialEditor::SetObjectMaterialProperty(const FDMObjectMaterialProperty& InObjectProperty)
+void SDMMaterialEditor::SetObjectMaterialProperty(const FDMObjectMaterialProperty& InObjectProperty, UDynamicMaterialModelBase* InPreviewMaterialModelBase)
 {
 	UDynamicMaterialModelBase* MaterialModelBase = InObjectProperty.GetMaterialModelBase();
 
@@ -253,7 +284,7 @@ void SDMMaterialEditor::SetObjectMaterialProperty(const FDMObjectMaterialPropert
 	}
 
 	ObjectMaterialPropertyOpt = InObjectProperty;
-	SetMaterialModelBase(MaterialModelBase);
+	SetMaterialModelBase(MaterialModelBase, InPreviewMaterialModelBase);
 
 	BindEditorOnlyDataUpdate(MaterialModelBase);
 }
@@ -268,6 +299,25 @@ AActor* SDMMaterialEditor::GetMaterialActor() const
 	return nullptr;
 }
 
+UDMMaterialComponent* SDMMaterialEditor::GetOriginalComponent(UDMMaterialComponent* InPreviewComponent) const
+{
+	if (!PreviewMaterialModelBase)
+	{
+		return nullptr;
+	}
+
+	UDynamicMaterialModelBase* OriginalMaterialModelBase = GetOriginalMaterialModelBase();
+
+	if (!OriginalMaterialModelBase)
+	{
+		return nullptr;
+	}
+
+	const FString RelativePath = InPreviewComponent->GetPathName(PreviewMaterialModelBase);
+
+	return UE::DynamicMaterialEditor::Private::FindSubobject<UDMMaterialComponent>(OriginalMaterialModelBase, *RelativePath);
+}
+
 EDMMaterialEditorMode SDMMaterialEditor::GetEditMode() const
 {
 	return SelectionContext.EditorMode;
@@ -323,7 +373,7 @@ void SDMMaterialEditor::SelectProperty(EDMMaterialPropertyType InProperty, bool
 
 	PageHistoryAdd({EDMMaterialEditorMode::EditSlot, InProperty});
 
-	if (UDynamicMaterialModelEditorOnlyData* EditorOnlyData = UDynamicMaterialModelEditorOnlyData::Get(MaterialModelBaseWeak))
+	if (UDynamicMaterialModelEditorOnlyData* EditorOnlyData = UDynamicMaterialModelEditorOnlyData::Get(GetPreviewMaterialModelBase()))
 	{
 		UDMMaterialSlot* CurrentSlot = SelectionContext.Slot.Get();
 
@@ -438,9 +488,7 @@ void SDMMaterialEditor::OnEffectSelected(const TSharedRef<SDMMaterialSlotLayerEf
 
 void SDMMaterialEditor::OpenMaterialPreviewTab()
 {
-	UDynamicMaterialModelBase* MaterialModelBase = GetMaterialModelBase();
-
-	if (!MaterialModelBase)
+	if (!PreviewMaterialModelBase)
 	{
 		return;
 	}
@@ -472,7 +520,7 @@ void SDMMaterialEditor::OpenMaterialPreviewTab()
 
 	MaterialPreviewTab = FGlobalTabmanager::Get()->TryInvokeTab(TabId);
 	MaterialPreviewTab->ActivateInParent(ETabActivationCause::SetDirectly);
-	MaterialPreviewTab->SetLabel(FText::FromString(MaterialModelBase->GetPathName()));
+	MaterialPreviewTab->SetLabel(FText::FromString(PreviewMaterialModelBase->GetPathName()));
 	MaterialPreviewTab->SetOnTabClosed(SDockTab::FOnTabClosedCallback::CreateSPLambda(
 		this,
 		[this](TSharedRef<SDockTab> InDockTab)
@@ -485,7 +533,7 @@ void SDMMaterialEditor::OpenMaterialPreviewTab()
 	MaterialPreviewTabSlot = TDMWidgetSlot<SWidget>(
 		Wrapper, 
 		0, 
-		SNew(SDMMaterialPreview, SharedThis(this), MaterialModelBase)
+		SNew(SDMMaterialPreview, SharedThis(this), PreviewMaterialModelBase)
 		.IsPopout(true)
 	);
 
@@ -504,9 +552,7 @@ void SDMMaterialEditor::CloseMaterialPreviewTab()
 
 TSharedPtr<IToolTip> SDMMaterialEditor::GetMaterialPreviewToolTip()
 {
-	UDynamicMaterialModelBase* MaterialModelBase = GetMaterialModelBase();
-
-	if (!MaterialModelBase)
+	if (!PreviewMaterialModelBase)
 	{
 		return nullptr;
 	}
@@ -539,7 +585,7 @@ TSharedPtr<IToolTip> SDMMaterialEditor::GetMaterialPreviewToolTip()
 	MaterialPreviewToolTipSlot = TDMWidgetSlot<SWidget>(
 		Wrapper,
 		0,
-		SNew(SDMMaterialPreview, SharedThis(this), MaterialModelBase)
+		SNew(SDMMaterialPreview, SharedThis(this), PreviewMaterialModelBase)
 		.ShowMenu(false)
 	);
 
@@ -569,9 +615,9 @@ void SDMMaterialEditor::Validate()
 		return;
 	}
 
-	UDynamicMaterialModelBase* MaterialModelBase = GetMaterialModelBase();
+	UDynamicMaterialModelBase* OriginalMaterialModelBase = GetOriginalMaterialModelBase();
 
-	if (!IsValid(MaterialModelBase))
+	if (!IsValid(OriginalMaterialModelBase) || !IsValid(PreviewMaterialModelBase))
 	{
 		Close();
 		return;
@@ -582,12 +628,12 @@ void SDMMaterialEditor::Validate()
 		const FDMObjectMaterialProperty& ObjectMaterialProperty = ObjectMaterialPropertyOpt.GetValue();
 		UDynamicMaterialModelBase* MaterialModelBaseFromProperty = ObjectMaterialProperty.GetMaterialModelBase();
 
-		if (!UDMMaterialModelFunctionLibrary::IsModelValid(MaterialModelBaseFromProperty))
+		if (!UDMMaterialModelFunctionLibrary::IsModelValid(OriginalMaterialModelBase))
 		{
 			MaterialModelBaseFromProperty = nullptr;
 		}
 
-		if (MaterialModelBase != MaterialModelBaseFromProperty)
+		if (OriginalMaterialModelBase != MaterialModelBaseFromProperty)
 		{
 			if (TSharedPtr<SDMMaterialDesigner> DesignerWidget = DesignerWidgetWeak.Pin())
 			{
@@ -596,7 +642,7 @@ void SDMMaterialEditor::Validate()
 			}
 		}
 	}
-	else if (!UDMMaterialModelFunctionLibrary::IsModelValid(MaterialModelBase))
+	else if (!UDMMaterialModelFunctionLibrary::IsModelValid(OriginalMaterialModelBase))
 	{
 		Close();
 		return;
@@ -671,6 +717,41 @@ void SDMMaterialEditor::PostRedo(bool bInSuccess)
 	OnUndo();
 }
 
+void SDMMaterialEditor::AddReferencedObjects(FReferenceCollector& InCollector)
+{
+	if (PreviewMaterialInstance)
+	{
+		InCollector.AddReferencedObject(PreviewMaterialInstance);
+	}
+
+	if (PreviewMaterialModelBase)
+	{
+		InCollector.AddReferencedObject(PreviewMaterialModelBase);
+	}
+}
+
+FString SDMMaterialEditor::GetReferencerName() const
+{
+	return TEXT("SDMMaterialEditor");
+}
+
+UPackage* SDMMaterialEditor::GetSaveablePackage(UObject* InObject)
+{
+	if (!IsValid(InObject))
+	{
+		return nullptr;
+	}
+
+	UPackage* Package = InObject->GetPackage();
+
+	if (!Package || Package->HasAllFlags(RF_Transient))
+	{
+		return nullptr;
+	}
+
+	return Package;
+}
+
 void SDMMaterialEditor::BindCommands(SDMMaterialSlotEditor* InSlotEditor)
 {
 	const FGenericCommands& GenericCommands = FGenericCommands::Get();
@@ -756,7 +837,7 @@ void SDMMaterialEditor::BindCommands(SDMMaterialSlotEditor* InSlotEditor)
 
 bool SDMMaterialEditor::IsPropertyValidForModel(EDMMaterialPropertyType InProperty) const
 {
-	UDynamicMaterialModelEditorOnlyData* EditorOnlyData = UDynamicMaterialModelEditorOnlyData::Get(MaterialModelBaseWeak);
+	UDynamicMaterialModelEditorOnlyData* EditorOnlyData = UDynamicMaterialModelEditorOnlyData::Get(GetPreviewMaterialModelBase());
 
 	if (!EditorOnlyData)
 	{
@@ -913,6 +994,68 @@ void SDMMaterialEditor::PageHistoryAdd(const FDMMaterialEditorPage& InPage)
 	SelectionContext.PageHistoryCount = NewPageIndex + 1;
 }
 
+void SDMMaterialEditor::ApplyToOriginal()
+{
+	UDynamicMaterialModelBase* Preview = GetPreviewMaterialModelBase();
+
+	if (!Preview)
+	{
+		return;
+	}
+
+	UDynamicMaterialModelBase* Original = GetOriginalMaterialModelBase();
+
+	if (!Original)
+	{
+		return;
+	}
+
+	UDynamicMaterialInstance* OriginalMaterialInstance = Original->GetDynamicMaterialInstance();
+
+	FScopedTransaction Transaction(LOCTEXT("ApplyMaterialModel", "Apply Material Designer Changes"));
+	Original->Modify();
+
+	UDMMaterialModelFunctionLibrary::MirrorMaterialModel(Preview, Original);
+
+	OriginalMaterialModelBaseWeak = Original;
+
+	if (OriginalMaterialInstance)
+	{
+		OriginalMaterialInstance->SetMaterialModel(Original);
+		Original->SetDynamicMaterialInstance(OriginalMaterialInstance);
+	}
+
+	if (UDynamicMaterialModelEditorOnlyData* EditorOnlyData = UDynamicMaterialModelEditorOnlyData::Get(Original))
+	{
+		EditorOnlyData->RequestMaterialBuild(/* Immediate */ true);
+	}
+
+	if (OriginalMaterialInstance)
+	{
+		OriginalMaterialInstance->InitializeMIDPublic();
+		Original->ApplyComponents(OriginalMaterialInstance);
+	}
+
+	if (TSharedPtr<SDMMaterialDesigner> DesignerWidget = GetDesignerWidget())
+	{
+		if (ObjectMaterialPropertyOpt.IsSet())
+		{
+			DesignerWidget->OpenObjectMaterialProperty(ObjectMaterialPropertyOpt.GetValue());
+		}
+		else
+		{
+			DesignerWidget->OpenMaterialModelBase(Original);
+		}
+	}
+
+	/*
+	if (PreviewMaterialModelBase)
+	{
+		PreviewMaterialModelBase->MarkOriginalUpdated();
+	}
+	*/
+}
+
 bool SDMMaterialEditor::SetActivePage(const FDMMaterialEditorPage& InPage)
 {
 	switch (InPage.EditorMode)
@@ -936,6 +1079,21 @@ bool SDMMaterialEditor::SetActivePage(const FDMMaterialEditorPage& InPage)
 	}
 }
 
+void SDMMaterialEditor::SaveOriginal()
+{
+	if (UDynamicMaterialModelBase* MaterialModelBase = GetOriginalMaterialModelBase())
+	{
+		ApplyToOriginal();
+
+		if (!!GetSaveablePackage(MaterialModelBase))
+		{
+			TArray<UObject*> AssetsToSave;
+			AssetsToSave.Add(MaterialModelBase);
+			UPackageTools::SavePackagesForObjects(AssetsToSave);
+		}
+	}
+}
+
 void SDMMaterialEditor::HandleDrop_CreateTextureSet(const TArray<FAssetData>& InTextureAssets)
 {
 	if (InTextureAssets.Num() < 2)
@@ -965,7 +1123,7 @@ void SDMMaterialEditor::HandleDrop_TextureSet(UDMTextureSet* InTextureSet)
 		return;
 	}
 
-	UDynamicMaterialModel* MaterialModel = GetMaterialModel();
+	UDynamicMaterialModel* MaterialModel = GetPreviewMaterialModel();
 
 	if (!MaterialModel)
 	{
@@ -1104,7 +1262,7 @@ TSharedRef<SDMToolBar> SDMMaterialEditor::CreateSlot_ToolBar()
 
 TSharedRef<SDMMaterialGlobalSettingsEditor> SDMMaterialEditor::CreateSlot_GlobalSettingsEditor()
 {
-	return SNew(SDMMaterialGlobalSettingsEditor, SharedThis(this), GetMaterialModelBase());
+	return SNew(SDMMaterialGlobalSettingsEditor, SharedThis(this), GetPreviewMaterialModelBase());
 }
 
 TSharedRef<SDMMaterialProperties> SDMMaterialEditor::CreateSlot_MaterialProperties()
@@ -1117,7 +1275,7 @@ TSharedRef<SWidget> SDMMaterialEditor::CreateSlot_Preview()
 	return SNew(SOverlay)
 		+ SOverlay::Slot()
 		[
-			SNew(SDMMaterialPreview, SharedThis(this), GetMaterialModelBase())
+			SNew(SDMMaterialPreview, SharedThis(this), GetPreviewMaterialModelBase())
 		]
 		+ SOverlay::Slot()
 		.HAlign(EHorizontalAlignment::HAlign_Left)
@@ -1140,7 +1298,7 @@ TSharedRef<SDMMaterialPropertySelector> SDMMaterialEditor::CreateSlot_PropertySe
 
 	if (SelectionContext.EditorMode == EDMMaterialEditorMode::EditSlot && SelectionContext.Property == EDMMaterialPropertyType::None)
 	{
-		if (UDynamicMaterialModel* MaterialModel = GetMaterialModel())
+		if (UDynamicMaterialModel* MaterialModel = GetPreviewMaterialModel())
 		{
 			if (UDynamicMaterialModelEditorOnlyData* EditorOnlyData = UDynamicMaterialModelEditorOnlyData::Get(MaterialModel))
 			{
@@ -1194,14 +1352,14 @@ TSharedRef<SDMMaterialComponentEditor> SDMMaterialEditor::CreateSlot_ComponentEd
 
 TSharedRef<SDMStatusBar> SDMMaterialEditor::CreateSlot_StatusBar()
 {
-	return SNew(SDMStatusBar, SharedThis(this), GetMaterialModelBase());
+	return SNew(SDMStatusBar, SharedThis(this), GetPreviewMaterialModelBase());
 }
 
 void SDMMaterialEditor::OnUndo()
 {
-	UDynamicMaterialModelBase* MaterialModel = GetMaterialModelBase();
+	UDynamicMaterialModelBase* OriginalMaterialModel = GetOriginalMaterialModelBase();
 
-	if (!IsValid(MaterialModel))
+	if (!IsValid(OriginalMaterialModel))
 	{
 		Close();
 		return;
@@ -1233,7 +1391,7 @@ void SDMMaterialEditor::ApplySelectionContext()
 
 	if (SelectionContext.EditorMode == EDMMaterialEditorMode::EditSlot)
 	{
-		if (UDynamicMaterialModelEditorOnlyData* EditorOnlyData = UDynamicMaterialModelEditorOnlyData::Get(MaterialModelBaseWeak))
+		if (UDynamicMaterialModelEditorOnlyData* EditorOnlyData = UDynamicMaterialModelEditorOnlyData::Get(GetPreviewMaterialModelBase()))
 		{
 			if (UDMMaterialSlot* Slot = SelectionContext.Slot.Get())
 			{
diff --git a/Engine/Plugins/Experimental/DynamicMaterial/Source/DynamicMaterialEditor/Private/UI/Widgets/SDMMaterialEditor.h b/Engine/Plugins/Experimental/DynamicMaterial/Source/DynamicMaterialEditor/Private/UI/Widgets/SDMMaterialEditor.h
--- a/Engine/Plugins/Experimental/DynamicMaterial/Source/DynamicMaterialEditor/Private/UI/Widgets/SDMMaterialEditor.h
+++ b/Engine/Plugins/Experimental/DynamicMaterial/Source/DynamicMaterialEditor/Private/UI/Widgets/SDMMaterialEditor.h
@@ -51,19 +51,23 @@ namespace UE::DynamicMaterialEditor::Private
 	inline const TCHAR* EditorLightBackground = TEXT("Brushes.Header");
 }
 
-class SDMMaterialEditor : public SCompoundWidget, public FSelfRegisteringEditorUndoClient
+class SDMMaterialEditor : public SCompoundWidget, public FSelfRegisteringEditorUndoClient, public FGCObject
 {
 	SLATE_DECLARE_WIDGET(SDMMaterialEditor, SCompoundWidget)
 
 	SLATE_BEGIN_ARGS(SDMMaterialEditor)
 		: _MaterialModelBase(nullptr)
 		, _MaterialProperty(TOptional<FDMObjectMaterialProperty>())
+		, _PreviewMaterialModelBase(nullptr)
 		{}
 		SLATE_ARGUMENT(UDynamicMaterialModelBase*, MaterialModelBase)
 		SLATE_ARGUMENT(TOptional<FDMObjectMaterialProperty>, MaterialProperty)
+		SLATE_ARGUMENT(UDynamicMaterialModelBase*, PreviewMaterialModelBase)
 	SLATE_END_ARGS()
 
 public:
+	static UPackage* GetSaveablePackage(UObject* InObject);
+
 	SDMMaterialEditor();
 
 	virtual ~SDMMaterialEditor() override;
@@ -73,9 +77,11 @@ public:
 	TSharedPtr<SDMMaterialDesigner> GetDesignerWidget() const;
 
 	/** Material Model / Actor */
-	UDynamicMaterialModelBase* GetMaterialModelBase() const;
+	UDynamicMaterialModelBase* GetOriginalMaterialModelBase() const;
+
+	UDynamicMaterialModelBase* GetPreviewMaterialModelBase() const;
 
-	UDynamicMaterialModel* GetMaterialModel() const;
+	UDynamicMaterialModel* GetPreviewMaterialModel() const;
 
 	bool IsDynamicModel() const;
 
@@ -83,6 +89,8 @@ public:
 
 	AActor* GetMaterialActor() const;
 
+	UDMMaterialComponent* GetOriginalComponent(UDMMaterialComponent* InPreviewComponent) const;
+
 	/** Widget Components */
 	EDMMaterialEditorMode GetEditMode() const;
 
@@ -139,6 +147,10 @@ public:
 
 	bool SetActivePage(const FDMMaterialEditorPage& InPage);
 
+	void SaveOriginal();
+
+	void ApplyToOriginal();
+
 	void HandleDrop_CreateTextureSet(const TArray<FAssetData>& InTextureAssets);
 
 	void HandleDrop_TextureSet(UDMTextureSet* InTextureSet);
@@ -154,6 +166,11 @@ public:
 	virtual void PostRedo(bool bInSuccess) override;
 	//~ End FUndoClient
 
+	//~ Begin FGCObject
+	virtual void AddReferencedObjects(FReferenceCollector& InCollector) override;
+	virtual FString GetReferencerName() const override;
+	//~ End FGCObject
+
 protected:
 	TWeakPtr<SDMMaterialDesigner> DesignerWidgetWeak;
 
@@ -169,7 +186,9 @@ protected:
 	TDMWidgetSlot<SDMMaterialComponentEditor> ComponentEditorSlot;
 	TDMWidgetSlot<SDMStatusBar> StatusBarSlot;
 
-	TWeakObjectPtr<UDynamicMaterialModelBase> MaterialModelBaseWeak;
+	TWeakObjectPtr<UDynamicMaterialModelBase> OriginalMaterialModelBaseWeak;
+	TObjectPtr<UDynamicMaterialInstance> PreviewMaterialInstance;
+	TObjectPtr<UDynamicMaterialModelBase> PreviewMaterialModelBase;
 	TOptional<FDMObjectMaterialProperty> ObjectMaterialPropertyOpt;
 
 	TSharedRef<FUICommandList> CommandList;
@@ -188,9 +207,9 @@ protected:
 	TWeakObjectPtr<UDynamicMaterialModelEditorOnlyData> EditorOnlyDataUpdateObject;
 
 	/** Operations */
-	void SetMaterialModelBase(UDynamicMaterialModelBase* InMaterialModelBase);
+	void SetMaterialModelBase(UDynamicMaterialModelBase* InMaterialModelBase, UDynamicMaterialModelBase* InPreviewMaterialModelBase);
 
-	void SetObjectMaterialProperty(const FDMObjectMaterialProperty& InObjectProperty);
+	void SetObjectMaterialProperty(const FDMObjectMaterialProperty& InObjectProperty, UDynamicMaterialModelBase* InPreviewMaterialModelBase);
 
 	void SetMaterialActor(AActor* InActor);
 
diff --git a/Engine/Plugins/Experimental/DynamicMaterial/Source/DynamicMaterialEditor/Private/UI/Widgets/Visualizers/SDMMaterialComponentPreview.cpp b/Engine/Plugins/Experimental/DynamicMaterial/Source/DynamicMaterialEditor/Private/UI/Widgets/Visualizers/SDMMaterialComponentPreview.cpp
--- a/Engine/Plugins/Experimental/DynamicMaterial/Source/DynamicMaterialEditor/Private/UI/Widgets/Visualizers/SDMMaterialComponentPreview.cpp
+++ b/Engine/Plugins/Experimental/DynamicMaterial/Source/DynamicMaterialEditor/Private/UI/Widgets/Visualizers/SDMMaterialComponentPreview.cpp
@@ -55,7 +55,7 @@ void SDMMaterialComponentPreview::Construct(const FArguments& InArgs, const TSha
 
 	PreviewMaterialBaseWeak = InEditorWidget->GetPreviewMaterialManager()->CreatePreviewMaterial(InComponent);
 	PreviewMaterialDynamicWeak = InEditorWidget->GetPreviewMaterialManager()->CreatePreviewMaterialDynamic(PreviewMaterialBaseWeak.Get());
-	MaterialModelBaseWeak = InEditorWidget->GetMaterialModelBase();
+	MaterialModelBaseWeak = InEditorWidget->GetPreviewMaterialModelBase();
 
 	SetCanTick(true);
 
@@ -211,15 +211,15 @@ void SDMMaterialComponentPreview::RecreateMaterial()
 		EditorWidget->GetPreviewMaterialManager()->FreePreviewMaterialDynamic(PreviewMaterialBase);
 		PreviewMaterialDynamicWeak = EditorWidget->GetPreviewMaterialManager()->CreatePreviewMaterialDynamic(PreviewMaterialBase);
 
-		UDynamicMaterialModelBase* MaterialModelBase = EditorWidget->GetMaterialModelBase();
+		UDynamicMaterialModelBase* PreviewMaterialModelBase = EditorWidget->GetPreviewMaterialModelBase();
 
-		if (UDynamicMaterialModel* MaterialModel = Cast<UDynamicMaterialModel>(MaterialModelBase))
+		if (UDynamicMaterialModel* PreviewMaterialModel = Cast<UDynamicMaterialModel>(PreviewMaterialModelBase))
 		{
-			MaterialModel->ApplyComponents(PreviewMaterialDynamicWeak.Get());
+			PreviewMaterialModel->ApplyComponents(PreviewMaterialDynamicWeak.Get());
 		}
-		else if (UDynamicMaterialModelDynamic* MaterialModelDynamic = Cast<UDynamicMaterialModelDynamic>(MaterialModelBase))
+		else if (UDynamicMaterialModelDynamic* PreviewMaterialModelDynamic = Cast<UDynamicMaterialModelDynamic>(PreviewMaterialModelBase))
 		{
-			MaterialModelDynamic->ApplyComponents(PreviewMaterialDynamicWeak.Get());
+			PreviewMaterialModelDynamic->ApplyComponents(PreviewMaterialDynamicWeak.Get());
 		}
 
 		Brush.SetMaterial(PreviewMaterialDynamicWeak.Get());
@@ -231,15 +231,15 @@ void SDMMaterialComponentPreview::RecreateMaterial()
 		EditorWidget->GetPreviewMaterialManager()->FreePreviewMaterialDynamic(PreviewMaterialBase);
 		PreviewMaterialDynamicWeak = EditorWidget->GetPreviewMaterialManager()->CreatePreviewMaterialDynamic(PreviewMaterialBase);
 
-		UDynamicMaterialModelBase* MaterialModelBase = EditorWidget->GetMaterialModelBase();
+		UDynamicMaterialModelBase* PreviewMaterialModelBase = EditorWidget->GetPreviewMaterialModelBase();
 
-		if (UDynamicMaterialModel* MaterialModel = Cast<UDynamicMaterialModel>(MaterialModelBase))
+		if (UDynamicMaterialModel* PreviewMaterialModel = Cast<UDynamicMaterialModel>(PreviewMaterialModelBase))
 		{
-			MaterialModel->ApplyComponents(PreviewMaterialDynamicWeak.Get());
+			PreviewMaterialModel->ApplyComponents(PreviewMaterialDynamicWeak.Get());
 		}
-		else if (UDynamicMaterialModelDynamic* MaterialModelDynamic = Cast<UDynamicMaterialModelDynamic>(MaterialModelBase))
+		else if (UDynamicMaterialModelDynamic* PreviewMaterialModelDynamic = Cast<UDynamicMaterialModelDynamic>(PreviewMaterialModelBase))
 		{
-			MaterialModelDynamic->ApplyComponents(PreviewMaterialDynamicWeak.Get());
+			PreviewMaterialModelDynamic->ApplyComponents(PreviewMaterialDynamicWeak.Get());
 		}
 
 		Brush.SetMaterial(PreviewMaterialDynamicWeak.Get());
diff --git a/Engine/Plugins/Experimental/DynamicMaterial/Source/DynamicMaterialEditor/Private/UI/Widgets/Visualizers/SDMTextureUVVisualizerPopout.cpp b/Engine/Plugins/Experimental/DynamicMaterial/Source/DynamicMaterialEditor/Private/UI/Widgets/Visualizers/SDMTextureUVVisualizerPopout.cpp
--- a/Engine/Plugins/Experimental/DynamicMaterial/Source/DynamicMaterialEditor/Private/UI/Widgets/Visualizers/SDMTextureUVVisualizerPopout.cpp
+++ b/Engine/Plugins/Experimental/DynamicMaterial/Source/DynamicMaterialEditor/Private/UI/Widgets/Visualizers/SDMTextureUVVisualizerPopout.cpp
@@ -11,9 +11,11 @@
 #include "DMWorldSubsystem.h"
 #include "Engine/World.h"
 #include "ICustomDetailsView.h"
+#include "IDynamicMaterialEditorModule.h"
 #include "Items/ICustomDetailsViewCustomItem.h"
 #include "Items/ICustomDetailsViewItem.h"
 #include "Materials/Material.h"
+#include "Model/DynamicMaterialModelBase.h"
 #include "Styling/SlateIconFinder.h"
 #include "UI/PropertyGenerators/DMTextureUVDynamicPropertyRowGenerator.h"
 #include "UI/PropertyGenerators/DMTextureUVPropertyRowGenerator.h"
@@ -179,6 +181,7 @@ void SDMTextureUVVisualizerPopout::Construct(const FArguments& InArgs, const TSh
 			.VAlign(EVerticalAlignment::VAlign_Top)
 			[
 				CreatePropertyWidget(
+					InEditorWidget,
 					InArgs._TextureUV 
 						? static_cast<UDMMaterialComponent*>(InArgs._TextureUV)
 						: static_cast<UDMMaterialComponent*>(InArgs._TextureUVDynamic)
@@ -188,6 +191,14 @@ void SDMTextureUVVisualizerPopout::Construct(const FArguments& InArgs, const TSh
 	];
 }
 
+void SDMTextureUVVisualizerPopout::NotifyPreChange(FProperty* InPropertyAboutToChange)
+{	
+}
+
+void SDMTextureUVVisualizerPopout::NotifyPostChange(const FPropertyChangedEvent& InPropertyChangedEvent, FProperty* InPropertyThatChanged)
+{
+}
+
 FText SDMTextureUVVisualizerPopout::GetModeButtonText() const
 {
 	if (Visualizer.IsValid() && Visualizer->IsInPivotEditMode())
@@ -238,7 +249,8 @@ FVector2D SDMTextureUVVisualizerPopout::GetSideBlockSize() const
 	return FVector2D::UnitVector;
 }
 
-TSharedRef<SWidget> SDMTextureUVVisualizerPopout::CreatePropertyWidget(UDMMaterialComponent* InComponent)
+TSharedRef<SWidget> SDMTextureUVVisualizerPopout::CreatePropertyWidget(const TSharedRef<SDMMaterialEditor>& InEditorWidget, 
+	UDMMaterialComponent* InComponent)
 {
 	FDMWidgetLibrary::Get().ClearPropertyHandles(this);
 
@@ -255,7 +267,17 @@ TSharedRef<SWidget> SDMTextureUVVisualizerPopout::CreatePropertyWidget(UDMMateri
 	Args.bAllowResetToDefault = true;
 	Args.bShowCategories = false;
 
-	if (UWorld* World = InComponent->GetWorld())
+	UWorld* World = InComponent->GetWorld();
+
+	if (!World)
+	{
+		if (UDynamicMaterialModelBase* OriginalMaterialModelBase = InEditorWidget->GetOriginalMaterialModelBase())
+		{
+			World = OriginalMaterialModelBase->GetWorld();
+		}
+	}
+
+	if (World)
 	{
 		if (UDMWorldSubsystem* WorldSubsystem = World->GetSubsystem<UDMWorldSubsystem>())
 		{
@@ -267,14 +289,22 @@ TSharedRef<SWidget> SDMTextureUVVisualizerPopout::CreatePropertyWidget(UDMMateri
 	const FCustomDetailsViewItemId RootId = DetailsView->GetRootItem()->GetItemId();
 
 	TArray<FDMPropertyHandle> TextureUVPropertyRows;
+	TSet<UObject*> ProcessedObjects;
+
+	FDMComponentPropertyRowGeneratorParams Params(TextureUVPropertyRows, ProcessedObjects);
+	Params.Owner = this;
+	Params.NotifyHook = this;
+	Params.Object = InComponent;
+	Params.PreviewMaterialModelBase = InEditorWidget->GetPreviewMaterialModelBase();
+	Params.OriginalMaterialModelBase = InEditorWidget->GetOriginalMaterialModelBase();
 
 	if (bIsDynamic)
 	{
-		FDMTextureUVDynamicPropertyRowGenerator::AddPopoutComponentProperties(SharedThis(this), InComponent, TextureUVPropertyRows);
+		FDMTextureUVDynamicPropertyRowGenerator::AddPopoutComponentProperties(Params);
 	}
 	else
 	{
-		FDMTextureUVPropertyRowGenerator::AddPopoutComponentProperties(SharedThis(this), InComponent, TextureUVPropertyRows);
+		FDMTextureUVPropertyRowGenerator::AddPopoutComponentProperties(Params);
 	}
 
 	FDMPropertyHandle EditModeButtonRow;
@@ -322,20 +352,20 @@ TSharedRef<SWidget> SDMTextureUVVisualizerPopout::CreatePropertyWidget(UDMMateri
 	{
 		const bool bHasValidCustomWidget = EditRow.ValueWidget.IsValid() && !EditRow.ValueName.IsNone() && EditRow.NameOverride.IsSet();
 
-		if (!EditRow.DetailTreeNode && !bHasValidCustomWidget)
+		if (!EditRow.PreviewHandle.DetailTreeNode && !bHasValidCustomWidget)
 		{
 			continue;
 		}
 
 		ECustomDetailsTreeInsertPosition Position = ECustomDetailsTreeInsertPosition::Child;
 
-		if (EditRow.PropertyHandle.IsValid())
+		if (EditRow.PreviewHandle.PropertyHandle.IsValid())
 		{
-			if (EditRow.PropertyHandle->HasMetaData("HighPriority"))
+			if (EditRow.PreviewHandle.PropertyHandle->HasMetaData("HighPriority"))
 			{
 				Position = ECustomDetailsTreeInsertPosition::FirstChild;
 			}
-			else if (EditRow.PropertyHandle->HasMetaData("LowPriority"))
+			else if (EditRow.PreviewHandle.PropertyHandle->HasMetaData("LowPriority"))
 			{
 				Position = ECustomDetailsTreeInsertPosition::LastChild;
 			}
@@ -365,12 +395,12 @@ TSharedRef<SWidget> SDMTextureUVVisualizerPopout::CreatePropertyWidget(UDMMateri
 			continue;
 		}
 
-		if (!EditRow.DetailTreeNode)
+		if (!EditRow.PreviewHandle.DetailTreeNode)
 		{
 			continue;
 		}
 
-		TSharedRef<ICustomDetailsViewItem> Item = DetailsView->CreateDetailTreeItem(DetailsView->GetRootItem(), EditRow.DetailTreeNode.ToSharedRef());
+		TSharedRef<ICustomDetailsViewItem> Item = DetailsView->CreateDetailTreeItem(DetailsView->GetRootItem(), EditRow.PreviewHandle.DetailTreeNode.ToSharedRef());
 
 		if (EditRow.NameOverride.IsSet())
 		{
@@ -388,7 +418,7 @@ TSharedRef<SWidget> SDMTextureUVVisualizerPopout::CreatePropertyWidget(UDMMateri
 			Item->SetEnabledOverride(false);
 		}
 
-		if (EditRow.PropertyHandle->HasMetaData("NotKeyframeable"))
+		if (EditRow.PreviewHandle.PropertyHandle->HasMetaData("NotKeyframeable"))
 		{
 			Item->SetKeyframeEnabled(false);
 		}
diff --git a/Engine/Plugins/Experimental/DynamicMaterial/Source/DynamicMaterialEditor/Private/UI/Widgets/Visualizers/SDMTextureUVVisualizerPopout.h b/Engine/Plugins/Experimental/DynamicMaterial/Source/DynamicMaterialEditor/Private/UI/Widgets/Visualizers/SDMTextureUVVisualizerPopout.h
--- a/Engine/Plugins/Experimental/DynamicMaterial/Source/DynamicMaterialEditor/Private/UI/Widgets/Visualizers/SDMTextureUVVisualizerPopout.h
+++ b/Engine/Plugins/Experimental/DynamicMaterial/Source/DynamicMaterialEditor/Private/UI/Widgets/Visualizers/SDMTextureUVVisualizerPopout.h
@@ -2,6 +2,7 @@
 
 #pragma once
 
+#include "Misc/NotifyHook.h"
 #include "Widgets/SCompoundWidget.h"
 
 #include "Templates/SharedPointerFwd.h"
@@ -24,7 +25,7 @@ enum class ECheckBoxState : uint8;
  * The popout specifically expands the visible area of the preview to 3x the normal size
  * on the smallest axis. The other axis is expanded to match the aspect ratio.
  */
-class SDMTextureUVVisualizerPopout : public SCompoundWidget
+class SDMTextureUVVisualizerPopout : public SCompoundWidget, public FNotifyHook
 {
 public:
 	static const FName TabId;
@@ -44,6 +45,11 @@ public:
 	/** The TextureUV should be a sub-property of the stage */
 	void Construct(const FArguments& InArgs, const TSharedRef<SDMMaterialEditor>& InEditorWidget, UDMMaterialStage* InMaterialStage);
 
+	//~ Begin FNotifyHook
+	virtual void NotifyPreChange(FProperty* InPropertyAboutToChange) override;
+	virtual void NotifyPostChange(const FPropertyChangedEvent& InPropertyChangedEvent, FProperty* InPropertyThatChanged) override;
+	//~ End FNotifyHook
+
 protected:
 	TSharedPtr<SDMTextureUVVisualizer> Visualizer;
 
@@ -53,7 +59,7 @@ protected:
 
 	FVector2D GetSideBlockSize() const;
 
-	TSharedRef<SWidget> CreatePropertyWidget(UDMMaterialComponent* InComponent);
+	TSharedRef<SWidget> CreatePropertyWidget(const TSharedRef<SDMMaterialEditor>& InEditorWidget, UDMMaterialComponent* InComponent);
 
 	ECheckBoxState GetModeCheckBoxState(bool bInIsPivot) const;
 
diff --git a/Engine/Plugins/Experimental/DynamicMaterial/Source/DynamicMaterialEditor/Private/Utils/DMMaterialModelFunctionLibrary.cpp b/Engine/Plugins/Experimental/DynamicMaterial/Source/DynamicMaterialEditor/Private/Utils/DMMaterialModelFunctionLibrary.cpp
--- a/Engine/Plugins/Experimental/DynamicMaterial/Source/DynamicMaterialEditor/Private/Utils/DMMaterialModelFunctionLibrary.cpp
+++ b/Engine/Plugins/Experimental/DynamicMaterial/Source/DynamicMaterialEditor/Private/Utils/DMMaterialModelFunctionLibrary.cpp
@@ -9,6 +9,7 @@
 #include "ContentBrowserModule.h"
 #include "DMWorldSubsystem.h"
 #include "DynamicMaterialEditorModule.h"
+#include "Engine.h"
 #include "Engine/Texture.h"
 #include "Engine/World.h"
 #include "EngineAnalytics.h"
@@ -545,4 +546,50 @@ FString UDMMaterialModelFunctionLibrary::RemoveAssetPrefix(const FString& InAsse
 	return InAssetName;
 }
 
+UDynamicMaterialModelBase* UDMMaterialModelFunctionLibrary::CreatePreviewModel(UDynamicMaterialModelBase* InOriginalModelBase)
+{
+	UDynamicMaterialModelBase* PreviewMaterialModelBase = Cast<UDynamicMaterialModelBase>(StaticDuplicateObject(InOriginalModelBase, GetTransientPackage()));
+	PreviewMaterialModelBase->SetDynamicMaterialInstance(nullptr);
+	PreviewMaterialModelBase->MarkOriginalUpdated();
+
+	if (UDynamicMaterialModelDynamic* MaterialModelDynamic = Cast<UDynamicMaterialModelDynamic>(PreviewMaterialModelBase))
+	{
+		MaterialModelDynamic->EnsureComponents();
+	}
+
+	if (UDynamicMaterialModelEditorOnlyData* EditorOnlyData = UDynamicMaterialModelEditorOnlyData::Get(PreviewMaterialModelBase))
+	{
+		EditorOnlyData->RequestMaterialBuild(/* Immediate */ true);
+	}
+
+	return PreviewMaterialModelBase;
+}
+
+void UDMMaterialModelFunctionLibrary::MirrorMaterialModel(UDynamicMaterialModelBase* InSource, UDynamicMaterialModelBase*& InTarget)
+{
+	if (!InSource || !InTarget || InSource->GetClass() != InTarget->GetClass())
+	{
+		return;
+	}
+
+	TMap<UObject*, UObject*> ReplacedObjects;
+
+	FObjectDuplicationParameters Params(InSource, InTarget->GetOuter());
+	Params.CreatedObjects = &ReplacedObjects;
+	Params.DestClass = InSource->GetClass();
+	Params.DestName = InTarget->GetFName();
+
+	if (UDynamicMaterialInstance* SourceInstance = InSource->GetDynamicMaterialInstance())
+	{
+		if (UDynamicMaterialInstance* TargetInstance = InTarget->GetDynamicMaterialInstance())
+		{
+			Params.DuplicationSeed.Add(SourceInstance, TargetInstance);
+		}
+	}
+
+	InTarget = Cast<UDynamicMaterialModelBase>(StaticDuplicateObjectEx(Params));
+
+	GEngine->NotifyToolsOfObjectReplacement(ReplacedObjects);
+}
+
 #undef LOCTEXT_NAMESPACE
diff --git a/Engine/Plugins/Experimental/DynamicMaterial/Source/DynamicMaterialEditor/Private/Utils/DMPrivate.cpp b/Engine/Plugins/Experimental/DynamicMaterial/Source/DynamicMaterialEditor/Private/Utils/DMPrivate.cpp
--- a/Engine/Plugins/Experimental/DynamicMaterial/Source/DynamicMaterialEditor/Private/Utils/DMPrivate.cpp
+++ b/Engine/Plugins/Experimental/DynamicMaterial/Source/DynamicMaterialEditor/Private/Utils/DMPrivate.cpp
@@ -139,6 +139,36 @@ namespace UE::DynamicMaterialEditor::Private
 
 		return GetMaterialPropertyLongDisplayName(InMaterialProperty);
 	}
+
+	UObject* FindSubobject(UObject* InOuter, FStringView InPath)
+	{
+		if (InPath.IsEmpty())
+		{
+			return InOuter;
+		}
+
+		int32 Index = INDEX_NONE;
+		InPath.FindChar(TEXT('.'), Index);
+
+		if (Index >= 0)
+		{
+			const FStringView FirstPath = InPath.Left(Index);
+			InPath.RightChopInline(Index + 1);
+
+			if (!FirstPath.IsEmpty())
+			{
+				const FString PathSection(FirstPath);
+
+				if (UObject* Subobject = StaticFindObjectFast(UObject::StaticClass(), InOuter, *PathSection))
+				{
+					return FindSubobject(Subobject, InPath);
+				}
+			}
+		}
+
+		const FString Path(InPath);
+		return StaticFindObjectFast(UObject::StaticClass(), InOuter, *Path);
+	}
 }
 
 FDMMaterialLayerReference::FDMMaterialLayerReference()
diff --git a/Engine/Plugins/Experimental/DynamicMaterial/Source/DynamicMaterialEditor/Private/Utils/DMPrivate.h b/Engine/Plugins/Experimental/DynamicMaterial/Source/DynamicMaterialEditor/Private/Utils/DMPrivate.h
--- a/Engine/Plugins/Experimental/DynamicMaterial/Source/DynamicMaterialEditor/Private/Utils/DMPrivate.h
+++ b/Engine/Plugins/Experimental/DynamicMaterial/Source/DynamicMaterialEditor/Private/Utils/DMPrivate.h
@@ -39,6 +39,14 @@ namespace UE::DynamicMaterialEditor::Private
 	FText GetMaterialPropertyLongDisplayName(EDMMaterialPropertyType InMaterialProperty);
 
 	FText GetMaterialPropertyShortDisplayName(EDMMaterialPropertyType InMaterialProperty);
+
+	UObject* FindSubobject(UObject* InOuter, FStringView InPath);
+
+	template<typename InClassName>
+	InClassName* FindSubobject(UObject* InOuter, FStringView InPath)
+	{
+		return Cast<InClassName>(FindSubobject(InOuter, InPath));
+	}
 }
 
 struct FDMMaterialLayerReference
diff --git a/Engine/Plugins/Experimental/DynamicMaterial/Source/DynamicMaterialEditor/Public/DMEDefs.h b/Engine/Plugins/Experimental/DynamicMaterial/Source/DynamicMaterialEditor/Public/DMEDefs.h
--- a/Engine/Plugins/Experimental/DynamicMaterial/Source/DynamicMaterialEditor/Public/DMEDefs.h
+++ b/Engine/Plugins/Experimental/DynamicMaterial/Source/DynamicMaterialEditor/Public/DMEDefs.h
@@ -121,11 +121,17 @@ enum class EDMPropertyHandlePriority : uint8
 	High
 };
 
-struct FDMPropertyHandle
+struct FDMPropertyGeneratorHandle
 {
 	TSharedPtr<IPropertyRowGenerator> PropertyRowGenerator;
 	TSharedPtr<IDetailTreeNode> DetailTreeNode;
 	TSharedPtr<IPropertyHandle> PropertyHandle;
+};
+
+struct FDMPropertyHandle
+{
+	FDMPropertyGeneratorHandle OriginalHandle;
+	FDMPropertyGeneratorHandle PreviewHandle;
 	TOptional<FText> NameOverride;
 	TOptional<FText> NameToolTipOverride;
 	TOptional<FResetToDefaultOverride> ResetToDefaultOverride;
diff --git a/Engine/Plugins/Experimental/DynamicMaterial/Source/DynamicMaterialEditor/Public/IDynamicMaterialEditorModule.h b/Engine/Plugins/Experimental/DynamicMaterial/Source/DynamicMaterialEditor/Public/IDynamicMaterialEditorModule.h
--- a/Engine/Plugins/Experimental/DynamicMaterial/Source/DynamicMaterialEditor/Public/IDynamicMaterialEditorModule.h
+++ b/Engine/Plugins/Experimental/DynamicMaterial/Source/DynamicMaterialEditor/Public/IDynamicMaterialEditorModule.h
@@ -7,8 +7,10 @@
 #include "Containers/ContainersFwd.h"
 #include "Delegates/Delegate.h"
 #include "Templates/SharedPointerFwd.h"
+#include "UI/Utils/IDMWidgetLibrary.h"
 
 class AActor;
+class FNotifyHook;
 class SDMMaterialComponentEditor;
 class UClass;
 class UDMMaterialComponent;
@@ -17,13 +19,35 @@ class UDynamicMaterialModelBase;
 class UObject;
 struct FDMObjectMaterialProperty;
 struct FDMPropertyHandle;
+struct FDMPropertyHandleGenerateParams;
 struct IDMOnWizardCompleteCallback;
 
+struct FDMComponentPropertyRowGeneratorParams
+{
+	const SWidget* Owner = nullptr;
+	FNotifyHook* NotifyHook = nullptr;
+	UDynamicMaterialModelBase* PreviewMaterialModelBase = nullptr;
+	UDynamicMaterialModelBase* OriginalMaterialModelBase = nullptr;
+	UObject* Object = nullptr;
+	TArray<FDMPropertyHandle>& PropertyRows;
+	TSet<UObject*>& ProcessedObjects;
+
+	FDMComponentPropertyRowGeneratorParams(TArray<FDMPropertyHandle>& InPropertyRows, TSet<UObject*>& InProcessedObjects)
+		: PropertyRows(InPropertyRows)
+		, ProcessedObjects(InProcessedObjects)
+	{
+	}
+
+	FDMPropertyHandleGenerateParams CreatePropertyHandleParams(FName InPropertyName) const
+	{
+		return {Owner, NotifyHook, PreviewMaterialModelBase, OriginalMaterialModelBase, Object, InPropertyName};
+	}
+};
+
 DECLARE_DELEGATE_RetVal_OneParam(TArray<FDMObjectMaterialProperty>, FDMGetObjectMaterialPropertiesDelegate, UObject* InObject)
 
 /** Creates property rows in the edit widget. */
-DECLARE_DELEGATE_FourParams(FDMComponentPropertyRowGeneratorDelegate, const TSharedRef<SDMMaterialComponentEditor>&, UDMMaterialComponent*,
-	TArray<FDMPropertyHandle>&, TSet<UDMMaterialComponent*>&)
+DECLARE_DELEGATE_OneParam(FDMComponentPropertyRowGeneratorDelegate, FDMComponentPropertyRowGeneratorParams&)
 
 /** Material Designer - Build your own materials in a slimline editor! */
 class IDynamicMaterialEditorModule : public IModuleInterface
diff --git a/Engine/Plugins/Experimental/DynamicMaterial/Source/DynamicMaterialEditor/Public/Model/DynamicMaterialModelEditorOnlyData.h b/Engine/Plugins/Experimental/DynamicMaterial/Source/DynamicMaterialEditor/Public/Model/DynamicMaterialModelEditorOnlyData.h
--- a/Engine/Plugins/Experimental/DynamicMaterial/Source/DynamicMaterialEditor/Public/Model/DynamicMaterialModelEditorOnlyData.h
+++ b/Engine/Plugins/Experimental/DynamicMaterial/Source/DynamicMaterialEditor/Public/Model/DynamicMaterialModelEditorOnlyData.h
@@ -264,48 +264,53 @@ protected:
 	EDMState State;
 
 	UPROPERTY(EditInstanceOnly, BlueprintReadOnly, Category = "Material Designer", 
-		meta = (ValidEnumValues = "MD_Surface,MD_PostProcess,MD_DeferredDecal,MD_LightFunction"))
+		meta = (NotKeyframeable, ValidEnumValues = "MD_Surface,MD_PostProcess,MD_DeferredDecal,MD_LightFunction"))
 	TEnumAsByte<EMaterialDomain> Domain;
 
 	UPROPERTY(EditInstanceOnly, BlueprintReadOnly, Category = "Material Designer",
-		meta = (ValidEnumValues = "BLEND_Opaque,BLEND_Translucent,BLEND_Masked,BLEND_Additive,BLEND_Modulate"))
+		meta = (NotKeyframeable, ValidEnumValues = "BLEND_Opaque,BLEND_Translucent,BLEND_Masked,BLEND_Additive,BLEND_Modulate"))
 	TEnumAsByte<EBlendMode> BlendMode;
 
-	UPROPERTY(EditInstanceOnly, BlueprintReadOnly, Category = "Material Designer", meta = (ValidEnumValues = "Unlit,DefaultLit"))
+	UPROPERTY(EditInstanceOnly, BlueprintReadOnly, Category = "Material Designer", 
+		meta = (NotKeyframeable, ValidEnumValues = "Unlit,DefaultLit"))
 	EDMMaterialShadingModel ShadingModel;
 
 	/**
 	 * Whether the opaque material has any pixel animations happening, that isn't included in the geometric velocities.
 	 * This allows to disable renderer's heuristics that assumes animation is fully described with motion vector, such as TSR's anti-flickering heuristic.
 	 */
-	UPROPERTY(EditInstanceOnly, BlueprintReadOnly, Category = "Material Designer")
+	UPROPERTY(EditInstanceOnly, BlueprintReadOnly, Category = "Material Designer", meta = (NotKeyframeable))
 	bool bHasPixelAnimation;
 
 	/** Indicates that the material should be rendered without backface culling and the normal should be flipped for backfaces. */
-	UPROPERTY(EditInstanceOnly, BlueprintReadOnly, Category = "Material Designer")
+	UPROPERTY(EditInstanceOnly, BlueprintReadOnly, Category = "Material Designer", meta = (NotKeyframeable))
 	bool bTwoSided;
 
 	/** When true, translucent materials will output motion vectors and write to depth buffer in velocity pass. */
-	UPROPERTY(EditInstanceOnly, BlueprintReadOnly, Category = "Material Designer", meta = (DisplayName = "Output Translucent Velocity"))
+	UPROPERTY(EditInstanceOnly, BlueprintReadOnly, Category = "Material Designer", 
+		meta = (NotKeyframeable, DisplayName = "Output Translucent Velocity"))
 	bool bOutputTranslucentVelocityEnabled;
 
 	/** Whether tessellation is enabled on the material. NOTE: Required for displacement to work. */
-	UPROPERTY(EditInstanceOnly, BlueprintReadOnly, Category = "Material Designer", meta = (DisplayName = "Nanite Tessellation"))
+	UPROPERTY(EditInstanceOnly, BlueprintReadOnly, Category = "Material Designer", 
+		meta = (NotKeyframeable, DisplayName = "Nanite Tessellation"))
 	bool bNaniteTessellationEnabled;
 
 	/** Mid point for displacement in the range 0-1. */
-	UPROPERTY(EditInstanceOnly, BlueprintReadOnly, Category = "Material Designer", meta = ( UIMin = 0, UIMax = 1, ClampMin = 0, ClampMax = 1))
+	UPROPERTY(EditInstanceOnly, BlueprintReadOnly, Category = "Material Designer", 
+		meta = (NotKeyframeable, UIMin = 0, UIMax = 1, ClampMin = 0, ClampMax = 1))
 	float DisplacementCenter;
 
 	/** Multipler for displacement values. */
-	UPROPERTY(EditInstanceOnly, BlueprintReadOnly, Category = "Material Designer")
+	UPROPERTY(EditInstanceOnly, BlueprintReadOnly, Category = "Material Designer", meta = (NotKeyframeable))
 	float DisplacementMagnitude;
 
 	/**
 	 * Indicates that the material should be rendered using responsive anti-aliasing. Improves sharpness of small moving particles such as sparks.
 	 * Only use for small moving features because it will cause aliasing of the background.
 	 */
-	UPROPERTY(EditInstanceOnly, BlueprintReadOnly, Category = "Material Designer", meta = (DisplayName = "Responsive AA"))
+	UPROPERTY(EditInstanceOnly, BlueprintReadOnly, Category = "Material Designer", 
+		meta = (NotKeyframeable, DisplayName = "Responsive AA"))
 	bool bResponsiveAAEnabled;
 
 	UPROPERTY(VisibleInstanceOnly, BlueprintReadOnly, TextExportTransient, Category = "Material Designer")
diff --git a/Engine/Plugins/Experimental/DynamicMaterial/Source/DynamicMaterialEditor/Public/UI/Menus/DMMenuContext.h b/Engine/Plugins/Experimental/DynamicMaterial/Source/DynamicMaterialEditor/Public/UI/Menus/DMMenuContext.h
--- a/Engine/Plugins/Experimental/DynamicMaterial/Source/DynamicMaterialEditor/Public/UI/Menus/DMMenuContext.h
+++ b/Engine/Plugins/Experimental/DynamicMaterial/Source/DynamicMaterialEditor/Public/UI/Menus/DMMenuContext.h
@@ -50,9 +50,9 @@ public:
 
 	DYNAMICMATERIALEDITOR_API UDMMaterialSlot* GetSlot() const;
 
-	DYNAMICMATERIALEDITOR_API UDynamicMaterialModelBase* GetModelBase() const;
+	DYNAMICMATERIALEDITOR_API UDynamicMaterialModelBase* GetPreviewModelBase() const;
 
-	DYNAMICMATERIALEDITOR_API UDynamicMaterialModel* GetModel() const;
+	DYNAMICMATERIALEDITOR_API UDynamicMaterialModel* GetPreviewModel() const;
 
 	DYNAMICMATERIALEDITOR_API UDMMaterialStage* GetStage() const;
 
diff --git a/Engine/Plugins/Experimental/DynamicMaterial/Source/DynamicMaterialEditor/Public/UI/PropertyGenerators/DMComponentPropertyRowGenerator.h b/Engine/Plugins/Experimental/DynamicMaterial/Source/DynamicMaterialEditor/Public/UI/PropertyGenerators/DMComponentPropertyRowGenerator.h
--- a/Engine/Plugins/Experimental/DynamicMaterial/Source/DynamicMaterialEditor/Public/UI/PropertyGenerators/DMComponentPropertyRowGenerator.h
+++ b/Engine/Plugins/Experimental/DynamicMaterial/Source/DynamicMaterialEditor/Public/UI/PropertyGenerators/DMComponentPropertyRowGenerator.h
@@ -4,12 +4,15 @@
 
 #include "Containers/ContainersFwd.h"
 #include "Templates/SharedPointerFwd.h"
+#include "UI/Utils/IDMWidgetLibrary.h"
 
-enum class EDMPropertyHandlePriority : uint8;
 class FName;
+class FNotifyHook;
 class FProperty;
 class SDMMaterialComponentEditor;
 class UDMMaterialComponent;
+enum class EDMPropertyHandlePriority : uint8;
+struct FDMComponentPropertyRowGeneratorParams;
 struct FDMPropertyHandle;
 
 /** Used to generate editable properties for objects editable in the Material Designer. */
@@ -27,29 +30,26 @@ public:
 	 * @param InOutPropertyRows The generated rows.
 	 * @param InOutProcessedObjects The already processed objects - add to this to avoid possible recursive generation.
 	 */
-	DYNAMICMATERIALEDITOR_API virtual void AddComponentProperties(const TSharedRef<SDMMaterialComponentEditor>& InComponentEditorWidget, 
-		UDMMaterialComponent* InComponent, TArray<FDMPropertyHandle>& InOutPropertyRows, TSet<UDMMaterialComponent*>& InOutProcessedObjects);
+	DYNAMICMATERIALEDITOR_API virtual void AddComponentProperties(FDMComponentPropertyRowGeneratorParams& InParams);
 
 	/**
 	 * Add the rows needed for a specific property.
 	 * @param InProperty The name of the property.
 	 */
-	DYNAMICMATERIALEDITOR_API virtual void AddPropertyEditRows(const TSharedRef<SDMMaterialComponentEditor>& InComponentEditorWidget, 
-		UDMMaterialComponent* InComponent, const FName& InProperty, TArray<FDMPropertyHandle>& InOutPropertyRows, 
-		TSet<UDMMaterialComponent*>& InOutProcessedObjects);
+	DYNAMICMATERIALEDITOR_API virtual void AddPropertyEditRows(FDMComponentPropertyRowGeneratorParams& InParams,
+		const FName& InProperty);
 
 	/**
 	 * Adds the rows needed for a specific property.
 	 * @param InProperty The FProperty for this property.
-	 * @param MemoryPtr The pointer to the property's value in memory
+	 * @param InMemoryPtr The pointer to the property's value in memory
 	 */
-	DYNAMICMATERIALEDITOR_API virtual void AddPropertyEditRows(const TSharedRef<SDMMaterialComponentEditor>& InComponentEditorWidget, 
-		UDMMaterialComponent* InComponent, FProperty* InProperty, void* MemoryPtr, TArray<FDMPropertyHandle>& InOutPropertyRows, 
-		TSet<UDMMaterialComponent*>& InOutProcessedObjects);
+	DYNAMICMATERIALEDITOR_API virtual void AddPropertyEditRows(FDMComponentPropertyRowGeneratorParams& InParams,
+		FProperty* InProperty, void* InMemoryPtr);
 
 	DYNAMICMATERIALEDITOR_API virtual bool AllowKeyframeButton(UDMMaterialComponent* InComponent, FProperty* InProperty);
 
 protected:
 	/** Returns true if the component edit widget is editing a Material Designer Dynamic. */
-	DYNAMICMATERIALEDITOR_API static bool IsDynamic(const TSharedRef<SDMMaterialComponentEditor>& InComponentEditorWidget);
+	DYNAMICMATERIALEDITOR_API static bool IsDynamic(FDMComponentPropertyRowGeneratorParams& InParams);
 };
