Horde issue  818426

diff --git a/Engine/Plugins/Experimental/DynamicMaterial/Source/DynamicMaterialEditor/Private/Utils/DMMaterialModelFunctionLibrary.cpp b/Engine/Plugins/Experimental/DynamicMaterial/Source/DynamicMaterialEditor/Private/Utils/DMMaterialModelFunctionLibrary.cpp
--- a/Engine/Plugins/Experimental/DynamicMaterial/Source/DynamicMaterialEditor/Private/Utils/DMMaterialModelFunctionLibrary.cpp
+++ b/Engine/Plugins/Experimental/DynamicMaterial/Source/DynamicMaterialEditor/Private/Utils/DMMaterialModelFunctionLibrary.cpp
@@ -9,7 +9,7 @@
 #include "ContentBrowserModule.h"
 #include "DMWorldSubsystem.h"
 #include "DynamicMaterialEditorModule.h"
-#include "Engine.h"
+#include "Engine/Engine.h"
 #include "Engine/Texture.h"
 #include "Engine/World.h"
 #include "EngineAnalytics.h"
