Horde issue 818513:

diff --git a/Engine/Plugins/Experimental/DynamicMaterial/Source/DynamicMaterialEditor/Private/UI/PropertyGenerators/DMMaterialEffectFunctionPropertyRowGenerator.cpp b/Engine/Plugins/Experimental/DynamicMaterial/Source/DynamicMaterialEditor/Private/UI/PropertyGenerators/DMMaterialEffectFunctionPropertyRowGenerator.cpp
--- a/Engine/Plugins/Experimental/DynamicMaterial/Source/DynamicMaterialEditor/Private/UI/PropertyGenerators/DMMaterialEffectFunctionPropertyRowGenerator.cpp
+++ b/Engine/Plugins/Experimental/DynamicMaterial/Source/DynamicMaterialEditor/Private/UI/PropertyGenerators/DMMaterialEffectFunctionPropertyRowGenerator.cpp
@@ -104,7 +104,7 @@ void FDMMaterialEffectFunctionPropertyRowGenerator::AddComponentProperties(FDMCo
 		const FText MaterialInputText = FText::Format(MaterialInputFormat, EffectFunction->GetEffectName());
 		const FName MaterialInputName = *MaterialInputText.ToString();
 
-		for (FDMPropertyHandle& ValuePropertyRow : ValuePropertyRows)
+		for (FDMPropertyHandle& ValuePropertyRow : ValuePropertyRows) //-V1078
 		{
 			ValuePropertyRow.NameToolTipOverride = Description;
 			ValuePropertyRow.CategoryOverrideName = MaterialInputName;
diff --git a/Engine/Plugins/Experimental/DynamicMaterial/Source/DynamicMaterialEditor/Private/UI/PropertyGenerators/DMMaterialStageFunctionPropertyRowGenerator.cpp b/Engine/Plugins/Experimental/DynamicMaterial/Source/DynamicMaterialEditor/Private/UI/PropertyGenerators/DMMaterialStageFunctionPropertyRowGenerator.cpp
--- a/Engine/Plugins/Experimental/DynamicMaterial/Source/DynamicMaterialEditor/Private/UI/PropertyGenerators/DMMaterialStageFunctionPropertyRowGenerator.cpp
+++ b/Engine/Plugins/Experimental/DynamicMaterial/Source/DynamicMaterialEditor/Private/UI/PropertyGenerators/DMMaterialStageFunctionPropertyRowGenerator.cpp
@@ -118,7 +118,7 @@ void FDMMaterialStageFunctionPropertyRowGenerator::AddComponentProperties(FDMCom
 			MaterialInputText = FText::Format(MaterialInputFormat, MaterialInputText);
 			const FName MaterialInputName = *MaterialInputText.ToString();
 
-			for (FDMPropertyHandle& ValuePropertyRow : ValuePropertyRows)
+			for (FDMPropertyHandle& ValuePropertyRow : ValuePropertyRows) //-V1078
 			{
 				ValuePropertyRow.NameToolTipOverride = Description;
 				ValuePropertyRow.CategoryOverrideName = MaterialInputName;
