Material Designer: Auto-compile option
- You now have the option to disable automatic preview material compile.
- You can force a compile by clicking the compile button at the top.
- Compilation will become available in 2 cases:
- When you make a structural change.
- When you make a value/parameter change and automatic-parameter-copy-to-source is disabled.
- The compile and apply buttons are hidden if it's a material instance.

diff --git a/Engine/Plugins/Experimental/DynamicMaterial/Source/DynamicMaterial/Private/Material/DynamicMaterialInstance.cpp b/Engine/Plugins/Experimental/DynamicMaterial/Source/DynamicMaterial/Private/Material/DynamicMaterialInstance.cpp
--- a/Engine/Plugins/Experimental/DynamicMaterial/Source/DynamicMaterial/Private/Material/DynamicMaterialInstance.cpp
+++ b/Engine/Plugins/Experimental/DynamicMaterial/Source/DynamicMaterial/Private/Material/DynamicMaterialInstance.cpp
@@ -141,7 +141,11 @@ void UDynamicMaterialInstance::PostDuplicate(bool bInDuplicateForPIE)
 		{
 			if (IDynamicMaterialModelEditorOnlyDataInterface* ModelEditorOnlyData = MaterialModel->GetEditorOnlyData())
 			{
-				ModelEditorOnlyData->RequestMaterialBuild(/* Immediate */ bInDuplicateForPIE);
+				ModelEditorOnlyData->RequestMaterialBuild(
+					bInDuplicateForPIE
+						? EDMBuildRequestType::Immediate
+						: EDMBuildRequestType::Async
+				);
 			}
 		}
 	}
diff --git a/Engine/Plugins/Experimental/DynamicMaterial/Source/DynamicMaterial/Private/Model/DynamicMaterialModel.cpp b/Engine/Plugins/Experimental/DynamicMaterial/Source/DynamicMaterial/Private/Model/DynamicMaterialModel.cpp
--- a/Engine/Plugins/Experimental/DynamicMaterial/Source/DynamicMaterial/Private/Model/DynamicMaterialModel.cpp
+++ b/Engine/Plugins/Experimental/DynamicMaterial/Source/DynamicMaterial/Private/Model/DynamicMaterialModel.cpp
@@ -717,7 +717,7 @@ void UDynamicMaterialModel::PostDuplicate(bool bInDuplicateForPIE)
 
 		if (IDynamicMaterialModelEditorOnlyDataInterface* ModelEditorOnlyData = GetEditorOnlyData())
 		{
-			ModelEditorOnlyData->RequestMaterialBuild();
+			ModelEditorOnlyData->RequestMaterialBuild(EDMBuildRequestType::Async);
 		}
 	}
 }
diff --git a/Engine/Plugins/Experimental/DynamicMaterial/Source/DynamicMaterial/Public/Model/IDynamicMaterialModelEditorOnlyDataInterface.h b/Engine/Plugins/Experimental/DynamicMaterial/Source/DynamicMaterial/Public/Model/IDynamicMaterialModelEditorOnlyDataInterface.h
--- a/Engine/Plugins/Experimental/DynamicMaterial/Source/DynamicMaterial/Public/Model/IDynamicMaterialModelEditorOnlyDataInterface.h
+++ b/Engine/Plugins/Experimental/DynamicMaterial/Source/DynamicMaterial/Public/Model/IDynamicMaterialModelEditorOnlyDataInterface.h
@@ -21,6 +21,22 @@ enum class EDMUpdateType : uint8;
 struct FDMComponentPathSegment;
 struct FDMComponentPath;
 
+UENUM(BlueprintType)
+enum class EDMBuildRequestType : uint8
+{
+	/** Compile the material immediately. */
+	Immediate,
+
+	/** Will try to add to the build queue or fallback to immediate compile. */
+	Async,
+
+	/**
+	 * If bAutomaticallyCompilePreviewMaterial is false, will just mark the material as changed,
+	 * but not actually try to compile. It will fall back to Async if automatic compile is true.
+	 */
+	Preview
+};
+
 /** Interface for the editor-only data so that editor-only parts of runtime components can interact with editor-only features. */
 UINTERFACE(MinimalAPI, BlueprintType, meta = (CannotImplementInterfaceInBlueprint))
 class UDynamicMaterialModelEditorOnlyDataInterface : public UInterface
@@ -50,7 +66,7 @@ public:
 #if WITH_EDITOR
 	/** Called when the model needs to have the material rebuild. */
 	UFUNCTION(BlueprintCallable, Category = "Material Designer")
-	virtual void RequestMaterialBuild(bool bInImmediate = false) PURE_VIRTUAL(UDynamicMaterialModelEditorOnlyDataInterface::RequestMaterialBuild)
+	virtual void RequestMaterialBuild(EDMBuildRequestType InRequestType = EDMBuildRequestType::Preview) PURE_VIRTUAL(UDynamicMaterialModelEditorOnlyDataInterface::RequestMaterialBuild)
 
 	/** Called to create the build state for building materials. */
 	virtual TSharedRef<IDMMaterialBuildStateInterface> CreateBuildStateInterface(UMaterial* InMaterialToBuild) const = 0;
diff --git a/Engine/Plugins/Experimental/DynamicMaterial/Source/DynamicMaterialEditor/Private/DynamicMaterialEditorSettings.cpp b/Engine/Plugins/Experimental/DynamicMaterial/Source/DynamicMaterialEditor/Private/DynamicMaterialEditorSettings.cpp
--- a/Engine/Plugins/Experimental/DynamicMaterial/Source/DynamicMaterialEditor/Private/DynamicMaterialEditorSettings.cpp
+++ b/Engine/Plugins/Experimental/DynamicMaterial/Source/DynamicMaterialEditor/Private/DynamicMaterialEditorSettings.cpp
@@ -135,6 +135,7 @@ UDynamicMaterialEditorSettings::UDynamicMaterialEditorSettings()
 	bFollowSelection = true;
 	bAddDetailsPanelButton = false;
 	bAutomaticallyCopyParametersToSourceMaterial = true;
+	bAutomaticallyCompilePreviewMaterial = false;
 	bUseLinearColorForVectors = true;
 
 	ResetAllLayoutSettings();
diff --git a/Engine/Plugins/Experimental/DynamicMaterial/Source/DynamicMaterialEditor/Private/DynamicMaterialEditorSettings.h b/Engine/Plugins/Experimental/DynamicMaterial/Source/DynamicMaterialEditor/Private/DynamicMaterialEditorSettings.h
--- a/Engine/Plugins/Experimental/DynamicMaterial/Source/DynamicMaterialEditor/Private/DynamicMaterialEditorSettings.h
+++ b/Engine/Plugins/Experimental/DynamicMaterial/Source/DynamicMaterialEditor/Private/DynamicMaterialEditorSettings.h
@@ -199,6 +199,12 @@ public:
 	UPROPERTY(Config, EditAnywhere, BlueprintReadWrite, Category = "Editor")
 	bool bAutomaticallyCopyParametersToSourceMaterial;
 
+	/**
+	 * When a structural material change is made, whether to automatically recompile the preview material to show this change.
+	 */
+	UPROPERTY(Config, EditAnywhere, BlueprintReadWrite, Category = "Editor")
+	bool bAutomaticallyCompilePreviewMaterial;
+
 	UPROPERTY(Config, EditAnywhere, BlueprintReadWrite, Category = "Layout")
 	EDMMaterialEditorLayout Layout;
 
diff --git a/Engine/Plugins/Experimental/DynamicMaterial/Source/DynamicMaterialEditor/Private/Material/DynamicMaterialInstanceFactory.cpp b/Engine/Plugins/Experimental/DynamicMaterial/Source/DynamicMaterialEditor/Private/Material/DynamicMaterialInstanceFactory.cpp
--- a/Engine/Plugins/Experimental/DynamicMaterial/Source/DynamicMaterialEditor/Private/Material/DynamicMaterialInstanceFactory.cpp
+++ b/Engine/Plugins/Experimental/DynamicMaterial/Source/DynamicMaterialEditor/Private/Material/DynamicMaterialInstanceFactory.cpp
@@ -58,7 +58,7 @@ UObject* UDynamicMaterialInstanceFactory::FactoryCreateNew(UClass* InClass, UObj
 	{
 		if (IDynamicMaterialModelEditorOnlyDataInterface* ModelEditorOnlyData = MaterialModel->GetEditorOnlyData())
 		{
-			ModelEditorOnlyData->RequestMaterialBuild();
+			ModelEditorOnlyData->RequestMaterialBuild(EDMBuildRequestType::Async);
 		}
 	}
 
diff --git a/Engine/Plugins/Experimental/DynamicMaterial/Source/DynamicMaterialEditor/Private/Model/DynamicMaterialModelEditorOnlyData.cpp b/Engine/Plugins/Experimental/DynamicMaterial/Source/DynamicMaterialEditor/Private/Model/DynamicMaterialModelEditorOnlyData.cpp
--- a/Engine/Plugins/Experimental/DynamicMaterial/Source/DynamicMaterialEditor/Private/Model/DynamicMaterialModelEditorOnlyData.cpp
+++ b/Engine/Plugins/Experimental/DynamicMaterial/Source/DynamicMaterialEditor/Private/Model/DynamicMaterialModelEditorOnlyData.cpp
@@ -177,6 +177,7 @@ UDynamicMaterialModelEditorOnlyData::UDynamicMaterialModelEditorOnlyData()
 	, DisplacementCenter(0.5f)
 	, DisplacementMagnitude(1.f)
 	, bResponsiveAAEnabled(true)
+	, bBuildRequested(false)
 {
 	BaseColor =           CreateDefaultSubobject<UDMMaterialPropertyBaseColor>(          "MaterialProperty_BaseColor");
 	EmissiveColor =       CreateDefaultSubobject<UDMMaterialPropertyEmissiveColor>(      "MaterialProperty_EmissiveColor");
@@ -584,12 +585,16 @@ void UDynamicMaterialModelEditorOnlyData::BuildMaterial(bool bInDirtyAssets)
 	}
 
 	State = EDMState::Idle;
+	bBuildRequested = false;
 
 	OnMaterialBuiltDelegate.Broadcast(MaterialModel);
 }
 
-void UDynamicMaterialModelEditorOnlyData::RequestMaterialBuild(bool bInImmediate)
+void UDynamicMaterialModelEditorOnlyData::RequestMaterialBuild(EDMBuildRequestType InRequestType)
 {
+	UDynamicMaterialEditorSettings* Settings = UDynamicMaterialEditorSettings::Get();
+	UDMBuildRequestSubsystem* BuildRequestSubsystem = UDMBuildRequestSubsystem::Get();
+
 	if (!HasAnyFlags(RF_ClassDefaultObject | RF_ArchetypeObject))
 	{
 		if (IsIn(GetTransientPackage()))
@@ -600,13 +605,37 @@ void UDynamicMaterialModelEditorOnlyData::RequestMaterialBuild(bool bInImmediate
 			}
 		}
 
-		if (bInImmediate)
-		{
-			BuildMaterial(/* Dirty Assets */ false);
-		}
-		else if (UDMBuildRequestSubsystem* BuildRequestSubsystem = UDMBuildRequestSubsystem::Get())
+		switch (InRequestType)
 		{
-			BuildRequestSubsystem->AddBuildRequest(this, /* Dirty Packages */ !UE::GetIsEditorLoadingPackage());
+			case EDMBuildRequestType::Immediate:
+				BuildMaterial(/* Dirty Assets */ false);
+				break;
+
+			case EDMBuildRequestType::Async:
+				if (BuildRequestSubsystem)
+				{
+					BuildRequestSubsystem->AddBuildRequest(this, /* Dirty Packages */ !UE::GetIsEditorLoadingPackage());
+				}
+				else
+				{
+					BuildMaterial(/* Dirty Assets */ false);
+				}
+				break;
+
+			case EDMBuildRequestType::Preview:
+				if (Settings && !Settings->bAutomaticallyCompilePreviewMaterial)
+				{
+					bBuildRequested = true;
+				}
+				else if (BuildRequestSubsystem)
+				{
+					BuildRequestSubsystem->AddBuildRequest(this, /* Dirty Packages */ !UE::GetIsEditorLoadingPackage());
+				}
+				else
+				{
+					BuildMaterial(/* Dirty Assets */ false);
+				}
+				break;
 		}
 	}
 }
@@ -1341,6 +1370,11 @@ void UDynamicMaterialModelEditorOnlyData::UnassignMaterialProperty(EDMMaterialPr
 	RequestMaterialBuild();
 }
 
+bool UDynamicMaterialModelEditorOnlyData::HasBuildBeenRequested() const
+{
+	return bBuildRequested;
+}
+
 void UDynamicMaterialModelEditorOnlyData::NotifyPostChange(const FPropertyChangedEvent& InPropertyChangedEvent, class FEditPropertyChain* InPropertyThatChanged)
 {
 	RequestMaterialBuild();
@@ -1393,7 +1427,12 @@ void UDynamicMaterialModelEditorOnlyData::PostDuplicate(bool bInDuplicateForPIE)
 	{
 		PostEditorDuplicate();
 		ReinitComponents();
-		RequestMaterialBuild();
+
+		RequestMaterialBuild(
+			bInDuplicateForPIE
+			? EDMBuildRequestType::Immediate
+			: EDMBuildRequestType::Async
+		);
 	}
 }
 
diff --git a/Engine/Plugins/Experimental/DynamicMaterial/Source/DynamicMaterialEditor/Private/UI/Widgets/Editor/SDMToolBar.cpp b/Engine/Plugins/Experimental/DynamicMaterial/Source/DynamicMaterialEditor/Private/UI/Widgets/Editor/SDMToolBar.cpp
--- a/Engine/Plugins/Experimental/DynamicMaterial/Source/DynamicMaterialEditor/Private/UI/Widgets/Editor/SDMToolBar.cpp
+++ b/Engine/Plugins/Experimental/DynamicMaterial/Source/DynamicMaterialEditor/Private/UI/Widgets/Editor/SDMToolBar.cpp
@@ -16,6 +16,7 @@
 #include "Model/DynamicMaterialModel.h"
 #include "Model/DynamicMaterialModelBase.h"
 #include "Model/DynamicMaterialModelDynamic.h"
+#include "Model/DynamicMaterialModelEditorOnlyData.h"
 #include "Selection.h"
 #include "Styling/StyleColors.h"
 #include "UI/Menus/DMToolBarMenus.h"
@@ -133,6 +134,7 @@ TSharedRef<SWidget> SDMToolBar::CreateToolBarEntries()
 				[
 					SAssignNew(ApplyButtonWidget, SButton)
 					.ContentPadding(LargeIconToolBarButtonContentPadding)
+					.Visibility(IsDynamicMaterialModel() ? EVisibility::Collapsed : EVisibility::Visible)
 					.IsEnabled(this, &SDMToolBar::IsApplyEnabled)
 					.ButtonStyle(FDynamicMaterialEditorStyle::Get(), "HoverHintOnly")
 					.ToolTipText(LOCTEXT("MaterialDesignerApplyTooltip", "Apply the preview Material Designer Material back to the original asset."))
@@ -142,6 +144,24 @@ TSharedRef<SWidget> SDMToolBar::CreateToolBarEntries()
 						.Image(FAppStyle::GetBrush("MaterialEditor.Apply"))
 						.DesiredSizeOverride(LargeIconToolBarButtonSize)
 					]
+				]
+				+ SHorizontalBox::Slot()
+				.AutoWidth()
+				.Padding(0.f)
+				.VAlign(EVerticalAlignment::VAlign_Center)
+				[
+					SAssignNew(CompileButtonWidget, SButton)
+					.ContentPadding(LargeIconToolBarButtonContentPadding)
+					.Visibility(IsDynamicMaterialModel() ? EVisibility::Collapsed : EVisibility::Visible)
+					.IsEnabled(this, &SDMToolBar::IsCompileEnabled)
+					.ButtonStyle(FDynamicMaterialEditorStyle::Get(), "HoverHintOnly")
+					.ToolTipText(LOCTEXT("MaterialDesignerCompileTooltip", "Compile the preview Material Designer Material, updating the preview."))
+					.OnClicked(this, &SDMToolBar::OnCompileClicked)
+					[
+						SNew(SImage)
+						.Image(FAppStyle::GetBrush("AssetEditor.CompileStatus.Background"))
+						.DesiredSizeOverride(LargeIconToolBarButtonSize)
+					]
 				]	
 				+ SHorizontalBox::Slot()
 				.AutoWidth()
@@ -543,6 +563,19 @@ void SDMToolBar::OnMaterialSlotChanged(TSharedPtr<FDMObjectMaterialProperty> InS
 	}
 }
 
+bool SDMToolBar::IsDynamicMaterialModel() const
+{
+	if (TSharedPtr<SDMMaterialEditor> EditorWidget = GetEditorWidget())
+	{
+		if (UDynamicMaterialModelBase* MaterialModelBase = EditorWidget->GetPreviewMaterialModel())
+		{
+			return MaterialModelBase->IsA<UDynamicMaterialModelDynamic>();
+		}
+	}
+
+	return false;
+}
+
 UDynamicMaterialModelBase* SDMToolBar::GetOriginalMaterialModelBase() const
 {
 	if (TSharedPtr<SDMMaterialEditor> EditorWidget = GetEditorWidget())
@@ -1051,4 +1084,30 @@ bool SDMToolBar::IsApplyEnabled() const
 	return false;
 }
 
+bool SDMToolBar::IsCompileEnabled() const
+{
+	if (TSharedPtr<SDMMaterialEditor> EditorWidget = EditorWidgetWeak.Pin())
+	{
+		if (UDynamicMaterialModel* PreviewMaterialModel = Cast<UDynamicMaterialModel>(EditorWidget->GetPreviewMaterialModelBase()))
+		{
+			if (UDynamicMaterialModelEditorOnlyData* EditorOnlyData = UDynamicMaterialModelEditorOnlyData::Get(PreviewMaterialModel))
+			{
+				return EditorOnlyData->HasBuildBeenRequested();
+			}
+		}
+	}
+
+	return false;
+}
+
+FReply SDMToolBar::OnCompileClicked()
+{
+	if (TSharedPtr<SDMMaterialEditor> EditorWidget = EditorWidgetWeak.Pin())
+	{
+		EditorWidget->Compile();
+	}
+
+	return FReply::Handled();
+}
+
 #undef LOCTEXT_NAMESPACE
diff --git a/Engine/Plugins/Experimental/DynamicMaterial/Source/DynamicMaterialEditor/Private/UI/Widgets/Editor/SDMToolBar.h b/Engine/Plugins/Experimental/DynamicMaterial/Source/DynamicMaterialEditor/Private/UI/Widgets/Editor/SDMToolBar.h
--- a/Engine/Plugins/Experimental/DynamicMaterial/Source/DynamicMaterialEditor/Private/UI/Widgets/Editor/SDMToolBar.h
+++ b/Engine/Plugins/Experimental/DynamicMaterial/Source/DynamicMaterialEditor/Private/UI/Widgets/Editor/SDMToolBar.h
@@ -56,6 +56,7 @@ protected:
 	TSharedPtr<SBox> PropertySelectorContainer;
 	TSharedPtr<SWidget> SaveButtonWidget;
 	TSharedPtr<SWidget> ApplyButtonWidget;
+	TSharedPtr<SWidget> CompileButtonWidget;
 	TSharedPtr<SWidget> ActorRowWidget;
 	TSharedPtr<SWidget> AssetRowWidget;
 	TSharedPtr<STextBlock> ActorNameWidget;
@@ -65,6 +66,8 @@ protected:
 	TSharedPtr<SWidget> ConvertToEditableButton;
 	TSharedPtr<SWidget> ConvertToInstanceButton;
 
+	bool IsDynamicMaterialModel() const;
+
 	UDynamicMaterialModelBase* GetOriginalMaterialModelBase() const;
 
 	TSharedRef<SWidget> CreateToolBarEntries();
@@ -111,4 +114,8 @@ protected:
 	TSharedRef<SWidget> GenerateSettingsMenu();
 
 	bool IsApplyEnabled() const;
+
+	bool IsCompileEnabled() const;
+
+	FReply OnCompileClicked();
 };
diff --git a/Engine/Plugins/Experimental/DynamicMaterial/Source/DynamicMaterialEditor/Private/UI/Widgets/SDMMaterialEditor.cpp b/Engine/Plugins/Experimental/DynamicMaterial/Source/DynamicMaterialEditor/Private/UI/Widgets/SDMMaterialEditor.cpp
--- a/Engine/Plugins/Experimental/DynamicMaterial/Source/DynamicMaterialEditor/Private/UI/Widgets/SDMMaterialEditor.cpp
+++ b/Engine/Plugins/Experimental/DynamicMaterial/Source/DynamicMaterialEditor/Private/UI/Widgets/SDMMaterialEditor.cpp
@@ -239,6 +239,8 @@ void SDMMaterialEditor::SetMaterialModelBase(UDynamicMaterialModelBase* InMateri
 		));
 
 		PreviewMaterialInstance = NewInstance;
+
+		PreviewMaterialModelBase->MarkOriginalUpdated();
 	}
 
 	EditGlobalSettings();
@@ -1010,6 +1012,17 @@ void SDMMaterialEditor::ApplyToOriginal()
 		return;
 	}
 
+	if (UDynamicMaterialModel* PreviewMaterialModel = Cast<UDynamicMaterialModel>(Preview))
+	{
+		if (UDynamicMaterialModelEditorOnlyData* EditorOnlyData = UDynamicMaterialModelEditorOnlyData::Get(PreviewMaterialModel))
+		{
+			if (EditorOnlyData->HasBuildBeenRequested())
+			{
+				EditorOnlyData->RequestMaterialBuild(EDMBuildRequestType::Immediate);
+			}
+		}		
+	}
+
 	UDynamicMaterialInstance* OriginalMaterialInstance = Original->GetDynamicMaterialInstance();
 
 	FScopedTransaction Transaction(LOCTEXT("ApplyMaterialModel", "Apply Material Designer Changes"));
@@ -1027,7 +1040,7 @@ void SDMMaterialEditor::ApplyToOriginal()
 
 	if (UDynamicMaterialModelEditorOnlyData* EditorOnlyData = UDynamicMaterialModelEditorOnlyData::Get(Original))
 	{
-		EditorOnlyData->RequestMaterialBuild(/* Immediate */ true);
+		EditorOnlyData->RequestMaterialBuild(EDMBuildRequestType::Immediate);
 	}
 
 	if (OriginalMaterialInstance)
@@ -1036,24 +1049,23 @@ void SDMMaterialEditor::ApplyToOriginal()
 		Original->ApplyComponents(OriginalMaterialInstance);
 	}
 
-	if (TSharedPtr<SDMMaterialDesigner> DesignerWidget = GetDesignerWidget())
-	{
-		if (ObjectMaterialPropertyOpt.IsSet())
-		{
-			DesignerWidget->OpenObjectMaterialProperty(ObjectMaterialPropertyOpt.GetValue());
-		}
-		else
-		{
-			DesignerWidget->OpenMaterialModelBase(Original);
-		}
-	}
+	MainSlot.Invalidate();
 
-	/*
 	if (PreviewMaterialModelBase)
 	{
 		PreviewMaterialModelBase->MarkOriginalUpdated();
 	}
-	*/
+}
+
+void SDMMaterialEditor::Compile()
+{
+	if (PreviewMaterialModelBase && !PreviewMaterialModelBase->IsA<UDynamicMaterialModelDynamic>())
+	{
+		if (UDynamicMaterialModelEditorOnlyData* EditorOnlyData = UDynamicMaterialModelEditorOnlyData::Get(PreviewMaterialModelBase.Get()))
+		{
+			EditorOnlyData->RequestMaterialBuild(EDMBuildRequestType::Immediate);
+		}
+	}
 }
 
 bool SDMMaterialEditor::SetActivePage(const FDMMaterialEditorPage& InPage)
diff --git a/Engine/Plugins/Experimental/DynamicMaterial/Source/DynamicMaterialEditor/Private/UI/Widgets/SDMMaterialEditor.h b/Engine/Plugins/Experimental/DynamicMaterial/Source/DynamicMaterialEditor/Private/UI/Widgets/SDMMaterialEditor.h
--- a/Engine/Plugins/Experimental/DynamicMaterial/Source/DynamicMaterialEditor/Private/UI/Widgets/SDMMaterialEditor.h
+++ b/Engine/Plugins/Experimental/DynamicMaterial/Source/DynamicMaterialEditor/Private/UI/Widgets/SDMMaterialEditor.h
@@ -151,6 +151,8 @@ public:
 
 	void ApplyToOriginal();
 
+	void Compile();
+
 	void HandleDrop_CreateTextureSet(const TArray<FAssetData>& InTextureAssets);
 
 	void HandleDrop_TextureSet(UDMTextureSet* InTextureSet);
diff --git a/Engine/Plugins/Experimental/DynamicMaterial/Source/DynamicMaterialEditor/Private/UI/Widgets/SDMMaterialWizard.cpp b/Engine/Plugins/Experimental/DynamicMaterial/Source/DynamicMaterialEditor/Private/UI/Widgets/SDMMaterialWizard.cpp
--- a/Engine/Plugins/Experimental/DynamicMaterial/Source/DynamicMaterialEditor/Private/UI/Widgets/SDMMaterialWizard.cpp
+++ b/Engine/Plugins/Experimental/DynamicMaterial/Source/DynamicMaterialEditor/Private/UI/Widgets/SDMMaterialWizard.cpp
@@ -1029,7 +1029,7 @@ void SDMMaterialWizard::CreateTemplateMaterialInActor(FName InChannelList, FDMOb
 	{
 		EditorOnlyData->SetChannelListPreset(InChannelList);
 		EditorOnlyData->OnWizardComplete();
-		EditorOnlyData->RequestMaterialBuild(/* Immediate */ true);
+		EditorOnlyData->RequestMaterialBuild(EDMBuildRequestType::Immediate);
 	}
 
 	if (!UDMMaterialInstanceFunctionLibrary::SetMaterialInObject(InMaterialObjectProperty, NewInstance))
diff --git a/Engine/Plugins/Experimental/DynamicMaterial/Source/DynamicMaterialEditor/Private/Utils/DMMaterialInstanceFunctionLibrary.cpp b/Engine/Plugins/Experimental/DynamicMaterial/Source/DynamicMaterialEditor/Private/Utils/DMMaterialInstanceFunctionLibrary.cpp
--- a/Engine/Plugins/Experimental/DynamicMaterial/Source/DynamicMaterialEditor/Private/Utils/DMMaterialInstanceFunctionLibrary.cpp
+++ b/Engine/Plugins/Experimental/DynamicMaterial/Source/DynamicMaterialEditor/Private/Utils/DMMaterialInstanceFunctionLibrary.cpp
@@ -121,7 +121,7 @@ UDynamicMaterialModel* UDMMaterialInstanceFunctionLibrary::CreateMaterialInObjec
 
 	if (IDynamicMaterialModelEditorOnlyDataInterface* EditorOnlyData = MaterialModel->GetEditorOnlyData())
 	{
-		EditorOnlyData->RequestMaterialBuild();
+		EditorOnlyData->RequestMaterialBuild(EDMBuildRequestType::Async);
 	}
 
 	return MaterialModel;
diff --git a/Engine/Plugins/Experimental/DynamicMaterial/Source/DynamicMaterialEditor/Private/Utils/DMMaterialModelFunctionLibrary.cpp b/Engine/Plugins/Experimental/DynamicMaterial/Source/DynamicMaterialEditor/Private/Utils/DMMaterialModelFunctionLibrary.cpp
--- a/Engine/Plugins/Experimental/DynamicMaterial/Source/DynamicMaterialEditor/Private/Utils/DMMaterialModelFunctionLibrary.cpp
+++ b/Engine/Plugins/Experimental/DynamicMaterial/Source/DynamicMaterialEditor/Private/Utils/DMMaterialModelFunctionLibrary.cpp
@@ -124,7 +124,7 @@ UDynamicMaterialInstance* UDMMaterialModelFunctionLibrary::ExportMaterial(UDynam
 			{
 				if (IDynamicMaterialModelEditorOnlyDataInterface* ModelEditorOnlyData = NewModel->GetEditorOnlyData())
 				{
-					ModelEditorOnlyData->RequestMaterialBuild();
+					ModelEditorOnlyData->RequestMaterialBuild(EDMBuildRequestType::Immediate);
 				}
 			}
 		}
@@ -478,7 +478,7 @@ bool UDMMaterialModelFunctionLibrary::DuplicateModelBetweenMaterials(UDynamicMat
 
 	if (UDynamicMaterialModelEditorOnlyData* EditorOnlyData = UDynamicMaterialModelEditorOnlyData::Get(NewModel))
 	{
-		EditorOnlyData->RequestMaterialBuild();
+		EditorOnlyData->RequestMaterialBuild(EDMBuildRequestType::Async);
 	}
 
 	return true;
@@ -550,7 +550,6 @@ UDynamicMaterialModelBase* UDMMaterialModelFunctionLibrary::CreatePreviewModel(U
 {
 	UDynamicMaterialModelBase* PreviewMaterialModelBase = Cast<UDynamicMaterialModelBase>(StaticDuplicateObject(InOriginalModelBase, GetTransientPackage()));
 	PreviewMaterialModelBase->SetDynamicMaterialInstance(nullptr);
-	PreviewMaterialModelBase->MarkOriginalUpdated();
 
 	if (UDynamicMaterialModelDynamic* MaterialModelDynamic = Cast<UDynamicMaterialModelDynamic>(PreviewMaterialModelBase))
 	{
@@ -559,9 +558,11 @@ UDynamicMaterialModelBase* UDMMaterialModelFunctionLibrary::CreatePreviewModel(U
 
 	if (UDynamicMaterialModelEditorOnlyData* EditorOnlyData = UDynamicMaterialModelEditorOnlyData::Get(PreviewMaterialModelBase))
 	{
-		EditorOnlyData->RequestMaterialBuild(/* Immediate */ true);
+		EditorOnlyData->RequestMaterialBuild(EDMBuildRequestType::Async);
 	}
 
+	PreviewMaterialModelBase->MarkOriginalUpdated();
+
 	return PreviewMaterialModelBase;
 }
 
diff --git a/Engine/Plugins/Experimental/DynamicMaterial/Source/DynamicMaterialEditor/Public/Model/DynamicMaterialModelEditorOnlyData.h b/Engine/Plugins/Experimental/DynamicMaterial/Source/DynamicMaterialEditor/Public/Model/DynamicMaterialModelEditorOnlyData.h
--- a/Engine/Plugins/Experimental/DynamicMaterial/Source/DynamicMaterialEditor/Public/Model/DynamicMaterialModelEditorOnlyData.h
+++ b/Engine/Plugins/Experimental/DynamicMaterial/Source/DynamicMaterialEditor/Public/Model/DynamicMaterialModelEditorOnlyData.h
@@ -212,6 +212,9 @@ public:
 
 	UFUNCTION(BlueprintCallable, Category = "Material Designer")
 	DYNAMICMATERIALEDITOR_API void UnassignMaterialProperty(EDMMaterialPropertyType InProperty);
+
+	UFUNCTION(BlueprintPure, Category = "Material Designer")
+	bool HasBuildBeenRequested() const;
 	
 	FDMOnMaterialBuilt::RegistrationType& GetOnMaterialBuiltDelegate() { return OnMaterialBuiltDelegate; }
 	FDMOnValueListUpdated::RegistrationType& GetOnValueListUpdateDelegate() { return OnValueListUpdateDelegate; }
@@ -246,7 +249,7 @@ public:
 
 	//~ Begin IDynamicMaterialModelEditorOnlyDataInterface
 	DYNAMICMATERIALEDITOR_API virtual void PostEditorDuplicate() override;
-	DYNAMICMATERIALEDITOR_API virtual void RequestMaterialBuild(bool bInImmediate = false) override;
+	DYNAMICMATERIALEDITOR_API virtual void RequestMaterialBuild(EDMBuildRequestType InRequestType = EDMBuildRequestType::Preview) override;
 	DYNAMICMATERIALEDITOR_API virtual void OnValueListUpdate() override;
 	DYNAMICMATERIALEDITOR_API virtual void OnValueUpdated(UDMMaterialValue* InValue, EDMUpdateType InUpdateType) override;
 	DYNAMICMATERIALEDITOR_API virtual void OnTextureUVUpdated(UDMTextureUV* InTextureUV) override;
@@ -396,6 +399,8 @@ protected:
 	UPROPERTY(BlueprintReadOnly, Instanced, Category = "Material Designer")
 	TObjectPtr<UDMMaterialProperty> Custom4;
 
+	bool bBuildRequested;
+
 	void CreateMaterial();
 
 	UFUNCTION(BlueprintCallable, Category = "Material Designer")
