Material Designer: UI Fixes
- Disabling follow selection will now only take effect when there is no material currently opened. Having it disabled with no opened material will open the first valid material on a selected item and keep it open.
- The Compile button has been removed.
- A Live Preview toggle button has been added to flip the value of automatic preview material compile. When enabling the automatic compilation, an immediate preview material compile will be triggered if needed.
- The red warning icon that indicates a property not available in the source material is now a yellow warning icon.

diff --git a/Engine/Plugins/Experimental/DynamicMaterial/Source/DynamicMaterialEditor/Private/UI/Widgets/Editor/SDMObjectEditorWidgetBase.cpp b/Engine/Plugins/Experimental/DynamicMaterial/Source/DynamicMaterialEditor/Private/UI/Widgets/Editor/SDMObjectEditorWidgetBase.cpp
--- a/Engine/Plugins/Experimental/DynamicMaterial/Source/DynamicMaterialEditor/Private/UI/Widgets/Editor/SDMObjectEditorWidgetBase.cpp
+++ b/Engine/Plugins/Experimental/DynamicMaterial/Source/DynamicMaterialEditor/Private/UI/Widgets/Editor/SDMObjectEditorWidgetBase.cpp
@@ -482,7 +482,7 @@ FPropertyRowExtensionButton SDMObjectEditorWidgetBase::CreateNeedsApplyButton()
 {
 	FPropertyRowExtensionButton NeedsApplyButton;
 	NeedsApplyButton.Label = LOCTEXT("NeedsApply", "Needs Apply");
-	NeedsApplyButton.Icon = FSlateIcon(FAppStyle::Get().GetStyleSetName(), TEXT("Icons.ErrorWithColor"));
+	NeedsApplyButton.Icon = FSlateIcon(FAppStyle::Get().GetStyleSetName(), TEXT("Icons.Warning.Solid"));
 	NeedsApplyButton.ToolTip = LOCTEXT("NeedsApplyTooltip", "This property only exists in the preview material. Apply changes to add this to the source material.");
 
 	return NeedsApplyButton;
diff --git a/Engine/Plugins/Experimental/DynamicMaterial/Source/DynamicMaterialEditor/Private/UI/Widgets/Editor/SDMToolBar.cpp b/Engine/Plugins/Experimental/DynamicMaterial/Source/DynamicMaterialEditor/Private/UI/Widgets/Editor/SDMToolBar.cpp
--- a/Engine/Plugins/Experimental/DynamicMaterial/Source/DynamicMaterialEditor/Private/UI/Widgets/Editor/SDMToolBar.cpp
+++ b/Engine/Plugins/Experimental/DynamicMaterial/Source/DynamicMaterialEditor/Private/UI/Widgets/Editor/SDMToolBar.cpp
@@ -17,6 +17,7 @@
 #include "Model/DynamicMaterialModelBase.h"
 #include "Model/DynamicMaterialModelDynamic.h"
 #include "Model/DynamicMaterialModelEditorOnlyData.h"
+#include "ScopedTransaction.h"
 #include "Selection.h"
 #include "Styling/StyleColors.h"
 #include "UI/Menus/DMToolBarMenus.h"
@@ -103,13 +104,13 @@ TSharedRef<SWidget> SDMToolBar::CreateToolBarEntries()
 					.Visibility(EVisibility::Collapsed)
 					.ContentPadding(LargeIconToolBarButtonContentPadding)
 					.ButtonStyle(FDynamicMaterialEditorStyle::Get(), "HoverHintOnly")
-					.ToolTipText(LOCTEXT("MaterialDesignerSaveTooltip", "Save the Material Designer asset\n\nCaution: If this asset lives inside an actor, the actor/level will be saved."))
+					.ToolTipText(LOCTEXT("MaterialDesignerSaveTooltip", "Save the Material Designer Asset\n\nCaution: If this asset lives inside an actor, the actor/level will be saved."))
 					.OnClicked(this, &SDMToolBar::OnSaveClicked)
 					[
 						SNew(SImage)
 						.Image(this, &SDMToolBar::GetSaveIcon)
 						.DesiredSizeOverride(LargeIconToolBarButtonSize)
-					]
+					] 
 				]
 				+ SHorizontalBox::Slot()
 				.AutoWidth()
@@ -137,31 +138,13 @@ TSharedRef<SWidget> SDMToolBar::CreateToolBarEntries()
 					.Visibility(IsDynamicMaterialModel() ? EVisibility::Collapsed : EVisibility::Visible)
 					.IsEnabled(this, &SDMToolBar::IsApplyEnabled)
 					.ButtonStyle(FDynamicMaterialEditorStyle::Get(), "HoverHintOnly")
-					.ToolTipText(LOCTEXT("MaterialDesignerApplyTooltip", "Apply the preview Material Designer Material back to the original asset."))
+					.ToolTipText(LOCTEXT("MaterialDesignerApplyTooltip", "Apply To Source\n\nApply the preview Material Designer Material back to the original asset."))
 					.OnClicked(this, &SDMToolBar::OnApplyClicked)
 					[
 						SNew(SImage)
 						.Image(FAppStyle::GetBrush("MaterialEditor.Apply"))
 						.DesiredSizeOverride(LargeIconToolBarButtonSize)
 					]
-				]
-				+ SHorizontalBox::Slot()
-				.AutoWidth()
-				.Padding(0.f)
-				.VAlign(EVerticalAlignment::VAlign_Center)
-				[
-					SAssignNew(CompileButtonWidget, SButton)
-					.ContentPadding(LargeIconToolBarButtonContentPadding)
-					.Visibility(IsDynamicMaterialModel() ? EVisibility::Collapsed : EVisibility::Visible)
-					.IsEnabled(this, &SDMToolBar::IsCompileEnabled)
-					.ButtonStyle(FDynamicMaterialEditorStyle::Get(), "HoverHintOnly")
-					.ToolTipText(LOCTEXT("MaterialDesignerCompileTooltip", "Compile the preview Material Designer Material, updating the preview."))
-					.OnClicked(this, &SDMToolBar::OnCompileClicked)
-					[
-						SNew(SImage)
-						.Image(FAppStyle::GetBrush("AssetEditor.CompileStatus.Background"))
-						.DesiredSizeOverride(LargeIconToolBarButtonSize)
-					]
 				]	
 				+ SHorizontalBox::Slot()
 				.AutoWidth()
@@ -311,6 +294,26 @@ TSharedRef<SWidget> SDMToolBar::CreateToolBarEntries()
 				]
 			]
 		]
+
+		+ SHorizontalBox::Slot().AutoWidth()
+		.AutoWidth()
+		.HAlign(HAlign_Right)
+		.VAlign(VAlign_Top)
+		.Padding(5.0f, 0.0f, 0.0f, 0.0f)
+		[
+			SAssignNew(AutoCompileButtonWidget, SButton)
+			.ContentPadding(LargeIconToolBarButtonContentPadding)
+			.Visibility(IsDynamicMaterialModel() ? EVisibility::Collapsed : EVisibility::Visible)
+			.ButtonStyle(FDynamicMaterialEditorStyle::Get(), "HoverHintOnly")
+			.ToolTipText(LOCTEXT("MaterialDesignerCompileTooltip", "Live Preview\n\nWhen enabled the preview material will be compiled for every structural change.\n\nRe-enabling live preview will cause any uncompiled changes to compile immediately."))
+			.OnClicked(this, &SDMToolBar::OnAutoCompileClicked)
+			[
+				SNew(SImage)
+				.Image(FAppStyle::GetBrush("MaterialEditor.LiveUpdate"))
+				.ColorAndOpacity(this, &SDMToolBar::GetAutoCompileColor)
+				.DesiredSizeOverride(LargeIconToolBarButtonSize)
+			]
+		]
 		
 		+ SHorizontalBox::Slot()
 		.AutoWidth()
@@ -321,7 +324,7 @@ TSharedRef<SWidget> SDMToolBar::CreateToolBarEntries()
 			SNew(SButton)
 			.ContentPadding(DefaultToolBarButtonContentPadding)
 			.ButtonStyle(FDynamicMaterialEditorStyle::Get(), "HoverHintOnly")
-			.ToolTipText(LOCTEXT("MaterialDesignerFollowSelectionTooltip", "Toggles whether the Material Designer display will change when selecting new objects and actors."))
+			.ToolTipText(LOCTEXT("MaterialDesignerFollowSelectionTooltip", "Follow Selection\n\nWhen enabled the Material Designer will open newly selected objects, assets and actors."))
 			.OnClicked(this, &SDMToolBar::OnFollowSelectionButtonClicked)
 			[
 				SNew(SImage)
@@ -603,7 +606,7 @@ FSlateColor SDMToolBar::GetFollowSelectionColor() const
 {
 	// We want the icon to stand out when it's locked.
 	static FSlateColor EnabledColor = FSlateColor(EStyleColor::AccentGray);
-	static FSlateColor DisabledColor = FSlateColor(EStyleColor::Primary);
+	static FSlateColor DisabledColor = FSlateColor(EStyleColor::AccentBlue);
 
 	if (SDMMaterialDesigner::IsFollowingSelection())
 	{
@@ -1084,27 +1087,47 @@ bool SDMToolBar::IsApplyEnabled() const
 	return false;
 }
 
-bool SDMToolBar::IsCompileEnabled() const
+FSlateColor SDMToolBar::GetAutoCompileColor() const
 {
-	if (TSharedPtr<SDMMaterialEditor> EditorWidget = EditorWidgetWeak.Pin())
+	static FSlateColor EnabledColor = FSlateColor(EStyleColor::AccentBlue);
+	static FSlateColor DisabledColor = FSlateColor(EStyleColor::AccentGray);
+
+	if (UDynamicMaterialEditorSettings* Settings = UDynamicMaterialEditorSettings::Get())
 	{
-		if (UDynamicMaterialModel* PreviewMaterialModel = Cast<UDynamicMaterialModel>(EditorWidget->GetPreviewMaterialModelBase()))
+		if (Settings->bAutomaticallyCompilePreviewMaterial)
 		{
-			if (UDynamicMaterialModelEditorOnlyData* EditorOnlyData = UDynamicMaterialModelEditorOnlyData::Get(PreviewMaterialModel))
-			{
-				return EditorOnlyData->HasBuildBeenRequested();
-			}
+			return EnabledColor;
 		}
 	}
 
-	return false;
+	return DisabledColor;
 }
 
-FReply SDMToolBar::OnCompileClicked()
+FReply SDMToolBar::OnAutoCompileClicked()
 {
-	if (TSharedPtr<SDMMaterialEditor> EditorWidget = EditorWidgetWeak.Pin())
+	if (UDynamicMaterialEditorSettings* Settings = UDynamicMaterialEditorSettings::Get())
 	{
-		EditorWidget->Compile();
+		FScopedTransaction Transaction(LOCTEXT("ToggleAutoPreviewCompile", "Toggle Live Preview"));
+		Settings->Modify();
+		Settings->bAutomaticallyCompilePreviewMaterial = !Settings->bAutomaticallyCompilePreviewMaterial;
+		Settings->SaveConfig();
+
+		if (Settings->bAutomaticallyCompilePreviewMaterial)
+		{
+			if (TSharedPtr<SDMMaterialEditor> EditorWidget = EditorWidgetWeak.Pin())
+			{
+				if (UDynamicMaterialModel* PreviewMaterialModel = Cast<UDynamicMaterialModel>(EditorWidget->GetPreviewMaterialModelBase()))
+				{
+					if (UDynamicMaterialModelEditorOnlyData* EditorOnlyData = UDynamicMaterialModelEditorOnlyData::Get(PreviewMaterialModel))
+					{
+						if (EditorOnlyData->HasBuildBeenRequested())
+						{
+							EditorOnlyData->RequestMaterialBuild(EDMBuildRequestType::Async);
+						}
+					}
+				}
+			}
+		}
 	}
 
 	return FReply::Handled();
diff --git a/Engine/Plugins/Experimental/DynamicMaterial/Source/DynamicMaterialEditor/Private/UI/Widgets/Editor/SDMToolBar.h b/Engine/Plugins/Experimental/DynamicMaterial/Source/DynamicMaterialEditor/Private/UI/Widgets/Editor/SDMToolBar.h
--- a/Engine/Plugins/Experimental/DynamicMaterial/Source/DynamicMaterialEditor/Private/UI/Widgets/Editor/SDMToolBar.h
+++ b/Engine/Plugins/Experimental/DynamicMaterial/Source/DynamicMaterialEditor/Private/UI/Widgets/Editor/SDMToolBar.h
@@ -56,7 +56,7 @@ protected:
 	TSharedPtr<SBox> PropertySelectorContainer;
 	TSharedPtr<SWidget> SaveButtonWidget;
 	TSharedPtr<SWidget> ApplyButtonWidget;
-	TSharedPtr<SWidget> CompileButtonWidget;
+	TSharedPtr<SWidget> AutoCompileButtonWidget;
 	TSharedPtr<SWidget> ActorRowWidget;
 	TSharedPtr<SWidget> AssetRowWidget;
 	TSharedPtr<STextBlock> ActorNameWidget;
@@ -115,7 +115,7 @@ protected:
 
 	bool IsApplyEnabled() const;
 
-	bool IsCompileEnabled() const;
+	FSlateColor GetAutoCompileColor() const;
 
-	FReply OnCompileClicked();
+	FReply OnAutoCompileClicked();
 };
diff --git a/Engine/Plugins/Experimental/DynamicMaterial/Source/DynamicMaterialEditor/Private/UI/Widgets/SDMMaterialDesigner.cpp b/Engine/Plugins/Experimental/DynamicMaterial/Source/DynamicMaterialEditor/Private/UI/Widgets/SDMMaterialDesigner.cpp
--- a/Engine/Plugins/Experimental/DynamicMaterial/Source/DynamicMaterialEditor/Private/UI/Widgets/SDMMaterialDesigner.cpp
+++ b/Engine/Plugins/Experimental/DynamicMaterial/Source/DynamicMaterialEditor/Private/UI/Widgets/SDMMaterialDesigner.cpp
@@ -116,7 +116,7 @@ void SDMMaterialDesigner::Empty()
 
 void SDMMaterialDesigner::OnMaterialModelBaseSelected(UDynamicMaterialModelBase* InMaterialModelBase)
 {
-	if (IsFollowingSelection())
+	if (ShouldFollowSelection())
 	{
 		OpenMaterialModelBase(InMaterialModelBase);
 	}
@@ -124,7 +124,7 @@ void SDMMaterialDesigner::OnMaterialModelBaseSelected(UDynamicMaterialModelBase*
 
 void SDMMaterialDesigner::OnMaterialInstanceSelected(UDynamicMaterialInstance* InMaterialInstance)
 {
-	if (IsFollowingSelection())
+	if (ShouldFollowSelection())
 	{
 		OpenMaterialModelBase(InMaterialInstance->GetMaterialModel());
 	}
@@ -132,7 +132,7 @@ void SDMMaterialDesigner::OnMaterialInstanceSelected(UDynamicMaterialInstance* I
 
 void SDMMaterialDesigner::OnObjectMaterialPropertySelected(const FDMObjectMaterialProperty& InObjectMaterialProperty)
 {
-	if (IsFollowingSelection())
+	if (ShouldFollowSelection())
 	{
 		OpenObjectMaterialProperty(InObjectMaterialProperty);
 	}
@@ -140,7 +140,7 @@ void SDMMaterialDesigner::OnObjectMaterialPropertySelected(const FDMObjectMateri
 
 void SDMMaterialDesigner::OnActorSelected(AActor* InActor)
 {
-	if (IsFollowingSelection())
+	if (ShouldFollowSelection())
 	{
 		OpenActor(InActor);
 	}
@@ -350,6 +350,11 @@ bool SDMMaterialDesigner::NeedsWizard(UDynamicMaterialModelBase* InMaterialModel
 	return false;
 }
 
+bool SDMMaterialDesigner::ShouldFollowSelection() const
+{
+	return IsFollowingSelection() || !GetOriginalMaterialModelBase();
+}
+
 bool SDMMaterialDesigner::OnAssetDraggedOver(TArrayView<FAssetData> InAssets)
 {
 	TArray<UClass*> AllowedClasses = {
