Horde Issue 820985

diff --git a/Engine/Plugins/Experimental/DynamicMaterial/Source/DynamicMaterialEditor/Private/UI/Widgets/SDMMaterialDesigner.h b/Engine/Plugins/Experimental/DynamicMaterial/Source/DynamicMaterialEditor/Private/UI/Widgets/SDMMaterialDesigner.h
--- a/Engine/Plugins/Experimental/DynamicMaterial/Source/DynamicMaterialEditor/Private/UI/Widgets/SDMMaterialDesigner.h
+++ b/Engine/Plugins/Experimental/DynamicMaterial/Source/DynamicMaterialEditor/Private/UI/Widgets/SDMMaterialDesigner.h
@@ -94,6 +94,8 @@ protected:
 
 	bool NeedsWizard(UDynamicMaterialModelBase* InMaterialModelBase) const;
 
+	bool ShouldFollowSelection() const;
+
 	/** Drag and drop. */
 	bool OnAssetDraggedOver(TArrayView<FAssetData> InAssets);
 
