Material Designer: UI Update
- Removed "Save As" button.
- Renamed menu export feature to "Save As new Material Designer asset" (already existed, just renamed.)
- Fixed toolbar becoming slightly lower when really thin (Moved dynamic instance buttons to the right of actor details.)
- Fixed actor icon.
- Fixed events in material designer to respond to preview material compilation/state changes.
- Added auto-apply-to-source on preview compile button.
- Added auto-apply-to-source option on UI.
- Removed apply button.
- Added material effect asset getter.
- Added browse to asset in content browser to effect rows if a related asset exists.
- Changed "Needs Apply" button to the apply icon, not a warning or error icon.

diff --git a/Engine/Plugins/Experimental/DynamicMaterial/Source/DynamicMaterialEditor/Private/Components/DMMaterialEffectFunction.cpp b/Engine/Plugins/Experimental/DynamicMaterial/Source/DynamicMaterialEditor/Private/Components/DMMaterialEffectFunction.cpp
--- a/Engine/Plugins/Experimental/DynamicMaterial/Source/DynamicMaterialEditor/Private/Components/DMMaterialEffectFunction.cpp
+++ b/Engine/Plugins/Experimental/DynamicMaterial/Source/DynamicMaterialEditor/Private/Components/DMMaterialEffectFunction.cpp
@@ -337,6 +337,11 @@ void UDMMaterialEffectFunction::ApplyTo(const TSharedRef<FDMMaterialBuildState>&
 	InOutLastExpressionOutputChannel = FDMMaterialStageConnectorChannel::WHOLE_CHANNEL;
 }
 
+UObject* UDMMaterialEffectFunction::GetAsset() const
+{
+	return GetMaterialFunction();
+}
+
 FText UDMMaterialEffectFunction::GetComponentDescription() const
 {
 	return GetEffectName();
diff --git a/Engine/Plugins/Experimental/DynamicMaterial/Source/DynamicMaterialEditor/Private/DynamicMaterialEditorSettings.cpp b/Engine/Plugins/Experimental/DynamicMaterial/Source/DynamicMaterialEditor/Private/DynamicMaterialEditorSettings.cpp
--- a/Engine/Plugins/Experimental/DynamicMaterial/Source/DynamicMaterialEditor/Private/DynamicMaterialEditorSettings.cpp
+++ b/Engine/Plugins/Experimental/DynamicMaterial/Source/DynamicMaterialEditor/Private/DynamicMaterialEditorSettings.cpp
@@ -136,6 +136,7 @@ UDynamicMaterialEditorSettings::UDynamicMaterialEditorSettings()
 	bAddDetailsPanelButton = false;
 	bAutomaticallyCopyParametersToSourceMaterial = true;
 	bAutomaticallyCompilePreviewMaterial = false;
+	bAutomaticallyApplyToSourceOnPreviewCompile = true;
 	bUseLinearColorForVectors = true;
 
 	ResetAllLayoutSettings();
diff --git a/Engine/Plugins/Experimental/DynamicMaterial/Source/DynamicMaterialEditor/Private/DynamicMaterialEditorSettings.h b/Engine/Plugins/Experimental/DynamicMaterial/Source/DynamicMaterialEditor/Private/DynamicMaterialEditorSettings.h
--- a/Engine/Plugins/Experimental/DynamicMaterial/Source/DynamicMaterialEditor/Private/DynamicMaterialEditorSettings.h
+++ b/Engine/Plugins/Experimental/DynamicMaterial/Source/DynamicMaterialEditor/Private/DynamicMaterialEditorSettings.h
@@ -194,7 +194,7 @@ public:
 
 	/**
 	 * When properties on the preview material are changed, if they are parameter-based then automatically
-	 * also set those values on the original material (in the world or the asset).
+	 * also set those values on the source material (in the world or asset).
 	 */
 	UPROPERTY(Config, EditAnywhere, BlueprintReadWrite, Category = "Editor")
 	bool bAutomaticallyCopyParametersToSourceMaterial;
@@ -205,6 +205,12 @@ public:
 	UPROPERTY(Config, EditAnywhere, BlueprintReadWrite, Category = "Editor")
 	bool bAutomaticallyCompilePreviewMaterial;
 
+	/**
+	 * When the preview material is recompiled, whether to automatically also apply those changes to the source material.
+	 */
+	UPROPERTY(Config, EditAnywhere, BlueprintReadWrite, Category = "Editor")
+	bool bAutomaticallyApplyToSourceOnPreviewCompile;
+
 	UPROPERTY(Config, EditAnywhere, BlueprintReadWrite, Category = "Layout")
 	EDMMaterialEditorLayout Layout;
 
diff --git a/Engine/Plugins/Experimental/DynamicMaterial/Source/DynamicMaterialEditor/Private/UI/Menus/DMToolBarMenus.cpp b/Engine/Plugins/Experimental/DynamicMaterial/Source/DynamicMaterialEditor/Private/UI/Menus/DMToolBarMenus.cpp
--- a/Engine/Plugins/Experimental/DynamicMaterial/Source/DynamicMaterialEditor/Private/UI/Menus/DMToolBarMenus.cpp
+++ b/Engine/Plugins/Experimental/DynamicMaterial/Source/DynamicMaterialEditor/Private/UI/Menus/DMToolBarMenus.cpp
@@ -102,14 +102,8 @@ void FDMToolBarMenus::AddExportMenu(UToolMenu* InMenu)
 		return;
 	}
 
-	const bool bAllowMaterialExport = IsValid(Material->GetOuter()) && !Material->GetOuter()->IsA<UPackage>();
 	const bool bAllowGeneratedMaterialExport = IsValid(MaterialModelBase->GetGeneratedMaterial());
 
-	if (!bAllowMaterialExport && !bAllowGeneratedMaterialExport)
-	{
-		return;
-	}
-
 	FToolMenuSection& NewSection = InMenu->AddSection(ToolBarMaterialExportSectionName, LOCTEXT("ExportSection", "Export"));
 
 	NewSection.AddMenuEntry(
@@ -137,19 +131,16 @@ void FDMToolBarMenus::AddExportMenu(UToolMenu* InMenu)
 		);
 	}
 
-	if (bAllowMaterialExport)
-	{
-		NewSection.AddMenuEntry(
-			NAME_None,
-			LOCTEXT("ExportMaterial", "Export Material Designer Material"),
-			LOCTEXT("ExportMaterialInstanceTooltip", "Export the Material Designer Material as an asset."),
-			FSlateIcon(),
-			FUIAction(FExecuteAction::CreateStatic(
-				&FDMToolBarMenus::ExportMaterial,
-				TWeakObjectPtr<UDynamicMaterialInstance>(Material)
-			))
-		);
-	}
+	NewSection.AddMenuEntry(
+		NAME_None,
+		LOCTEXT("ExportMaterial", "Save As new Material Designer Asset"),
+		LOCTEXT("ExportMaterialInstanceTooltip", "Export the Material Designer Material as a new asset."),
+		FSlateIcon(),
+		FUIAction(FExecuteAction::CreateStatic(
+			&FDMToolBarMenus::ExportMaterial,
+			TWeakObjectPtr<UDynamicMaterialInstance>(Material)
+		))
+	);
 
 	NewSection.AddSubMenu(
 		NAME_None,
diff --git a/Engine/Plugins/Experimental/DynamicMaterial/Source/DynamicMaterialEditor/Private/UI/Widgets/Editor/SDMObjectEditorWidgetBase.cpp b/Engine/Plugins/Experimental/DynamicMaterial/Source/DynamicMaterialEditor/Private/UI/Widgets/Editor/SDMObjectEditorWidgetBase.cpp
--- a/Engine/Plugins/Experimental/DynamicMaterial/Source/DynamicMaterialEditor/Private/UI/Widgets/Editor/SDMObjectEditorWidgetBase.cpp
+++ b/Engine/Plugins/Experimental/DynamicMaterial/Source/DynamicMaterialEditor/Private/UI/Widgets/Editor/SDMObjectEditorWidgetBase.cpp
@@ -482,7 +482,7 @@ FPropertyRowExtensionButton SDMObjectEditorWidgetBase::CreateNeedsApplyButton()
 {
 	FPropertyRowExtensionButton NeedsApplyButton;
 	NeedsApplyButton.Label = LOCTEXT("NeedsApply", "Needs Apply");
-	NeedsApplyButton.Icon = FSlateIcon(FAppStyle::Get().GetStyleSetName(), TEXT("Icons.Warning.Solid"));
+	NeedsApplyButton.Icon = FSlateIcon(FAppStyle::Get().GetStyleSetName(), TEXT("MaterialEditor.Apply"));
 	NeedsApplyButton.ToolTip = LOCTEXT("NeedsApplyTooltip", "This property only exists in the preview material. Apply changes to add this to the source material.");
 
 	return NeedsApplyButton;
diff --git a/Engine/Plugins/Experimental/DynamicMaterial/Source/DynamicMaterialEditor/Private/UI/Widgets/Editor/SDMToolBar.cpp b/Engine/Plugins/Experimental/DynamicMaterial/Source/DynamicMaterialEditor/Private/UI/Widgets/Editor/SDMToolBar.cpp
--- a/Engine/Plugins/Experimental/DynamicMaterial/Source/DynamicMaterialEditor/Private/UI/Widgets/Editor/SDMToolBar.cpp
+++ b/Engine/Plugins/Experimental/DynamicMaterial/Source/DynamicMaterialEditor/Private/UI/Widgets/Editor/SDMToolBar.cpp
@@ -19,6 +19,7 @@
 #include "Model/DynamicMaterialModelEditorOnlyData.h"
 #include "ScopedTransaction.h"
 #include "Selection.h"
+#include "Styling/SlateIconFinder.h"
 #include "Styling/StyleColors.h"
 #include "UI/Menus/DMToolBarMenus.h"
 #include "UI/Widgets/SDMMaterialDesigner.h"
@@ -76,6 +77,18 @@
 {
 	using namespace UE::DynamicMaterialEditor::Private;
 
+	const FSlateBrush* ActorBrush = nullptr;
+
+	if (AActor* Actor = MaterialActorWeak.Get())
+	{
+		ActorBrush = FSlateIconFinder::FindIconBrushForClass(Actor->GetClass());
+	}
+
+	if (!ActorBrush)
+	{
+		ActorBrush = FSlateIconFinder::FindIconBrushForClass(AActor::StaticClass());
+	}
+
 	return SNew(SHorizontalBox)
 		
 		+ SHorizontalBox::Slot()
@@ -92,113 +105,3 @@
 			+ SWrapBox::Slot()
 			.HAlign(HAlign_Left)
 			.VAlign(VAlign_Center)
-			.Padding(5.0f, 0.0f, 0.0f, 0.0f)
-			[
-				SNew(SHorizontalBox)
-				+ SHorizontalBox::Slot()
-				.AutoWidth()
-				.Padding(0.f)
-				.VAlign(EVerticalAlignment::VAlign_Center)
-				[
-					SAssignNew(SaveButtonWidget, SButton)
-					.Visibility(EVisibility::Collapsed)
-					.ContentPadding(LargeIconToolBarButtonContentPadding)
-					.ButtonStyle(FDynamicMaterialEditorStyle::Get(), "HoverHintOnly")
-					.ToolTipText(LOCTEXT("MaterialDesignerSaveTooltip", "Save the Material Designer Asset\n\nCaution: If this asset lives inside an actor, the actor/level will be saved."))
-					.OnClicked(this, &SDMToolBar::OnSaveClicked)
-					[
-						SNew(SImage)
-						.Image(this, &SDMToolBar::GetSaveIcon)
-						.DesiredSizeOverride(LargeIconToolBarButtonSize)
-					] 
-				]
-				+ SHorizontalBox::Slot()
-				.AutoWidth()
-				.Padding(0.f)
-				.VAlign(EVerticalAlignment::VAlign_Center)
-				[
-					SNew(SButton)
-					.ContentPadding(LargeIconToolBarButtonContentPadding)
-					.ButtonStyle(FDynamicMaterialEditorStyle::Get(), "HoverHintOnly")
-					.ToolTipText(LOCTEXT("ExportMaterial", "Save As"))
-					.OnClicked(this, &SDMToolBar::OnExportMaterialInstanceButtonClicked)
-					[
-						SNew(SImage)
-						.Image(FAppStyle::GetBrush("AssetEditor.SaveAssetAs"))
-						.DesiredSizeOverride(LargeIconToolBarButtonSize)
-					]
-				]
-				+ SHorizontalBox::Slot()
-				.AutoWidth()
-				.Padding(0.f)
-				.VAlign(EVerticalAlignment::VAlign_Center)
-				[
-					SAssignNew(ApplyButtonWidget, SButton)
-					.ContentPadding(LargeIconToolBarButtonContentPadding)
-					.Visibility(IsDynamicMaterialModel() ? EVisibility::Collapsed : EVisibility::Visible)
-					.IsEnabled(this, &SDMToolBar::IsApplyEnabled)
-					.ButtonStyle(FDynamicMaterialEditorStyle::Get(), "HoverHintOnly")
-					.ToolTipText(LOCTEXT("MaterialDesignerApplyTooltip", "Apply To Source\n\nApply the preview Material Designer Material back to the original asset."))
-					.OnClicked(this, &SDMToolBar::OnApplyClicked)
-					[
-						SNew(SImage)
-						.Image(FAppStyle::GetBrush("MaterialEditor.Apply"))
-						.DesiredSizeOverride(LargeIconToolBarButtonSize)
-					]
-				]	
-				+ SHorizontalBox::Slot()
-				.AutoWidth()
-				.VAlign(VAlign_Center)
-				.Padding(0.f)
-				[
-					SAssignNew(OpenParentButton, SButton)
-					.Visibility(EVisibility::Collapsed)
-					.ContentPadding(LargeIconToolBarButtonContentPadding)
-					.ButtonStyle(FDynamicMaterialEditorStyle::Get(), "HoverHintOnly")
-					.ToolTipText(LOCTEXT("MaterialDesignerOpenParentTooltip", "Open the parent of this Material Designer Instance."))
-					.OnClicked(this, &SDMToolBar::OnOpenParentClicked)
-					[
-						SNew(SImage)
-						.Image(FAppStyle::GetBrush(TEXT("Icons.Blueprints")))
-						.DesiredSizeOverride(LargeIconToolBarButtonSize)
-					]
-				]
-				+ SHorizontalBox::Slot()
-				.AutoWidth()
-				.VAlign(VAlign_Center)
-				.Padding(0.f)
-				[
-					SAssignNew(ConvertToEditableButton, SButton)
-					.Visibility(EVisibility::Collapsed)
-					.ContentPadding(LargeIconToolBarButtonContentPadding)
-					.ButtonStyle(FDynamicMaterialEditorStyle::Get(), "HoverHintOnly")
-					.ToolTipText(LOCTEXT("MaterialDesignerConvertToEditableTooltip", "Convert this Material Designer Instance to a fully editable Material (and create a new shader)."))
-					.OnClicked(this, &SDMToolBar::OnConvertToEditableClicked)
-					[
-						SNew(SImage)
-						.Image(FAppStyle::GetBrush(TEXT("Icons.Edit")))
-						.DesiredSizeOverride(LargeIconToolBarButtonSize)
-					]
-				]
-				+ SHorizontalBox::Slot()
-				.AutoWidth()
-				.VAlign(VAlign_Center)
-				.Padding(0.f)
-				[
-					SAssignNew(ConvertToInstanceButton, SButton)
-					.Visibility(EVisibility::Collapsed)
-					.ContentPadding(LargeIconToolBarButtonContentPadding)
-					.ButtonStyle(FDynamicMaterialEditorStyle::Get(), "HoverHintOnly")
-					.ToolTipText(LOCTEXT("MaterialDesignerConvertToInstanceTooltip", "Convert this Material Designer Material to an Instance, creating a local Instance inside the Actor."))
-					.OnClicked(this, &SDMToolBar::OnConvertToInstanceClicked)
-					[
-						SNew(SImage)
-						.Image(FAppStyle::GetBrush(TEXT("Icons.Toolbar.Export")))
-						.DesiredSizeOverride(LargeIconToolBarButtonSize)
-					]
-				]
-			]
-
-			+ SWrapBox::Slot()
-			.HAlign(HAlign_Left)
-			.VAlign(VAlign_Center)
@@ -255,7 +158,7 @@
 				.VAlign(EVerticalAlignment::VAlign_Center)
 				[
 					SNew(SImage)
-					.Image(FAppStyle::GetBrush("ClassIcon.Actor"))
+					.Image(ActorBrush)
 					.DesiredSizeOverride(LargeIconToolBarButtonSize)
 				]
 				+ SHorizontalBox::Slot()
@@ -293,19 +196,115 @@
 					]
 				]
 			]
+
+			+ SWrapBox::Slot()
+			.HAlign(HAlign_Left)
+			.VAlign(VAlign_Center)
+			.Padding(5.0f, 0.0f, 0.0f, 0.0f)
+			[
+				SNew(SHorizontalBox)
+				+ SHorizontalBox::Slot()
+				.AutoWidth()
+				.Padding(0.f)
+				.VAlign(EVerticalAlignment::VAlign_Center)
+				[
+					SAssignNew(SaveButtonWidget, SButton)
+					.Visibility(EVisibility::Collapsed)
+					.ContentPadding(LargeIconToolBarButtonContentPadding)
+					.ButtonStyle(FDynamicMaterialEditorStyle::Get(), "HoverHintOnly")
+					.ToolTipText(LOCTEXT("MaterialDesignerSaveTooltip", "Save the Material Designer Asset\n\nCaution: If this asset lives inside an actor, the actor/level will be saved."))
+					.OnClicked(this, &SDMToolBar::OnSaveClicked)
+					[
+						SNew(SImage)
+						.Image(this, &SDMToolBar::GetSaveIcon)
+						.DesiredSizeOverride(LargeIconToolBarButtonSize)
+					] 
+				]
+				+ SHorizontalBox::Slot()
+				.AutoWidth()
+				.VAlign(VAlign_Center)
+				.Padding(0.f)
+				[
+					SAssignNew(OpenParentButton, SButton)
+					.Visibility(EVisibility::Collapsed)
+					.ContentPadding(LargeIconToolBarButtonContentPadding)
+					.ButtonStyle(FDynamicMaterialEditorStyle::Get(), "HoverHintOnly")
+					.ToolTipText(LOCTEXT("MaterialDesignerOpenParentTooltip", "Open the parent of this Material Designer Instance."))
+					.OnClicked(this, &SDMToolBar::OnOpenParentClicked)
+					[
+						SNew(SImage)
+						.Image(FAppStyle::GetBrush(TEXT("Icons.Blueprints")))
+						.DesiredSizeOverride(LargeIconToolBarButtonSize)
+					]
+				]
+				+ SHorizontalBox::Slot()
+				.AutoWidth()
+				.VAlign(VAlign_Center)
+				.Padding(0.f)
+				[
+					SAssignNew(ConvertToEditableButton, SButton)
+					.Visibility(EVisibility::Collapsed)
+					.ContentPadding(LargeIconToolBarButtonContentPadding)
+					.ButtonStyle(FDynamicMaterialEditorStyle::Get(), "HoverHintOnly")
+					.ToolTipText(LOCTEXT("MaterialDesignerConvertToEditableTooltip", "Convert this Material Designer Instance to a fully editable Material (and create a new shader)."))
+					.OnClicked(this, &SDMToolBar::OnConvertToEditableClicked)
+					[
+						SNew(SImage)
+						.Image(FAppStyle::GetBrush(TEXT("Icons.Edit")))
+						.DesiredSizeOverride(LargeIconToolBarButtonSize)
+					]
+				]
+				+ SHorizontalBox::Slot()
+				.AutoWidth()
+				.VAlign(VAlign_Center)
+				.Padding(0.f)
+				[
+					SAssignNew(ConvertToInstanceButton, SButton)
+					.Visibility(EVisibility::Collapsed)
+					.ContentPadding(LargeIconToolBarButtonContentPadding)
+					.ButtonStyle(FDynamicMaterialEditorStyle::Get(), "HoverHintOnly")
+					.ToolTipText(LOCTEXT("MaterialDesignerConvertToInstanceTooltip", "Convert this Material Designer Material to an Instance, creating a local Instance inside the Actor."))
+					.OnClicked(this, &SDMToolBar::OnConvertToInstanceClicked)
+					[
+						SNew(SImage)
+						.Image(FAppStyle::GetBrush(TEXT("Icons.Toolbar.Export")))
+						.DesiredSizeOverride(LargeIconToolBarButtonSize)
+					]
+				]
+			]
 		]
 
-		+ SHorizontalBox::Slot().AutoWidth()
+		+ SHorizontalBox::Slot()
 		.AutoWidth()
 		.HAlign(HAlign_Right)
 		.VAlign(VAlign_Top)
 		.Padding(5.0f, 0.0f, 0.0f, 0.0f)
 		[
-			SAssignNew(AutoCompileButtonWidget, SButton)
+			SNew(SButton)
 			.ContentPadding(LargeIconToolBarButtonContentPadding)
 			.Visibility(IsDynamicMaterialModel() ? EVisibility::Collapsed : EVisibility::Visible)
 			.ButtonStyle(FDynamicMaterialEditorStyle::Get(), "HoverHintOnly")
-			.ToolTipText(LOCTEXT("MaterialDesignerCompileTooltip", "Live Preview\n\nWhen enabled the preview material will be compiled for every structural change.\n\nRe-enabling live preview will cause any uncompiled changes to compile immediately."))
+			.ToolTipText(LOCTEXT("MaterialDesignerAutoApplyTooltip", "Auto Apply to Source\n\nWhen enabled, and when the preview material is recompiled, it will automatically apply the changes to the source asset.\n\nToggling on Auto Apply to Source will cause any compiled changes to apply to the source immediately."))
+			.OnClicked(this, &SDMToolBar::OnAutoApplyClicked)
+			[
+				SNew(SImage)
+				.Image(FAppStyle::GetBrush("MaterialEditor.Apply"))
+				.ColorAndOpacity(this, &SDMToolBar::GetAutoApplyColor)
+				.DesiredSizeOverride(LargeIconToolBarButtonSize)
+			]
+		]
+
+		+ SHorizontalBox::Slot()
+		.AutoWidth()
+		.HAlign(HAlign_Right)
+		.VAlign(VAlign_Top)
+		.Padding(5.0f, 0.0f, 0.0f, 0.0f)
+		[
+			SNew(SButton)
+			.ContentPadding(LargeIconToolBarButtonContentPadding)
+			.Visibility(IsDynamicMaterialModel() ? EVisibility::Collapsed : EVisibility::Visible)
+			.ButtonStyle(FDynamicMaterialEditorStyle::Get(), "HoverHintOnly")
+			.ToolTipText(LOCTEXT("MaterialDesignerAutoCompileTooltip", "Live Preview\n\nWhen enabled the preview material will be compiled for every structural change.\n\nToggling on Live Preview will cause any uncompiled changes to compile immediately."))
 			.OnClicked(this, &SDMToolBar::OnAutoCompileClicked)
 			[
 				SNew(SImage)
@@ -857,16 +856,6 @@
 	return FReply::Handled();
 }
 
-FReply SDMToolBar::OnApplyClicked()
-{
-	if (TSharedPtr<SDMMaterialEditor> EditorWidget = GetEditorWidget())
-	{
-		EditorWidget->ApplyToOriginal();
-	}
-
-	return FReply::Handled();
-}
-
 FReply SDMToolBar::OnOpenParentClicked()
 {
 	if (TSharedPtr<SDMMaterialEditor> EditorWidget = GetEditorWidget())
@@ -1074,19 +1063,6 @@
 	return FDMToolBarMenus::MakeEditorLayoutMenu(GetEditorWidget());
 }
 
-bool SDMToolBar::IsApplyEnabled() const
-{
-	if (TSharedPtr<SDMMaterialEditor> EditorWidget = EditorWidgetWeak.Pin())
-	{
-		if (UDynamicMaterialModelBase* MaterialModelBase = EditorWidget->GetPreviewMaterialModelBase())
-		{
-			return MaterialModelBase->IsPreviewModified();
-		}
-	}
-
-	return false;
-}
-
 FSlateColor SDMToolBar::GetAutoCompileColor() const
 {
 	static FSlateColor EnabledColor = FSlateColor(EStyleColor::AccentBlue);
@@ -1133,4 +1109,41 @@
 	return FReply::Handled();
 }
 
+FSlateColor SDMToolBar::GetAutoApplyColor() const
+{
+	static FSlateColor EnabledColor = FSlateColor(EStyleColor::AccentBlue);
+	static FSlateColor DisabledColor = FSlateColor(EStyleColor::AccentGray);
+
+	if (UDynamicMaterialEditorSettings* Settings = UDynamicMaterialEditorSettings::Get())
+	{
+		if (Settings->bAutomaticallyApplyToSourceOnPreviewCompile)
+		{
+			return EnabledColor;
+		}
+	}
+
+	return DisabledColor;
+}
+
+FReply SDMToolBar::OnAutoApplyClicked()
+{
+	if (UDynamicMaterialEditorSettings* Settings = UDynamicMaterialEditorSettings::Get())
+	{
+		FScopedTransaction Transaction(LOCTEXT("ToggleAutoApply", "Toggle Auto Apply to Source"));
+		Settings->Modify();
+		Settings->bAutomaticallyApplyToSourceOnPreviewCompile = !Settings->bAutomaticallyApplyToSourceOnPreviewCompile;
+		Settings->SaveConfig();
+
+		if (Settings->bAutomaticallyApplyToSourceOnPreviewCompile)
+		{
+			if (TSharedPtr<SDMMaterialEditor> EditorWidget = EditorWidgetWeak.Pin())
+			{
+				EditorWidget->ApplyToOriginal();
+			}
+		}
+	}
+
+	return FReply::Handled();
+}
+
 #undef LOCTEXT_NAMESPACE
diff --git a/Engine/Plugins/Experimental/DynamicMaterial/Source/DynamicMaterialEditor/Private/UI/Widgets/Editor/SDMToolBar.h b/Engine/Plugins/Experimental/DynamicMaterial/Source/DynamicMaterialEditor/Private/UI/Widgets/Editor/SDMToolBar.h
--- a/Engine/Plugins/Experimental/DynamicMaterial/Source/DynamicMaterialEditor/Private/UI/Widgets/Editor/SDMToolBar.h
+++ b/Engine/Plugins/Experimental/DynamicMaterial/Source/DynamicMaterialEditor/Private/UI/Widgets/Editor/SDMToolBar.h
@@ -55,8 +55,6 @@ protected:
 
 	TSharedPtr<SBox> PropertySelectorContainer;
 	TSharedPtr<SWidget> SaveButtonWidget;
-	TSharedPtr<SWidget> ApplyButtonWidget;
-	TSharedPtr<SWidget> AutoCompileButtonWidget;
 	TSharedPtr<SWidget> ActorRowWidget;
 	TSharedPtr<SWidget> AssetRowWidget;
 	TSharedPtr<STextBlock> ActorNameWidget;
@@ -103,8 +101,6 @@ protected:
 
 	FReply OnSaveClicked();
 
-	FReply OnApplyClicked();
-
 	FReply OnOpenParentClicked();
 
 	FReply OnConvertToEditableClicked();
@@ -113,9 +109,11 @@ protected:
 
 	TSharedRef<SWidget> GenerateSettingsMenu();
 
-	bool IsApplyEnabled() const;
-
 	FSlateColor GetAutoCompileColor() const;
 
 	FReply OnAutoCompileClicked();
+
+	FSlateColor GetAutoApplyColor() const;
+
+	FReply OnAutoApplyClicked();
 };
diff --git a/Engine/Plugins/Experimental/DynamicMaterial/Source/DynamicMaterialEditor/Private/UI/Widgets/Editor/SlotEditor/SDMMaterialSlotLayerEffectItem.cpp b/Engine/Plugins/Experimental/DynamicMaterial/Source/DynamicMaterialEditor/Private/UI/Widgets/Editor/SlotEditor/SDMMaterialSlotLayerEffectItem.cpp
--- a/Engine/Plugins/Experimental/DynamicMaterial/Source/DynamicMaterialEditor/Private/UI/Widgets/Editor/SlotEditor/SDMMaterialSlotLayerEffectItem.cpp
+++ b/Engine/Plugins/Experimental/DynamicMaterial/Source/DynamicMaterialEditor/Private/UI/Widgets/Editor/SlotEditor/SDMMaterialSlotLayerEffectItem.cpp
@@ -4,6 +4,7 @@
 
 #include "Components/DMMaterialEffectStack.h"
 #include "DynamicMaterialEditorStyle.h"
+#include "Editor.h"
 #include "Model/DynamicMaterialModel.h"
 #include "UI/DragDrop/DMLayerEffectsDragDropOperation.h"
 #include "UI/Widgets/Editor/SlotEditor/SDMMaterialSlotLayerEffectView.h"
@@ -105,6 +106,14 @@ TSharedRef<SWidget> SDMMaterialSlotLayerEffectItem::CreateMainContent()
 		.AutoWidth()
 		.HAlign(HAlign_Right)
 		.VAlign(VAlign_Center)
+		.Padding(HorizontalSpacing, VerticalSpacing, HorizontalSpacing, VerticalSpacing)
+		[
+			CreateBrowseToEffectButton()
+		]
+		+ SHorizontalBox::Slot()
+		.AutoWidth()
+		.HAlign(HAlign_Right)
+		.VAlign(VAlign_Center)
 		.Padding(HorizontalSpacingEnd, VerticalSpacing, HorizontalSpacingEnd, VerticalSpacing)
 		[
 			CreateLayerRemoveButton()
@@ -143,6 +152,22 @@ TSharedRef<SWidget> SDMMaterialSlotLayerEffectItem::CreateLayerRemoveButton()
 		];
 }
 
+TSharedRef<SWidget> SDMMaterialSlotLayerEffectItem::CreateBrowseToEffectButton()
+{
+	return SNew(SButton)
+		.Visibility(!!GetEffectAsset() ? EVisibility::Visible : EVisibility::Collapsed)
+		.ContentPadding(FMargin(2.0f, 2.f))
+		.ButtonStyle(FDynamicMaterialEditorStyle::Get(), "HoverHintOnly")
+		.ToolTipText(LOCTEXT("BrowseToEffectTooltip", "Browse to Effect in Content Browser"))
+		.Cursor(EMouseCursor::Default)
+		.OnClicked(this, &SDMMaterialSlotLayerEffectItem::OnBrowseToEffectButtonClick)
+		[
+			SNew(SImage)
+			.Image(FAppStyle::Get().GetBrush("SystemWideCommands.FindInContentBrowser"))
+			.DesiredSizeOverride(FVector2D(12.0f))
+		];
+}
+
 FText SDMMaterialSlotLayerEffectItem::GetToolTipText() const
 {
 	if (UDMMaterialEffect* MaterialEffect = GetMaterialEffect())
@@ -219,6 +244,17 @@ FReply SDMMaterialSlotLayerEffectItem::OnLayerRemoveButtonClick()
 	return FReply::Handled();
 }
 
+FReply SDMMaterialSlotLayerEffectItem::OnBrowseToEffectButtonClick()
+{
+	if (UObject* Asset = GetEffectAsset())
+	{
+		TArray<UObject*> Assets = {Asset};
+		GEditor->SyncBrowserToObjects(Assets); // Already validated in GetEffectAsset
+	}
+
+	return FReply::Handled();
+}
+
 int32 SDMMaterialSlotLayerEffectItem::OnEffectItemPaintDropIndicator(EItemDropZone InItemDropZone, const FPaintArgs& InArgs,
 	const FGeometry& InAllottedGeometry, const FSlateRect& InMyCullingRect, FSlateWindowElementList& OutDrawElements, int32 InLayerId, 
 	const FWidgetStyle& InWidgetStyle, bool bInParentEnabled) const
@@ -242,6 +278,25 @@ int32 SDMMaterialSlotLayerEffectItem::OnEffectItemPaintDropIndicator(EItemDropZo
 	return InLayerId;
 }
 
+UObject* SDMMaterialSlotLayerEffectItem::GetEffectAsset() const
+{
+	if (GEditor)
+	{
+		if (UDMMaterialEffect* Effect = GetMaterialEffect())
+		{
+			if (UObject* Asset = Effect->GetAsset())
+			{
+				if (Asset->IsAsset())
+				{
+					return Asset;
+				}
+			}
+		}
+	}
+
+	return nullptr;
+}
+
 TOptional<EItemDropZone> SDMMaterialSlotLayerEffectItem::OnEffectItemCanAcceptDrop(const FDragDropEvent& InDragDropEvent, 
 	EItemDropZone InDropZone, UDMMaterialEffect* InMaterialEffect) const
 {
diff --git a/Engine/Plugins/Experimental/DynamicMaterial/Source/DynamicMaterialEditor/Private/UI/Widgets/Editor/SlotEditor/SDMMaterialSlotLayerEffectItem.h b/Engine/Plugins/Experimental/DynamicMaterial/Source/DynamicMaterialEditor/Private/UI/Widgets/Editor/SlotEditor/SDMMaterialSlotLayerEffectItem.h
--- a/Engine/Plugins/Experimental/DynamicMaterial/Source/DynamicMaterialEditor/Private/UI/Widgets/Editor/SlotEditor/SDMMaterialSlotLayerEffectItem.h
+++ b/Engine/Plugins/Experimental/DynamicMaterial/Source/DynamicMaterialEditor/Private/UI/Widgets/Editor/SlotEditor/SDMMaterialSlotLayerEffectItem.h
@@ -36,6 +36,8 @@ protected:
 		const FSlateRect& InMyCullingRect, FSlateWindowElementList& OutDrawElements, int32 InLayerId, const FWidgetStyle& InWidgetStyle, 
 		bool bInParentEnabled) const;
 
+	UObject* GetEffectAsset() const;
+
 	/** Slots */
 	TSharedRef<SWidget> CreateMainContent();
 
@@ -43,6 +45,8 @@ protected:
 
 	TSharedRef<SWidget> CreateLayerRemoveButton();
 
+	TSharedRef<SWidget> CreateBrowseToEffectButton();
+
 	FText GetToolTipText() const;
 
 	FText GetLayerHeaderText() const;
@@ -53,6 +57,8 @@ protected:
 
 	FReply OnLayerRemoveButtonClick();
 
+	FReply OnBrowseToEffectButtonClick();
+
 	/** Drag and Drop */
 	TOptional<EItemDropZone> OnEffectItemCanAcceptDrop(const FDragDropEvent& InDragDropEvent, EItemDropZone InDropZone,
 		UDMMaterialEffect* InMaterialEffect) const;
diff --git a/Engine/Plugins/Experimental/DynamicMaterial/Source/DynamicMaterialEditor/Private/UI/Widgets/SDMMaterialEditor.cpp b/Engine/Plugins/Experimental/DynamicMaterial/Source/DynamicMaterialEditor/Private/UI/Widgets/SDMMaterialEditor.cpp
--- a/Engine/Plugins/Experimental/DynamicMaterial/Source/DynamicMaterialEditor/Private/UI/Widgets/SDMMaterialEditor.cpp
+++ b/Engine/Plugins/Experimental/DynamicMaterial/Source/DynamicMaterialEditor/Private/UI/Widgets/SDMMaterialEditor.cpp
@@ -141,12 +141,7 @@ SDMMaterialEditor::~SDMMaterialEditor()
 		return;
 	}
 
-	if (UDynamicMaterialModelEditorOnlyData* EditorOnlyData = EditorOnlyDataUpdateObject.Get())
-	{
-		EditorOnlyData->GetOnMaterialBuiltDelegate().RemoveAll(this);
-		EditorOnlyData->GetOnPropertyUpdateDelegate().RemoveAll(this);
-		EditorOnlyData->GetOnSlotListUpdateDelegate().RemoveAll(this);
-	}
+	UnbindEditorOnlyDataUpdate();
 
 	if (UDynamicMaterialEditorSettings* Settings = GetMutableDefault<UDynamicMaterialEditorSettings>())
 	{
@@ -247,7 +242,12 @@ void SDMMaterialEditor::SetMaterialModelBase(UDynamicMaterialModelBase* InMateri
 
 	CreateLayout();
 
-	BindEditorOnlyDataUpdate(InMaterialModelBase);
+	UnbindEditorOnlyDataUpdate();
+
+	if (PreviewMaterialModelBase)
+	{
+		BindEditorOnlyDataUpdate(PreviewMaterialModelBase);
+	}
 }
 
 UDynamicMaterialModel* SDMMaterialEditor::GetPreviewMaterialModel() const
@@ -287,8 +287,6 @@ void SDMMaterialEditor::SetObjectMaterialProperty(const FDMObjectMaterialPropert
 
 	ObjectMaterialPropertyOpt = InObjectProperty;
 	SetMaterialModelBase(MaterialModelBase, InPreviewMaterialModelBase);
-
-	BindEditorOnlyDataUpdate(MaterialModelBase);
 }
 
 AActor* SDMMaterialEditor::GetMaterialActor() const
@@ -1462,9 +1460,27 @@ void SDMMaterialEditor::BindEditorOnlyDataUpdate(UDynamicMaterialModelBase* InMa
 	}
 }
 
+void SDMMaterialEditor::UnbindEditorOnlyDataUpdate()
+{
+	if (UDynamicMaterialModelEditorOnlyData* EditorOnlyData = EditorOnlyDataUpdateObject.Get())
+	{
+		EditorOnlyData->GetOnMaterialBuiltDelegate().RemoveAll(this);
+		EditorOnlyData->GetOnPropertyUpdateDelegate().RemoveAll(this);
+		EditorOnlyData->GetOnSlotListUpdateDelegate().RemoveAll(this);
+	}
+}
+
 void SDMMaterialEditor::OnMaterialBuilt(UDynamicMaterialModelBase* InMaterialModelBase)
 {
 	PropertySelectorSlot.Invalidate();
+
+	if (UDynamicMaterialEditorSettings* Settings = UDynamicMaterialEditorSettings::Get())
+	{
+		if (Settings->bAutomaticallyApplyToSourceOnPreviewCompile)
+		{
+			ApplyToOriginal();
+		}
+	}
 }
 
 void SDMMaterialEditor::OnPropertyUpdate(UDynamicMaterialModelBase* InMaterialModelBase)
diff --git a/Engine/Plugins/Experimental/DynamicMaterial/Source/DynamicMaterialEditor/Private/UI/Widgets/SDMMaterialEditor.h b/Engine/Plugins/Experimental/DynamicMaterial/Source/DynamicMaterialEditor/Private/UI/Widgets/SDMMaterialEditor.h
--- a/Engine/Plugins/Experimental/DynamicMaterial/Source/DynamicMaterialEditor/Private/UI/Widgets/SDMMaterialEditor.h
+++ b/Engine/Plugins/Experimental/DynamicMaterial/Source/DynamicMaterialEditor/Private/UI/Widgets/SDMMaterialEditor.h
@@ -268,6 +268,8 @@ protected:
 
 	void BindEditorOnlyDataUpdate(UDynamicMaterialModelBase* InMaterialModelBase);
 
+	void UnbindEditorOnlyDataUpdate();
+
 	void OnMaterialBuilt(UDynamicMaterialModelBase* InMaterialModelBase);
 
 	void OnPropertyUpdate(UDynamicMaterialModelBase* InMaterialModelBase);
diff --git a/Engine/Plugins/Experimental/DynamicMaterial/Source/DynamicMaterialEditor/Public/Components/DMMaterialEffect.h b/Engine/Plugins/Experimental/DynamicMaterial/Source/DynamicMaterialEditor/Public/Components/DMMaterialEffect.h
--- a/Engine/Plugins/Experimental/DynamicMaterial/Source/DynamicMaterialEditor/Public/Components/DMMaterialEffect.h
+++ b/Engine/Plugins/Experimental/DynamicMaterial/Source/DynamicMaterialEditor/Public/Components/DMMaterialEffect.h
@@ -75,6 +75,9 @@ public:
 	virtual void ApplyTo(const TSharedRef<FDMMaterialBuildState>& InBuildState, TArray<UMaterialExpression*>& InOutExpressions, 
 		int32& InOutLastExpressionOutputChannel, int32& InLastExpressionOutputIndex) const PURE_VIRTUAL(UDMMaterialEffect::ApplyTo)
 
+	/** Returns the asset associated with this effect, if any. */
+	virtual UObject* GetAsset() const PURE_VIRTUAL(UMaterialEffect::GetAsset, return nullptr;)
+
 	//~ Begin UDMMaterialComponent
 	DYNAMICMATERIALEDITOR_API virtual UDMMaterialComponent* GetParentComponent() const override;
 	DYNAMICMATERIALEDITOR_API virtual FString GetComponentPathComponent() const override;
diff --git a/Engine/Plugins/Experimental/DynamicMaterial/Source/DynamicMaterialEditor/Public/Components/DMMaterialEffectFunction.h b/Engine/Plugins/Experimental/DynamicMaterial/Source/DynamicMaterialEditor/Public/Components/DMMaterialEffectFunction.h
--- a/Engine/Plugins/Experimental/DynamicMaterial/Source/DynamicMaterialEditor/Public/Components/DMMaterialEffectFunction.h
+++ b/Engine/Plugins/Experimental/DynamicMaterial/Source/DynamicMaterialEditor/Public/Components/DMMaterialEffectFunction.h
@@ -45,6 +45,7 @@ public:
 	DYNAMICMATERIALEDITOR_API virtual bool IsCompatibleWith(UDMMaterialEffect* InEffect) const override;
 	DYNAMICMATERIALEDITOR_API virtual void ApplyTo(const TSharedRef<FDMMaterialBuildState>& InBuildState, TArray<UMaterialExpression*>& InOutExpressions,
 		int32& InOutLastExpressionOutputChannel, int32& InOutLastExpressionOutputIndex) const override;
+	virtual UObject* GetAsset() const override;
 	//~ End UDMMaterialEffect
 
 	//~ Begin UDMMaterialComponent
