Material Designer: Undo crash fix

diff --git a/Engine/Plugins/Experimental/DynamicMaterial/Source/DynamicMaterialEditor/Private/Model/DynamicMaterialModelEditorOnlyData.cpp b/Engine/Plugins/Experimental/DynamicMaterial/Source/DynamicMaterialEditor/Private/Model/DynamicMaterialModelEditorOnlyData.cpp
--- a/Engine/Plugins/Experimental/DynamicMaterial/Source/DynamicMaterialEditor/Private/Model/DynamicMaterialModelEditorOnlyData.cpp
+++ b/Engine/Plugins/Experimental/DynamicMaterial/Source/DynamicMaterialEditor/Private/Model/DynamicMaterialModelEditorOnlyData.cpp
@@ -1475,11 +1475,14 @@ void UDynamicMaterialModelEditorOnlyData::OnValueUpdated(UDMMaterialValue* InVal
 
 	// Non-exported materials have their values update via settings parameters
 	// Exported materials need to be rebuilt to update the main material.
-	const bool bMaterialInDifferentPackage = MaterialModel->DynamicMaterial ? MaterialModel->DynamicMaterial->GetPackage() != GetPackage() : true;
-
-	if (EnumHasAnyFlags(InUpdateType, EDMUpdateType::Structure) || bMaterialInDifferentPackage)
+	if (MaterialModel)
 	{
-		RequestMaterialBuild();
+		const bool bMaterialInDifferentPackage = MaterialModel->DynamicMaterial ? MaterialModel->DynamicMaterial->GetPackage() != GetPackage() : true;
+
+		if (EnumHasAnyFlags(InUpdateType, EDMUpdateType::Structure) || bMaterialInDifferentPackage)
+		{
+			RequestMaterialBuild();
+		}
 	}
 }
 
