Material Designer: Fixed texture edge color property display

diff --git a/Engine/Plugins/Experimental/DynamicMaterial/Source/DynamicMaterialEditor/Private/Components/MaterialStageExpressions/DMMSETextureSampleEdgeColor.cpp b/Engine/Plugins/Experimental/DynamicMaterial/Source/DynamicMaterialEditor/Private/Components/MaterialStageExpressions/DMMSETextureSampleEdgeColor.cpp
--- a/Engine/Plugins/Experimental/DynamicMaterial/Source/DynamicMaterialEditor/Private/Components/MaterialStageExpressions/DMMSETextureSampleEdgeColor.cpp
+++ b/Engine/Plugins/Experimental/DynamicMaterial/Source/DynamicMaterialEditor/Private/Components/MaterialStageExpressions/DMMSETextureSampleEdgeColor.cpp
@@ -77,6 +77,11 @@ void UDMMaterialStageExpressionTextureSampleEdgeColor::AddDefaultInput(int32 InI
 	return Super::AddDefaultInput(InInputIndex);
 }
 
+void UDMMaterialStageExpressionTextureSampleEdgeColor::PreEditChange(FEditPropertyChain& InPropertyAboutToChange)
+{
+	PreEditEdgeLocation = EdgeLocation;
+}
+
 void UDMMaterialStageExpressionTextureSampleEdgeColor::PostEditChangeProperty(FPropertyChangedEvent& InPropertyChangedEvent)
 {
 	Super::PostEditChangeProperty(InPropertyChangedEvent);
@@ -91,78 +96,78 @@ void UDMMaterialStageExpressionTextureSampleEdgeColor::PostEditChangeProperty(FP
 
 void UDMMaterialStageExpressionTextureSampleEdgeColor::OnEdgeLocationChanged()
 {
-	if (EdgeLocation == EDMEdgeLocation::Custom)
+	if (EdgeLocation != EDMEdgeLocation::Custom)
 	{
-		return;
-	}
-
-	using namespace UE::DynamicMaterialEditor::TextureSampleEdgeColor::Private;
-
-	UDMMaterialStage* Stage = GetStage();
-	check(Stage);
-	check(Stage->GetInputs().IsValidIndex(1));
+		using namespace UE::DynamicMaterialEditor::TextureSampleEdgeColor::Private;
 
-	UDMMaterialStageInputValue* InputValue = Cast<UDMMaterialStageInputValue>(Stage->GetInputs()[1]);
-	check(InputValue);
-	
-	UDMMaterialValueFloat2* Float2 = Cast<UDMMaterialValueFloat2>(InputValue->GetValue());
-	check(Float2);
+		UDMMaterialStage* Stage = GetStage();
+		check(Stage);
+		check(Stage->GetInputs().IsValidIndex(1));
+
+		UDMMaterialStageInputValue* InputValue = CastChecked<UDMMaterialStageInputValue>(Stage->GetInputs()[1]);
+		UDMMaterialValueFloat2* Float2 = CastChecked<UDMMaterialValueFloat2>(InputValue->GetValue());
+
+		FVector2D Float2Value = FVector2D::ZeroVector;
+
+		switch (EdgeLocation)
+		{
+			case EDMEdgeLocation::TopLeft:
+				Float2Value.X = EdgeGap;
+				Float2Value.Y = EdgeGap;
+				break;
+
+			case EDMEdgeLocation::Top:
+				Float2Value.X = 0.5;
+				Float2Value.Y = EdgeGap;
+				break;
+
+			case EDMEdgeLocation::TopRight:
+				Float2Value.X = 1.f - EdgeGap;
+				Float2Value.Y = EdgeGap;
+				break;
+
+			case EDMEdgeLocation::Left:
+				Float2Value.X = EdgeGap;
+				Float2Value.Y = 0.5;
+				break;
+
+			case EDMEdgeLocation::Center:
+				Float2Value.X = 0.5;
+				Float2Value.Y = 0.5;
+				break;
+
+			case EDMEdgeLocation::Right:
+				Float2Value.X = 1.f - EdgeGap;
+				Float2Value.Y = 0.5;
+				break;
+
+			case EDMEdgeLocation::BottomLeft:
+				Float2Value.X = EdgeGap;
+				Float2Value.Y = 1.f - EdgeGap;
+				break;
+
+			case EDMEdgeLocation::Bottom:
+				Float2Value.X = 0.5;
+				Float2Value.Y = 1.f - EdgeGap;
+				break;
+
+			case EDMEdgeLocation::BottomRight:
+				Float2Value.X = 1.f - EdgeGap;
+				Float2Value.Y = 1.f - EdgeGap;
+				break;
+
+			default:
+				checkNoEntry();
+				break;
+		}
+
+		Float2->SetValue(Float2Value);
+	}
 
-	FVector2D Float2Value = FVector2D::ZeroVector;
-	
-	switch (EdgeLocation)
+	if (EdgeLocation == EDMEdgeLocation::Custom || PreEditEdgeLocation == EDMEdgeLocation::Custom)
 	{
-		case EDMEdgeLocation::TopLeft:
-			Float2Value.X = EdgeGap;
-			Float2Value.Y = EdgeGap;
-			break;
-
-		case EDMEdgeLocation::Top:
-			Float2Value.X = 0.5;
-			Float2Value.Y = EdgeGap;
-			break;
-
-		case EDMEdgeLocation::TopRight:
-			Float2Value.X = 1.f - EdgeGap;
-			Float2Value.Y = EdgeGap;
-			break;
-
-		case EDMEdgeLocation::Left:
-			Float2Value.X = EdgeGap;
-			Float2Value.Y = 0.5;
-			break;
-
-		case EDMEdgeLocation::Center:
-			Float2Value.X = 0.5;
-			Float2Value.Y = 0.5;
-			break;
-
-		case EDMEdgeLocation::Right:
-			Float2Value.X = 1.f - EdgeGap;
-			Float2Value.Y = 0.5;
-			break;
-
-		case EDMEdgeLocation::BottomLeft:
-			Float2Value.X = EdgeGap;
-			Float2Value.Y = 1.f - EdgeGap;
-			break;
-
-		case EDMEdgeLocation::Bottom:
-			Float2Value.X = 0.5;
-			Float2Value.Y = 1.f - EdgeGap;
-			break;
-
-		case EDMEdgeLocation::BottomRight:
-			Float2Value.X = 1.f - EdgeGap;
-			Float2Value.Y = 1.f - EdgeGap;
-			break;
-
-		default:
-			checkNoEntry();
-			break;
+		Update(this, EDMUpdateType::RefreshDetailView);
 	}
-
-	Float2->SetValue(Float2Value);
 }
 
 #undef LOCTEXT_NAMESPACE
diff --git a/Engine/Plugins/Experimental/DynamicMaterial/Source/DynamicMaterialEditor/Public/Components/MaterialStageExpressions/DMMSETextureSampleEdgeColor.h b/Engine/Plugins/Experimental/DynamicMaterial/Source/DynamicMaterialEditor/Public/Components/MaterialStageExpressions/DMMSETextureSampleEdgeColor.h
--- a/Engine/Plugins/Experimental/DynamicMaterial/Source/DynamicMaterialEditor/Public/Components/MaterialStageExpressions/DMMSETextureSampleEdgeColor.h
+++ b/Engine/Plugins/Experimental/DynamicMaterial/Source/DynamicMaterialEditor/Public/Components/MaterialStageExpressions/DMMSETextureSampleEdgeColor.h
@@ -35,6 +35,7 @@ public:
 	//~ End UDMMaterialStageThroughput
 
 	//~ Begin UObject
+	DYNAMICMATERIALEDITOR_API virtual void PreEditChange(FEditPropertyChain& InPropertyAboutToChange) override;
 	DYNAMICMATERIALEDITOR_API virtual void PostEditChangeProperty(FPropertyChangedEvent& InPropertyChangedEvent) override;
 	//~ End UObject
 
@@ -42,5 +43,7 @@ protected:
 	UPROPERTY(EditInstanceOnly, BlueprintReadOnly, Category = "Material Designer", meta=(AllowPrivateAccess = "true"))
 	EDMEdgeLocation EdgeLocation;
 
+	EDMEdgeLocation PreEditEdgeLocation;
+
 	void OnEdgeLocationChanged();
 };
