Material Designer: Material build issues
- Fixed properties generating nodes when they are disabled.
- Fixed translucent materials not being translucent at creation

diff --git a/Engine/Plugins/Experimental/DynamicMaterial/Source/DynamicMaterialEditor/Private/Components/DMMaterialProperty.cpp b/Engine/Plugins/Experimental/DynamicMaterial/Source/DynamicMaterialEditor/Private/Components/DMMaterialProperty.cpp
--- a/Engine/Plugins/Experimental/DynamicMaterial/Source/DynamicMaterialEditor/Private/Components/DMMaterialProperty.cpp
+++ b/Engine/Plugins/Experimental/DynamicMaterial/Source/DynamicMaterialEditor/Private/Components/DMMaterialProperty.cpp
@@ -797,7 +797,7 @@ void UDMMaterialProperty::SetOutputProcessor(UMaterialFunctionInterface* InFunct
 	OnOutputProcessorUpdated();
 }
 
-bool UDMMaterialProperty::IsValidForModel(UDynamicMaterialModelEditorOnlyData& InMaterialModel) const
+bool UDMMaterialProperty::IsValidForModel(const UDynamicMaterialModelEditorOnlyData& InMaterialModel) const
 {
 	return FDMMaterialUtils::IsMaterialPropertyActive({
 		FDMUtils::MaterialPropertyTypeToMaterialProperty(MaterialProperty),
diff --git a/Engine/Plugins/Experimental/DynamicMaterial/Source/DynamicMaterialEditor/Private/Model/DynamicMaterialModelEditorOnlyData.cpp b/Engine/Plugins/Experimental/DynamicMaterial/Source/DynamicMaterialEditor/Private/Model/DynamicMaterialModelEditorOnlyData.cpp
--- a/Engine/Plugins/Experimental/DynamicMaterial/Source/DynamicMaterialEditor/Private/Model/DynamicMaterialModelEditorOnlyData.cpp
+++ b/Engine/Plugins/Experimental/DynamicMaterial/Source/DynamicMaterialEditor/Private/Model/DynamicMaterialModelEditorOnlyData.cpp
@@ -342,6 +342,7 @@ void UDynamicMaterialModelEditorOnlyData::CreateMaterial()
 		UMaterialFactoryNew* MaterialFactory = NewObject<UMaterialFactoryNew>();
 		check(MaterialFactory);
 
+		// Replace existing material
 		MaterialModel->DynamicMaterial = Cast<UMaterial>(MaterialFactory->FactoryCreateNew(
 			UMaterial::StaticClass(),
 			MaterialModel,
@@ -355,7 +356,13 @@ void UDynamicMaterialModelEditorOnlyData::CreateMaterial()
 	{
 		FString MaterialBaseName = GetName() + "-" + FGuid::NewGuid().ToString();
 		const FString FullName = "/Game/DynamicMaterials/" + MaterialBaseName;
-		UPackage* Package = CreatePackage(*FullName);
+
+		UPackage* Package = LoadObject<UPackage>(GetTransientPackage(), *FullName);
+
+		if (!Package)
+		{
+			Package = CreatePackage(*FullName);
+		}
 
 		UMaterialFactoryNew* MaterialFactory = NewObject<UMaterialFactoryNew>();
 		check(MaterialFactory);
@@ -388,10 +395,7 @@ void UDynamicMaterialModelEditorOnlyData::BuildMaterial(bool bInDirtyAssets)
 
 	UE_LOG(LogDynamicMaterialEditor, Display, TEXT("Building Material Designer Material (%s)..."), *MaterialModel->GetFullName());
 
-	if (!IsValid(MaterialModel->DynamicMaterial))
-	{
-		CreateMaterial();
-	}
+	CreateMaterial();
 
 	State = EDMState::Building;
 	Expressions.Empty();
@@ -442,7 +446,8 @@ void UDynamicMaterialModelEditorOnlyData::BuildMaterial(bool bInDirtyAssets)
 
 			if (UDMMaterialProperty* Property = GetMaterialProperty(InType))
 			{
-				if (Property->IsEnabled() && Property->IsMaterialPin() && GetSlotForMaterialProperty(InType))
+				if (Property->IsEnabled() && Property->IsMaterialPin() && GetSlotForMaterialProperty(InType)
+					&& Property->IsValidForModel(*this))
 				{
 					Property->GenerateExpressions(BuildState.ToSharedRef());
 
@@ -554,7 +559,8 @@ void UDynamicMaterialModelEditorOnlyData::BuildMaterial(bool bInDirtyAssets)
 		{
 			if (UDMMaterialProperty* Property = GetMaterialProperty(InType))
 			{
-				if (Property->IsMaterialPin() && Property->IsEnabled() && GetSlotForMaterialProperty(InType))
+				if (Property->IsMaterialPin() && Property->IsEnabled() && GetSlotForMaterialProperty(InType)
+					&& Property->IsValidForModel(*this))
 				{
 					Property->AddOutputProcessor(BuildState.ToSharedRef());
 				}
@@ -1151,7 +1157,7 @@ UDMMaterialSlot* UDynamicMaterialModelEditorOnlyData::GetSlotForEnabledMaterialP
 {
 	UDMMaterialProperty* Property = GetMaterialProperty(InType);
 
-	if (!Property || !Property->IsEnabled())
+	if (!Property || !Property->IsEnabled() || !Property->IsValidForModel(*this))
 	{
 		return nullptr;
 	}
@@ -1213,8 +1219,8 @@ UDMMaterialSlot* UDynamicMaterialModelEditorOnlyData::AddSlotForMaterialProperty
 		return ExistingSlot;
 	}
 
-	// BaseColor and Emissive are mutually exclusive so if something tries to add one of them, the other must
-	// be checked. If it is found, it is converted and returned. The same goes for Opacity and OpacityMask.
+	// Opacity and OpacityMask are mutually exclusive so if something tries to add one of them, the other must
+	// be checked. If it is found, it is converted and returned.
 	switch (InType)
 	{
 		case EDMMaterialPropertyType::Opacity:
diff --git a/Engine/Plugins/Experimental/DynamicMaterial/Source/DynamicMaterialEditor/Public/Components/DMMaterialProperty.h b/Engine/Plugins/Experimental/DynamicMaterial/Source/DynamicMaterialEditor/Public/Components/DMMaterialProperty.h
--- a/Engine/Plugins/Experimental/DynamicMaterial/Source/DynamicMaterialEditor/Public/Components/DMMaterialProperty.h
+++ b/Engine/Plugins/Experimental/DynamicMaterial/Source/DynamicMaterialEditor/Public/Components/DMMaterialProperty.h
@@ -105,7 +105,7 @@ public:
 	 * @param InMaterialModel The model with the settings.
 	 * @return True if it is valid.
 	 */
-	DYNAMICMATERIALEDITOR_API virtual bool IsValidForModel(UDynamicMaterialModelEditorOnlyData& InMaterialModel) const;
+	DYNAMICMATERIALEDITOR_API virtual bool IsValidForModel(const UDynamicMaterialModelEditorOnlyData& InMaterialModel) const;
 
 	/**
 	 * Empty the connection map and recalculate the input types.
