Material Designer: Fixed marking map dirty on load
- When the preview material is initially compiled on open, it will not be automatically applied to the source material.

diff --git a/Engine/Plugins/Experimental/DynamicMaterial/Source/DynamicMaterial/Private/Model/DynamicMaterialModel.cpp b/Engine/Plugins/Experimental/DynamicMaterial/Source/DynamicMaterial/Private/Model/DynamicMaterialModel.cpp
--- a/Engine/Plugins/Experimental/DynamicMaterial/Source/DynamicMaterial/Private/Model/DynamicMaterialModel.cpp
+++ b/Engine/Plugins/Experimental/DynamicMaterial/Source/DynamicMaterial/Private/Model/DynamicMaterialModel.cpp
@@ -747,7 +747,7 @@ void UDynamicMaterialModel::PostEditorDuplicate()
 		{
 			if (GUndo)
 			{
-				ParameterPair.Value->Modify();
+				Parameter->Modify();
 			}
 
 			Parameter->PostEditorDuplicate(this, nullptr);
diff --git a/Engine/Plugins/Experimental/DynamicMaterial/Source/DynamicMaterialEditor/Private/UI/Widgets/SDMMaterialEditor.cpp b/Engine/Plugins/Experimental/DynamicMaterial/Source/DynamicMaterialEditor/Private/UI/Widgets/SDMMaterialEditor.cpp
--- a/Engine/Plugins/Experimental/DynamicMaterial/Source/DynamicMaterialEditor/Private/UI/Widgets/SDMMaterialEditor.cpp
+++ b/Engine/Plugins/Experimental/DynamicMaterial/Source/DynamicMaterialEditor/Private/UI/Widgets/SDMMaterialEditor.cpp
@@ -122,7 +122,16 @@ SDMMaterialEditor::SDMMaterialEditor()
 	: SplitterSlot(nullptr)
 	, CommandList(MakeShared<FUICommandList>())
 	, PreviewMaterialManager(MakeShared<FDMPreviewMaterialManager>())
+	, bSkipApplyOnCompile(false)
 {
+	if (UDynamicMaterialEditorSettings* Settings = UDynamicMaterialEditorSettings::Get())
+	{
+		if (Settings->bAutomaticallyApplyToSourceOnPreviewCompile)
+		{
+			// If we are automatically applying, we need to skip the initial compile event.
+			bSkipApplyOnCompile = true;
+		}
+	}
 }
 
 SDMMaterialEditor::~SDMMaterialEditor()
@@ -1474,6 +1483,12 @@ void SDMMaterialEditor::OnMaterialBuilt(UDynamicMaterialModelBase* InMaterialMod
 {
 	PropertySelectorSlot.Invalidate();
 
+	if (bSkipApplyOnCompile)
+	{
+		bSkipApplyOnCompile = false;
+		return;
+	}
+
 	if (UDynamicMaterialEditorSettings* Settings = UDynamicMaterialEditorSettings::Get())
 	{
 		if (Settings->bAutomaticallyApplyToSourceOnPreviewCompile)
diff --git a/Engine/Plugins/Experimental/DynamicMaterial/Source/DynamicMaterialEditor/Private/UI/Widgets/SDMMaterialEditor.h b/Engine/Plugins/Experimental/DynamicMaterial/Source/DynamicMaterialEditor/Private/UI/Widgets/SDMMaterialEditor.h
--- a/Engine/Plugins/Experimental/DynamicMaterial/Source/DynamicMaterialEditor/Private/UI/Widgets/SDMMaterialEditor.h
+++ b/Engine/Plugins/Experimental/DynamicMaterial/Source/DynamicMaterialEditor/Private/UI/Widgets/SDMMaterialEditor.h
@@ -208,6 +208,8 @@ protected:
 
 	TWeakObjectPtr<UDynamicMaterialModelEditorOnlyData> EditorOnlyDataUpdateObject;
 
+	bool bSkipApplyOnCompile;
+
 	/** Operations */
 	void SetMaterialModelBase(UDynamicMaterialModelBase* InMaterialModelBase, UDynamicMaterialModelBase* InPreviewMaterialModelBase);
 
