Material Designer: "Live Edit" mode
- Option to combine live preview and auto apply into a single mode called "Live Edit"

diff --git a/Engine/Plugins/Experimental/DynamicMaterial/Source/DynamicMaterialEditor/Private/DynamicMaterialEditorSettings.cpp b/Engine/Plugins/Experimental/DynamicMaterial/Source/DynamicMaterialEditor/Private/DynamicMaterialEditorSettings.cpp
--- a/Engine/Plugins/Experimental/DynamicMaterial/Source/DynamicMaterialEditor/Private/DynamicMaterialEditorSettings.cpp
+++ b/Engine/Plugins/Experimental/DynamicMaterial/Source/DynamicMaterialEditor/Private/DynamicMaterialEditorSettings.cpp
@@ -137,6 +137,7 @@ UDynamicMaterialEditorSettings::UDynamicMaterialEditorSettings()
 	bAutomaticallyCopyParametersToSourceMaterial = true;
 	bAutomaticallyCompilePreviewMaterial = false;
 	bAutomaticallyApplyToSourceOnPreviewCompile = true;
+	LiveEditMode = EDMLiveEditMode::LiveEditOn;
 	bUseLinearColorForVectors = true;
 
 	ResetAllLayoutSettings();
diff --git a/Engine/Plugins/Experimental/DynamicMaterial/Source/DynamicMaterialEditor/Private/DynamicMaterialEditorSettings.h b/Engine/Plugins/Experimental/DynamicMaterial/Source/DynamicMaterialEditor/Private/DynamicMaterialEditorSettings.h
--- a/Engine/Plugins/Experimental/DynamicMaterial/Source/DynamicMaterialEditor/Private/DynamicMaterialEditorSettings.h
+++ b/Engine/Plugins/Experimental/DynamicMaterial/Source/DynamicMaterialEditor/Private/DynamicMaterialEditorSettings.h
@@ -22,10 +22,10 @@ struct FDMMaterialEffectList
 {
 	GENERATED_BODY()
 
-	UPROPERTY(Config, EditAnywhere, BlueprintReadWrite, Category = "Effects")
+	UPROPERTY(Config, EditAnywhere, BlueprintReadOnly, Category = "Effects")
 	FString Name;
 
-	UPROPERTY(Config, EditAnywhere, BlueprintReadWrite, Category = "Effects")
+	UPROPERTY(Config, EditAnywhere, BlueprintReadOnly, Category = "Effects")
 	TArray<TSoftObjectPtr<UMaterialFunctionInterface>> Effects;
 };
 
@@ -47,13 +47,13 @@ struct FDMDefaultMaterialPropertySlotValue
 	FDMDefaultMaterialPropertySlotValue(EDMDefaultMaterialPropertySlotValueType InDefaultType, 
 		const TSoftObjectPtr<UTexture>& InTexture, const FLinearColor& InColor);
 
-	UPROPERTY(EditAnywhere, BlueprintReadWrite, Category = "Material Designer")
+	UPROPERTY(EditAnywhere, BlueprintReadOnly, Category = "Material Designer")
 	EDMDefaultMaterialPropertySlotValueType DefaultType;
 
-	UPROPERTY(EditAnywhere, BlueprintReadWrite, Category = "Material Designer")
+	UPROPERTY(EditAnywhere, BlueprintReadOnly, Category = "Material Designer")
 	TSoftObjectPtr<UTexture> Texture;
 
-	UPROPERTY(EditAnywhere, BlueprintReadWrite, Category = "Material Designer")
+	UPROPERTY(EditAnywhere, BlueprintReadOnly, Category = "Material Designer")
 	FLinearColor Color;
 };
 
@@ -71,58 +71,58 @@ struct FDMMaterialChannelListPreset
 	UPROPERTY(Config, EditAnywhere, BlueprintReadOnly, Category = "Channel")
 	bool bEmissive = false;
 
-	UPROPERTY(Config, EditAnywhere, BlueprintReadWrite, Category = "Channel")
+	UPROPERTY(Config, EditAnywhere, BlueprintReadOnly, Category = "Channel")
 	bool bOpacity = false;
 
-	UPROPERTY(Config, EditAnywhere, BlueprintReadWrite, Category = "Channel")
+	UPROPERTY(Config, EditAnywhere, BlueprintReadOnly, Category = "Channel")
 	bool bRoughness = false;
 
-	UPROPERTY(Config, EditAnywhere, BlueprintReadWrite, Category = "Channel")
+	UPROPERTY(Config, EditAnywhere, BlueprintReadOnly, Category = "Channel")
 	bool bSpecular = false;
 
-	UPROPERTY(Config, EditAnywhere, BlueprintReadWrite, Category = "Channel")
+	UPROPERTY(Config, EditAnywhere, BlueprintReadOnly, Category = "Channel")
 	bool bMetallic = false;
 
-	UPROPERTY(Config, EditAnywhere, BlueprintReadWrite, Category = "Channel")
+	UPROPERTY(Config, EditAnywhere, BlueprintReadOnly, Category = "Channel")
 	bool bNormal = false;
 
-	UPROPERTY(Config, EditAnywhere, BlueprintReadWrite, Category = "Channel")
+	UPROPERTY(Config, EditAnywhere, BlueprintReadOnly, Category = "Channel")
 	bool bPixelDepthOffset = false;
 
-	UPROPERTY(Config, EditAnywhere, BlueprintReadWrite, Category = "Channel")
+	UPROPERTY(Config, EditAnywhere, BlueprintReadOnly, Category = "Channel")
 	bool bWorldPositionOffset = false;
 
-	UPROPERTY(Config, EditAnywhere, BlueprintReadWrite, Category = "Channel")
+	UPROPERTY(Config, EditAnywhere, BlueprintReadOnly, Category = "Channel")
 	bool bAmbientOcclusion = false;
 
-	UPROPERTY(Config, EditAnywhere, BlueprintReadWrite, Category = "Channel")
+	UPROPERTY(Config, EditAnywhere, BlueprintReadOnly, Category = "Channel")
 	bool bAnisotropy = false;
 
-	UPROPERTY(Config, EditAnywhere, BlueprintReadWrite, Category = "Channel")
+	UPROPERTY(Config, EditAnywhere, BlueprintReadOnly, Category = "Channel")
 	bool bRefraction = false;
 
-	UPROPERTY(Config, EditAnywhere, BlueprintReadWrite, Category = "Channel")
+	UPROPERTY(Config, EditAnywhere, BlueprintReadOnly, Category = "Channel")
 	bool bTangent = false;
 
-	UPROPERTY(Config, EditAnywhere, BlueprintReadWrite, Category = "Channel")
+	UPROPERTY(Config, EditAnywhere, BlueprintReadOnly, Category = "Channel")
 	bool bDisplacement = false;
 
-	UPROPERTY(Config, EditAnywhere, BlueprintReadWrite, Category = "Channel")
+	UPROPERTY(Config, EditAnywhere, BlueprintReadOnly, Category = "Channel")
 	bool bSubsurfaceColor = false;
 
-	UPROPERTY(Config, EditAnywhere, BlueprintReadWrite, Category = "Channel")
+	UPROPERTY(Config, EditAnywhere, BlueprintReadOnly, Category = "Channel")
 	bool bSurfaceThickness = false;
 
-	UPROPERTY(Config, EditAnywhere, BlueprintReadWrite, Category = "Channel")
+	UPROPERTY(Config, EditAnywhere, BlueprintReadOnly, Category = "Channel")
 	TEnumAsByte<EBlendMode> DefaultBlendMode = BLEND_Opaque;
 
-	UPROPERTY(Config, EditAnywhere, BlueprintReadWrite, Category = "Channel")
+	UPROPERTY(Config, EditAnywhere, BlueprintReadOnly, Category = "Channel")
 	EDMMaterialShadingModel DefaultShadingModel = EDMMaterialShadingModel::Unlit;
 
-	UPROPERTY(Config, EditAnywhere, BlueprintReadWrite, Category = "Channel")
+	UPROPERTY(Config, EditAnywhere, BlueprintReadOnly, Category = "Channel")
 	bool bDefaultAnimated = false;
 
-	UPROPERTY(Config, EditAnywhere, BlueprintReadWrite, Category = "Channel")
+	UPROPERTY(Config, EditAnywhere, BlueprintReadOnly, Category = "Channel")
 	bool bDefaultTwoSided = true;
 
 	bool IsPropertyEnabled(EDMMaterialPropertyType InProperty) const;
@@ -149,6 +149,14 @@ enum class EDMMaterialEditorLayout : uint8
 	Last = Left
 };
 
+UENUM(BlueprintType)
+enum class EDMLiveEditMode : uint8
+{
+	Disabled,
+	LiveEditOff,
+	LiveEditOn
+};
+
 USTRUCT()
 struct FDMContentBrowserThumbnailSettings
 {
@@ -175,6 +183,8 @@ class UDynamicMaterialEditorSettings : public UDeveloperSettings
 {
 	GENERATED_BODY()
 
+	friend class SDMToolBar;
+
 public:
 	UDynamicMaterialEditorSettings();
 	virtual ~UDynamicMaterialEditorSettings() override {}
@@ -182,85 +192,125 @@ public:
 	static UDynamicMaterialEditorSettings* Get();
 
 	/** Changes the currently active material in the designer following actor/object selection. */
-	UPROPERTY(Config, EditAnywhere, BlueprintReadWrite, Category = "Editor")
+	UPROPERTY(Config, EditAnywhere, BlueprintReadOnly, Category = "Editor")
 	bool bFollowSelection;
 
 	/**
 	 * If true, a "Create with Material Designer" button will show on properties and material lists. The edit
 	 * buttons will always appear.
 	 */
-	UPROPERTY(Config, EditAnywhere, BlueprintReadWrite, Category = "Editor")
+	UPROPERTY(Config, EditAnywhere, BlueprintReadOnly, Category = "Editor")
 	bool bAddDetailsPanelButton;
 
 	/**
 	 * When properties on the preview material are changed, if they are parameter-based then automatically
 	 * also set those values on the source material (in the world or asset).
 	 */
-	UPROPERTY(Config, EditAnywhere, BlueprintReadWrite, Category = "Editor")
+	UPROPERTY(Config, EditAnywhere, BlueprintReadOnly, Category = "Editor")
 	bool bAutomaticallyCopyParametersToSourceMaterial;
 
+protected:
 	/**
 	 * When a structural material change is made, whether to automatically recompile the preview material to show this change.
 	 */
-	UPROPERTY(Config, EditAnywhere, BlueprintReadWrite, Category = "Editor")
+	UPROPERTY(Config, EditAnywhere, Category = "Editor", meta = (EditCondition = "LiveEditMode == EDMLiveEditMode::Disabled"))
 	bool bAutomaticallyCompilePreviewMaterial;
 
 	/**
 	 * When the preview material is recompiled, whether to automatically also apply those changes to the source material.
 	 */
-	UPROPERTY(Config, EditAnywhere, BlueprintReadWrite, Category = "Editor")
+	UPROPERTY(Config, EditAnywhere, Category = "Editor", meta = (EditCondition = "LiveEditMode == EDMLiveEditMode::Disabled"))
 	bool bAutomaticallyApplyToSourceOnPreviewCompile;
 
-	UPROPERTY(Config, EditAnywhere, BlueprintReadWrite, Category = "Layout")
+	/**
+	 * When enabled, reduces compile options to a single mode called Live Edit.
+	 */
+	UPROPERTY(Config, EditAnywhere, BlueprintReadOnly, Category = "Editor")
+	EDMLiveEditMode LiveEditMode;
+
+public:
+	UFUNCTION(BlueprintPure, Category = "Material Designer")
+	bool ShouldAutomaticallyCompilePreviewMaterial() const
+	{
+		switch (LiveEditMode)
+		{
+			case EDMLiveEditMode::LiveEditOff:
+				return false;
+
+			case EDMLiveEditMode::LiveEditOn:
+				return true;
+
+			default:
+				return bAutomaticallyCompilePreviewMaterial;
+		}
+	}
+
+	UFUNCTION(BlueprintPure, Category = "Material Designer")
+	bool ShouldAutomaticallyApplyToSourceOnPreviewCompile() const
+	{
+		switch (LiveEditMode)
+		{
+			case EDMLiveEditMode::LiveEditOff:
+				return false;
+
+			case EDMLiveEditMode::LiveEditOn:
+				return true;
+
+			default:
+				return bAutomaticallyApplyToSourceOnPreviewCompile;
+		}
+	}
+
+	UPROPERTY(Config, EditAnywhere, BlueprintReadOnly, Category = "Layout")
 	EDMMaterialEditorLayout Layout;
 
 	/** Adjusts the vertical size of the material layer view. */
-	UPROPERTY(Config, EditAnywhere, BlueprintReadWrite, Category = "Layout", meta = (
+	UPROPERTY(Config, EditAnywhere, BlueprintReadOnly, Category = "Layout", meta = (
 		ClampMin = "0.05", UIMin = "0.05", ClampMax = "0.95", UIMax = "0.95"))
 	float SplitterLocation;
 
 	/** Adjusts the vertical size of the material layer view. */
-	UPROPERTY(Config, EditAnywhere, BlueprintReadWrite, Category = "Layout", meta = (
+	UPROPERTY(Config, EditAnywhere, BlueprintReadOnly, Category = "Layout", meta = (
 		ClampMin = "0.05", UIMin = "0.05", ClampMax = "0.95", UIMax = "0.95"))
 	float PreviewSplitterLocation;
 
-	UPROPERTY(Config, EditAnywhere, BlueprintReadWrite, Category = "Layout")
+	UPROPERTY(Config, EditAnywhere, BlueprintReadOnly, Category = "Layout")
 	bool bUVVisualizerVisible;
 
-	UPROPERTY(Config, EditAnywhere, BlueprintReadWrite, Category = "Layout")
+	UPROPERTY(Config, EditAnywhere, BlueprintReadOnly, Category = "Layout")
 	bool bUseFullChannelNamesInTopSlimLayout;
 
-	UPROPERTY(Config, EditAnywhere, BlueprintReadWrite, Category = "Layout")
+	UPROPERTY(Config, EditAnywhere, BlueprintReadOnly, Category = "Layout")
 	float StagePreviewSize;
 
-	UPROPERTY(Config, EditAnywhere, BlueprintReadWrite, Category = "Layout", meta = (DisplayName = "Channel Preview Size"))
+	UPROPERTY(Config, EditAnywhere, BlueprintReadOnly, Category = "Layout", meta = (DisplayName = "Channel Preview Size"))
 	float PropertyPreviewSize;
 
-	UPROPERTY(Config, EditAnywhere, BlueprintReadWrite, Category = "Thumnnails")
+	UPROPERTY(Config, EditAnywhere, BlueprintReadOnly, Category = "Thumnnails")
 	double ThumbnailSize;
 
-	UPROPERTY(Config, EditAnywhere, BlueprintReadWrite, Category = "Preview")
+	UPROPERTY(Config, EditAnywhere, BlueprintReadOnly, Category = "Preview")
 	EDMMaterialPreviewMesh PreviewMesh;
 
-	UPROPERTY(Config, EditAnywhere, BlueprintReadWrite, Category = "Preview")
+	UPROPERTY(Config, EditAnywhere, BlueprintReadOnly, Category = "Preview")
 	TSoftObjectPtr<UStaticMesh> CustomPreviewMesh;
 
 	UPROPERTY(Config, EditAnywhere, Category = "Preview")
 	FDMContentBrowserThumbnailSettings ContentBrowserThumbnail;
 
-	UPROPERTY(Config, EditAnywhere, BlueprintReadWrite, Category = "Preview")
+	UPROPERTY(Config, EditAnywhere, BlueprintReadOnly, Category = "Preview")
 	bool bShowPreviewBackground;
 
-	UPROPERTY(Config, EditAnywhere, BlueprintReadWrite, Category = "Preview")
+	UPROPERTY(Config, EditAnywhere, BlueprintReadOnly, Category = "Preview")
 	bool bPreviewImagesUseTextureUVs;
 
-	UPROPERTY(Config, EditAnywhere, BlueprintReadWrite, Category = "Preview")
+	UPROPERTY(Config, EditAnywhere, BlueprintReadOnly, Category = "Preview")
 	TSoftObjectPtr<UTexture> DefaultMask;
 
 	/**
 	 * Overrides the default values given to slots created in the given material property.
 	 */
-	UPROPERTY(Config, EditAnywhere, BlueprintReadWrite, Category = "Channels")
+	UPROPERTY(Config, EditAnywhere, BlueprintReadOnly, Category = "Channels")
 	TMap<EDMMaterialPropertyType, FDMDefaultMaterialPropertySlotValue> DefaultSlotValueOverrides;
 
 	/*
@@ -277,13 +327,13 @@ public:
 	 * - /Game/Some/Path/UV/Asset.Asset -> Category: UV
 	 * - /Plugin/Some/Path/Color/OtherAsset.OtherAsset -> Category: Color
 	 */
-	UPROPERTY(Config, EditAnywhere, BlueprintReadWrite, Category = "Effects")
+	UPROPERTY(Config, EditAnywhere, BlueprintReadOnly, Category = "Effects")
 	TArray<FName> CustomEffectsFolders;
 
-	UPROPERTY(Config, EditAnywhere, BlueprintReadWrite, Category = "Channels", meta = (TitleProperty = Name))
+	UPROPERTY(Config, EditAnywhere, BlueprintReadOnly, Category = "Channels", meta = (TitleProperty = Name))
 	TArray<FDMMaterialChannelListPreset> MaterialChannelPresets;
 
-	UPROPERTY(Config, EditAnywhere, BlueprintReadWrite, Category = "Material Functions")
+	UPROPERTY(Config, EditAnywhere, BlueprintReadOnly, Category = "Material Functions")
 	bool bUseLinearColorForVectors;
 
 	UPROPERTY(Config)
diff --git a/Engine/Plugins/Experimental/DynamicMaterial/Source/DynamicMaterialEditor/Private/Model/DynamicMaterialModelEditorOnlyData.cpp b/Engine/Plugins/Experimental/DynamicMaterial/Source/DynamicMaterialEditor/Private/Model/DynamicMaterialModelEditorOnlyData.cpp
--- a/Engine/Plugins/Experimental/DynamicMaterial/Source/DynamicMaterialEditor/Private/Model/DynamicMaterialModelEditorOnlyData.cpp
+++ b/Engine/Plugins/Experimental/DynamicMaterial/Source/DynamicMaterialEditor/Private/Model/DynamicMaterialModelEditorOnlyData.cpp
@@ -629,7 +629,7 @@ void UDynamicMaterialModelEditorOnlyData::RequestMaterialBuild(EDMBuildRequestTy
 				break;
 
 			case EDMBuildRequestType::Preview:
-				if (Settings && !Settings->bAutomaticallyCompilePreviewMaterial)
+				if (Settings && !Settings->ShouldAutomaticallyCompilePreviewMaterial())
 				{
 					bBuildRequested = true;
 				}
diff --git a/Engine/Plugins/Experimental/DynamicMaterial/Source/DynamicMaterialEditor/Private/UI/Widgets/Editor/SDMToolBar.cpp b/Engine/Plugins/Experimental/DynamicMaterial/Source/DynamicMaterialEditor/Private/UI/Widgets/Editor/SDMToolBar.cpp
--- a/Engine/Plugins/Experimental/DynamicMaterial/Source/DynamicMaterialEditor/Private/UI/Widgets/Editor/SDMToolBar.cpp
+++ b/Engine/Plugins/Experimental/DynamicMaterial/Source/DynamicMaterialEditor/Private/UI/Widgets/Editor/SDMToolBar.cpp
@@ -282,7 +282,7 @@ TSharedRef<SWidget> SDMToolBar::CreateToolBarEntries()
 		[
 			SNew(SButton)
 			.ContentPadding(LargeIconToolBarButtonContentPadding)
-			.Visibility(IsDynamicMaterialModel() ? EVisibility::Collapsed : EVisibility::Visible)
+			.Visibility(this, &SDMToolBar::GetAutoApplyVisibility)
 			.ButtonStyle(FDynamicMaterialEditorStyle::Get(), "HoverHintOnly")
 			.ToolTipText(LOCTEXT("MaterialDesignerAutoApplyTooltip", "Auto Apply to Source\n\nWhen enabled, and when the preview material is recompiled, it will automatically apply the changes to the source asset.\n\nToggling on Auto Apply to Source will cause any compiled changes to apply to the source immediately."))
 			.OnClicked(this, &SDMToolBar::OnAutoApplyClicked)
@@ -302,7 +302,7 @@ TSharedRef<SWidget> SDMToolBar::CreateToolBarEntries()
 		[
 			SNew(SButton)
 			.ContentPadding(LargeIconToolBarButtonContentPadding)
-			.Visibility(IsDynamicMaterialModel() ? EVisibility::Collapsed : EVisibility::Visible)
+			.Visibility(this, &SDMToolBar::GetAutoCompileVisibility)
 			.ButtonStyle(FDynamicMaterialEditorStyle::Get(), "HoverHintOnly")
 			.ToolTipText(LOCTEXT("MaterialDesignerAutoCompileTooltip", "Live Preview\n\nWhen enabled the preview material will be compiled for every structural change.\n\nToggling on Live Preview will cause any uncompiled changes to compile immediately."))
 			.OnClicked(this, &SDMToolBar::OnAutoCompileClicked)
@@ -313,6 +313,26 @@ TSharedRef<SWidget> SDMToolBar::CreateToolBarEntries()
 				.DesiredSizeOverride(LargeIconToolBarButtonSize)
 			]
 		]
+
+		+ SHorizontalBox::Slot()
+		.AutoWidth()
+		.HAlign(HAlign_Right)
+		.VAlign(VAlign_Top)
+		.Padding(5.0f, 0.0f, 0.0f, 0.0f)
+		[
+			SNew(SButton)
+			.ContentPadding(LargeIconToolBarButtonContentPadding)
+			.Visibility(this, &SDMToolBar::GetLiveEditVisibility)
+			.ButtonStyle(FDynamicMaterialEditorStyle::Get(), "HoverHintOnly")
+			.ToolTipText(LOCTEXT("MaterialDesignerCombinedAutoApplyTooltip", "Live Edit\n\nWhen enabled the preview and source material will recompile for every structural change made to the material (such as adding new layers or changing layer types).\n\nWhen disabled the material is freely editable and will not trigger any material compiles until this option is enabled."))
+			.OnClicked(this, &SDMToolBar::OnLiveEditClicked)
+			[
+				SNew(SImage)
+				.Image(FAppStyle::GetBrush("MaterialEditor.Apply"))
+				.ColorAndOpacity(this, &SDMToolBar::GetLiveEditColor)
+				.DesiredSizeOverride(LargeIconToolBarButtonSize)
+			]
+		]
 		
 		+ SHorizontalBox::Slot()
 		.AutoWidth()
@@ -578,6 +598,16 @@ bool SDMToolBar::IsDynamicMaterialModel() const
 	return false;
 }
 
+bool SDMToolBar::IsAutoCompileAndApplyOverridden() const
+{
+	if (UDynamicMaterialEditorSettings* Settings = UDynamicMaterialEditorSettings::Get())
+	{
+		return Settings->LiveEditMode != EDMLiveEditMode::Disabled;
+	}
+
+	return true;
+}
+
 UDynamicMaterialModelBase* SDMToolBar::GetOriginalMaterialModelBase() const
 {
 	if (TSharedPtr<SDMMaterialEditor> EditorWidget = GetEditorWidget())
@@ -1063,6 +1093,13 @@ TSharedRef<SWidget> SDMToolBar::GenerateSettingsMenu()
 	return FDMToolBarMenus::MakeEditorLayoutMenu(GetEditorWidget());
 }
 
+EVisibility SDMToolBar::GetAutoCompileVisibility() const
+{
+	return IsDynamicMaterialModel() || IsAutoCompileAndApplyOverridden()
+		? EVisibility::Collapsed
+		: EVisibility::Visible;
+}
+
 FSlateColor SDMToolBar::GetAutoCompileColor() const
 {
 	static FSlateColor EnabledColor = FSlateColor(EStyleColor::AccentBlue);
@@ -1109,6 +1146,13 @@ FReply SDMToolBar::OnAutoCompileClicked()
 	return FReply::Handled();
 }
 
+EVisibility SDMToolBar::GetAutoApplyVisibility() const
+{
+	return IsDynamicMaterialModel() || IsAutoCompileAndApplyOverridden()
+		? EVisibility::Collapsed
+		: EVisibility::Visible;
+}
+
 FSlateColor SDMToolBar::GetAutoApplyColor() const
 {
 	static FSlateColor EnabledColor = FSlateColor(EStyleColor::AccentBlue);
@@ -1146,4 +1190,59 @@ FReply SDMToolBar::OnAutoApplyClicked()
 	return FReply::Handled();
 }
 
+EVisibility SDMToolBar::GetLiveEditVisibility() const
+{
+	return IsDynamicMaterialModel() || !IsAutoCompileAndApplyOverridden()
+		? EVisibility::Collapsed
+		: EVisibility::Visible;
+}
+
+FSlateColor SDMToolBar::GetLiveEditColor() const
+{
+	static FSlateColor EnabledColor = FSlateColor(EStyleColor::AccentBlue);
+	static FSlateColor DisabledColor = FSlateColor(EStyleColor::AccentGray);
+
+	if (UDynamicMaterialEditorSettings* Settings = UDynamicMaterialEditorSettings::Get())
+	{
+		if (Settings->LiveEditMode == EDMLiveEditMode::LiveEditOn)
+		{
+			return EnabledColor;
+		}
+	}
+
+	return DisabledColor;
+}
+
+FReply SDMToolBar::OnLiveEditClicked()
+{
+	if (UDynamicMaterialEditorSettings* Settings = UDynamicMaterialEditorSettings::Get())
+	{
+		FScopedTransaction Transaction(LOCTEXT("ToggleLiveEdit", "Toggle Live Edit"));
+		Settings->Modify();
+
+		switch (Settings->LiveEditMode)
+		{
+			case EDMLiveEditMode::LiveEditOn:
+				Settings->LiveEditMode = EDMLiveEditMode::LiveEditOff;
+				break;
+
+			case EDMLiveEditMode::LiveEditOff:
+				Settings->LiveEditMode = EDMLiveEditMode::LiveEditOn;
+				break;
+		}
+
+		Settings->SaveConfig();
+
+		if (Settings->ShouldAutomaticallyApplyToSourceOnPreviewCompile())
+		{
+			if (TSharedPtr<SDMMaterialEditor> EditorWidget = EditorWidgetWeak.Pin())
+			{
+				EditorWidget->ApplyToOriginal();
+			}
+		}
+	}
+
+	return FReply::Handled();
+}
+
 #undef LOCTEXT_NAMESPACE
diff --git a/Engine/Plugins/Experimental/DynamicMaterial/Source/DynamicMaterialEditor/Private/UI/Widgets/Editor/SDMToolBar.h b/Engine/Plugins/Experimental/DynamicMaterial/Source/DynamicMaterialEditor/Private/UI/Widgets/Editor/SDMToolBar.h
--- a/Engine/Plugins/Experimental/DynamicMaterial/Source/DynamicMaterialEditor/Private/UI/Widgets/Editor/SDMToolBar.h
+++ b/Engine/Plugins/Experimental/DynamicMaterial/Source/DynamicMaterialEditor/Private/UI/Widgets/Editor/SDMToolBar.h
@@ -66,6 +66,8 @@ protected:
 
 	bool IsDynamicMaterialModel() const;
 
+	bool IsAutoCompileAndApplyOverridden() const;
+
 	UDynamicMaterialModelBase* GetOriginalMaterialModelBase() const;
 
 	TSharedRef<SWidget> CreateToolBarEntries();
@@ -109,11 +111,15 @@ protected:
 
 	TSharedRef<SWidget> GenerateSettingsMenu();
 
+	EVisibility GetAutoCompileVisibility() const;
 	FSlateColor GetAutoCompileColor() const;
-
 	FReply OnAutoCompileClicked();
 
+	EVisibility GetAutoApplyVisibility() const;
 	FSlateColor GetAutoApplyColor() const;
-
 	FReply OnAutoApplyClicked();
+
+	EVisibility GetLiveEditVisibility() const;
+	FSlateColor GetLiveEditColor() const;
+	FReply OnLiveEditClicked();
 };
diff --git a/Engine/Plugins/Experimental/DynamicMaterial/Source/DynamicMaterialEditor/Private/UI/Widgets/SDMMaterialEditor.cpp b/Engine/Plugins/Experimental/DynamicMaterial/Source/DynamicMaterialEditor/Private/UI/Widgets/SDMMaterialEditor.cpp
--- a/Engine/Plugins/Experimental/DynamicMaterial/Source/DynamicMaterialEditor/Private/UI/Widgets/SDMMaterialEditor.cpp
+++ b/Engine/Plugins/Experimental/DynamicMaterial/Source/DynamicMaterialEditor/Private/UI/Widgets/SDMMaterialEditor.cpp
@@ -126,7 +126,7 @@ SDMMaterialEditor::SDMMaterialEditor()
 {
 	if (UDynamicMaterialEditorSettings* Settings = UDynamicMaterialEditorSettings::Get())
 	{
-		if (Settings->bAutomaticallyApplyToSourceOnPreviewCompile)
+		if (Settings->ShouldAutomaticallyApplyToSourceOnPreviewCompile())
 		{
 			// If we are automatically applying, we need to skip the initial compile event.
 			bSkipApplyOnCompile = true;
@@ -1491,7 +1491,7 @@ void SDMMaterialEditor::OnMaterialBuilt(UDynamicMaterialModelBase* InMaterialMod
 
 	if (UDynamicMaterialEditorSettings* Settings = UDynamicMaterialEditorSettings::Get())
 	{
-		if (Settings->bAutomaticallyApplyToSourceOnPreviewCompile)
+		if (Settings->ShouldAutomaticallyApplyToSourceOnPreviewCompile())
 		{
 			ApplyToOriginal();
 		}
