Material Designer: Opacity pre-blend update
- The Add, Subtract, Multiply and Divide blend modes now use those operations on their base layers to create the opacity (mask) base layer input.
- Updated display name on Texture Clamp to indicate it affects UVs.

diff --git a/Engine/Plugins/Experimental/DynamicMaterial/Source/DynamicMaterialEditor/Private/Components/DMMaterialProperty.cpp b/Engine/Plugins/Experimental/DynamicMaterial/Source/DynamicMaterialEditor/Private/Components/DMMaterialProperty.cpp
--- a/Engine/Plugins/Experimental/DynamicMaterial/Source/DynamicMaterialEditor/Private/Components/DMMaterialProperty.cpp
+++ b/Engine/Plugins/Experimental/DynamicMaterial/Source/DynamicMaterialEditor/Private/Components/DMMaterialProperty.cpp
@@ -487,40 +487,25 @@ void UDMMaterialProperty::GenerateOpacityExpressions(const TSharedRef<FDMMateria
 			continue;
 		}
 
-		UMaterialExpressionMax* Max = InBuildState->GetBuildUtils().CreateExpression<UMaterialExpressionMax>(UE_DM_NodeComment_Default);
-		check(Max);
+		TArray<UMaterialExpression*> BaseExpressions;
+		int32 BlendedOutputIndex;
+		int32 BlendedOutputChannel;
 
-		Max->A.Expression = OutExpression;
-		Max->A.OutputIndex = OutOutputIndex;
-		Max->A.Mask = 0;
-
-		if (OutOutputChannel != FDMMaterialStageConnectorChannel::WHOLE_CHANNEL)
+		if (UDMMaterialStageBlend* BaseBlend = Cast<UDMMaterialStageBlend>(BaseStage->GetSource()))
 		{
-			Max->A.Mask = 1;
-			Max->A.MaskR = !!(OutOutputChannel & FDMMaterialStageConnectorChannel::FIRST_CHANNEL);
-			Max->A.MaskG = !!(OutOutputChannel & FDMMaterialStageConnectorChannel::SECOND_CHANNEL);
-			Max->A.MaskB = !!(OutOutputChannel & FDMMaterialStageConnectorChannel::THIRD_CHANNEL);
-			Max->A.MaskA = !!(OutOutputChannel & FDMMaterialStageConnectorChannel::FOURTH_CHANNEL);
+			BaseBlend->BlendOpacityLayer(InBuildState, OutExpression, OutOutputIndex, OutOutputChannel, MaskOutputExpression,
+				MaskOutputIndex, MaskOutputChannel, BaseExpressions, BlendedOutputIndex, BlendedOutputChannel);
 		}
-
-		Max->B.Expression = MaskOutputExpression;
-		Max->B.OutputIndex = MaskOutputIndex;
-		Max->B.Mask = 0;
-
-		if (MaskOutputChannel != FDMMaterialStageConnectorChannel::WHOLE_CHANNEL)
+		else
 		{
-			Max->B.Mask = 1;
-			Max->B.MaskR = !!(MaskOutputChannel & FDMMaterialStageConnectorChannel::FIRST_CHANNEL);
-			Max->B.MaskG = !!(MaskOutputChannel & FDMMaterialStageConnectorChannel::SECOND_CHANNEL);
-			Max->B.MaskB = !!(MaskOutputChannel & FDMMaterialStageConnectorChannel::THIRD_CHANNEL);
-			Max->B.MaskA = !!(MaskOutputChannel & FDMMaterialStageConnectorChannel::FOURTH_CHANNEL);
+			UDMMaterialStageBlend::CreateBlendOpacityLayer<UMaterialExpressionMax>(InBuildState, OutExpression, OutOutputIndex, 
+				OutOutputChannel, MaskOutputExpression, MaskOutputIndex, MaskOutputChannel, BaseExpressions, BlendedOutputIndex,
+				BlendedOutputChannel);
 		}
 
-		OutExpression = Max;
-
-		// If we have to combine, it will use the Max node's output info
-		OutOutputIndex = 0;
-		OutOutputChannel = FDMMaterialStageConnectorChannel::WHOLE_CHANNEL;
+		OutExpression = BaseExpressions.Last();
+		OutOutputIndex = BlendedOutputIndex;
+		OutOutputChannel = BlendedOutputChannel;
 	}
 }
 
diff --git a/Engine/Plugins/Experimental/DynamicMaterial/Source/DynamicMaterialEditor/Private/Components/DMMaterialStageBlend.cpp b/Engine/Plugins/Experimental/DynamicMaterial/Source/DynamicMaterialEditor/Private/Components/DMMaterialStageBlend.cpp
--- a/Engine/Plugins/Experimental/DynamicMaterial/Source/DynamicMaterialEditor/Private/Components/DMMaterialStageBlend.cpp
+++ b/Engine/Plugins/Experimental/DynamicMaterial/Source/DynamicMaterialEditor/Private/Components/DMMaterialStageBlend.cpp
@@ -12,6 +12,7 @@
 #include "Components/MaterialValues/DMMaterialValueFloat1.h"
 #include "DMEDefs.h"
 #include "Materials/MaterialExpressionComponentMask.h"
+#include "Materials/MaterialExpressionMax.h"
 #include "Model/DMMaterialBuildState.h"
 #include "Model/DMMaterialBuildUtils.h"
 #include "Model/DynamicMaterialModel.h"
@@ -85,6 +86,15 @@ const FText& UDMMaterialStageBlend::GetBlendDescription() const
 	return BlendDescription;
 }
 
+void UDMMaterialStageBlend::BlendOpacityLayer(const TSharedRef<FDMMaterialBuildState>& InBuildState,
+	UMaterialExpression* InBaseLayerOpacityExpression, int32 InBaseOutputIndex, int32 InBaseOutputChannel,
+	UMaterialExpression* InMyLayerOpacityExpression, int32 InMyOutputIndex, int32 InMyOutputChannel,
+	TArray<UMaterialExpression*>& OutAddedExpressions, int32& OutOutputIndex, int32& OutOutputChannel) const
+{
+	CreateBlendOpacityLayer<UMaterialExpressionMax>(InBuildState, InBaseLayerOpacityExpression, InBaseOutputIndex, InBaseOutputChannel,
+		InMyLayerOpacityExpression, InMyOutputIndex, InMyOutputChannel, OutAddedExpressions, OutOutputIndex, OutOutputChannel);
+}
+
 int32 UDMMaterialStageBlend::GetDefaultBaseChannelOverrideOutputIndex() const
 {
 	UDMMaterialStageInput* StageInputB = GetInputB();
@@ -294,7 +304,7 @@ void UDMMaterialStageBlend::GenerateBlendList()
 		Blends.Add(SourceClass);
 	}
 }
- 
+
 bool UDMMaterialStageBlend::CanInputAcceptType(int32 InputIndex, EDMValueType ValueType) const
 {
 	check(InputConnectors.IsValidIndex(InputIndex));
diff --git a/Engine/Plugins/Experimental/DynamicMaterial/Source/DynamicMaterialEditor/Private/Components/MaterialStageBlends/DMMSBAdd.cpp b/Engine/Plugins/Experimental/DynamicMaterial/Source/DynamicMaterialEditor/Private/Components/MaterialStageBlends/DMMSBAdd.cpp
--- a/Engine/Plugins/Experimental/DynamicMaterial/Source/DynamicMaterialEditor/Private/Components/MaterialStageBlends/DMMSBAdd.cpp
+++ b/Engine/Plugins/Experimental/DynamicMaterial/Source/DynamicMaterialEditor/Private/Components/MaterialStageBlends/DMMSBAdd.cpp
@@ -2,6 +2,8 @@
 
 #include "Components/MaterialStageBlends/DMMSBAdd.h"
 
+#include "Materials/MaterialExpressionAdd.h"
+
 #define LOCTEXT_NAMESPACE "DMMaterialStageBlendAdd"
 
 UDMMaterialStageBlendAdd::UDMMaterialStageBlendAdd()
@@ -14,4 +16,13 @@ UDMMaterialStageBlendAdd::UDMMaterialStageBlendAdd()
 {
 }
 
+void UDMMaterialStageBlendAdd::BlendOpacityLayer(const TSharedRef<FDMMaterialBuildState>& InBuildState,
+	UMaterialExpression* InBaseLayerOpacityExpression, int32 InBaseOutputIndex, int32 InBaseOutputChannel,
+	UMaterialExpression* InMyLayerOpacityExpression, int32 InMyOutputIndex, int32 InMyOutputChannel,
+	TArray<UMaterialExpression*>& OutAddedExpressions, int32& OutOutputIndex, int32& OutOutputChannel) const
+{
+	CreateBlendOpacityLayer<UMaterialExpressionAdd>(InBuildState, InBaseLayerOpacityExpression, InBaseOutputIndex, InBaseOutputChannel,
+		InMyLayerOpacityExpression, InMyOutputIndex, InMyOutputChannel, OutAddedExpressions, OutOutputIndex, OutOutputChannel);
+}
+
 #undef LOCTEXT_NAMESPACE
diff --git a/Engine/Plugins/Experimental/DynamicMaterial/Source/DynamicMaterialEditor/Private/Components/MaterialStageBlends/DMMSBDivide.cpp b/Engine/Plugins/Experimental/DynamicMaterial/Source/DynamicMaterialEditor/Private/Components/MaterialStageBlends/DMMSBDivide.cpp
--- a/Engine/Plugins/Experimental/DynamicMaterial/Source/DynamicMaterialEditor/Private/Components/MaterialStageBlends/DMMSBDivide.cpp
+++ b/Engine/Plugins/Experimental/DynamicMaterial/Source/DynamicMaterialEditor/Private/Components/MaterialStageBlends/DMMSBDivide.cpp
@@ -2,6 +2,8 @@
 
 #include "Components/MaterialStageBlends/DMMSBDivide.h"
 
+#include "Materials/MaterialExpressionDivide.h"
+
 #define LOCTEXT_NAMESPACE "DMMaterialStageBlendDivide"
 
 UDMMaterialStageBlendDivide::UDMMaterialStageBlendDivide()
@@ -14,4 +16,13 @@ UDMMaterialStageBlendDivide::UDMMaterialStageBlendDivide()
 {
 }
 
+void UDMMaterialStageBlendDivide::BlendOpacityLayer(const TSharedRef<FDMMaterialBuildState>& InBuildState,
+	UMaterialExpression* InBaseLayerOpacityExpression, int32 InBaseOutputIndex, int32 InBaseOutputChannel,
+	UMaterialExpression* InMyLayerOpacityExpression, int32 InMyOutputIndex, int32 InMyOutputChannel,
+	TArray<UMaterialExpression*>& OutAddedExpressions, int32& OutOutputIndex, int32& OutOutputChannel) const
+{
+	CreateBlendOpacityLayer<UMaterialExpressionDivide>(InBuildState, InBaseLayerOpacityExpression, InBaseOutputIndex, InBaseOutputChannel,
+		InMyLayerOpacityExpression, InMyOutputIndex, InMyOutputChannel, OutAddedExpressions, OutOutputIndex, OutOutputChannel);
+}
+
 #undef LOCTEXT_NAMESPACE
diff --git a/Engine/Plugins/Experimental/DynamicMaterial/Source/DynamicMaterialEditor/Private/Components/MaterialStageBlends/DMMSBMultiply.cpp b/Engine/Plugins/Experimental/DynamicMaterial/Source/DynamicMaterialEditor/Private/Components/MaterialStageBlends/DMMSBMultiply.cpp
--- a/Engine/Plugins/Experimental/DynamicMaterial/Source/DynamicMaterialEditor/Private/Components/MaterialStageBlends/DMMSBMultiply.cpp
+++ b/Engine/Plugins/Experimental/DynamicMaterial/Source/DynamicMaterialEditor/Private/Components/MaterialStageBlends/DMMSBMultiply.cpp
@@ -2,6 +2,8 @@
 
 #include "Components/MaterialStageBlends/DMMSBMultiply.h"
 
+#include "Materials/MaterialExpressionMultiply.h"
+
 #define LOCTEXT_NAMESPACE "DMMaterialStageBlendMultiply"
 
 UDMMaterialStageBlendMultiply::UDMMaterialStageBlendMultiply()
@@ -14,4 +16,13 @@ UDMMaterialStageBlendMultiply::UDMMaterialStageBlendMultiply()
 {
 }
 
+void UDMMaterialStageBlendMultiply::BlendOpacityLayer(const TSharedRef<FDMMaterialBuildState>& InBuildState,
+	UMaterialExpression* InBaseLayerOpacityExpression, int32 InBaseOutputIndex, int32 InBaseOutputChannel,
+	UMaterialExpression* InMyLayerOpacityExpression, int32 InMyOutputIndex, int32 InMyOutputChannel,
+	TArray<UMaterialExpression*>& OutAddedExpressions, int32& OutOutputIndex, int32& OutOutputChannel) const
+{
+	CreateBlendOpacityLayer<UMaterialExpressionMultiply>(InBuildState, InBaseLayerOpacityExpression, InBaseOutputIndex, InBaseOutputChannel,
+		InMyLayerOpacityExpression, InMyOutputIndex, InMyOutputChannel, OutAddedExpressions, OutOutputIndex, OutOutputChannel);
+}
+
 #undef LOCTEXT_NAMESPACE
diff --git a/Engine/Plugins/Experimental/DynamicMaterial/Source/DynamicMaterialEditor/Private/Components/MaterialStageBlends/DMMSBSubtract.cpp b/Engine/Plugins/Experimental/DynamicMaterial/Source/DynamicMaterialEditor/Private/Components/MaterialStageBlends/DMMSBSubtract.cpp
--- a/Engine/Plugins/Experimental/DynamicMaterial/Source/DynamicMaterialEditor/Private/Components/MaterialStageBlends/DMMSBSubtract.cpp
+++ b/Engine/Plugins/Experimental/DynamicMaterial/Source/DynamicMaterialEditor/Private/Components/MaterialStageBlends/DMMSBSubtract.cpp
@@ -2,6 +2,8 @@
 
 #include "Components/MaterialStageBlends/DMMSBSubtract.h"
 
+#include "Materials/MaterialExpressionSubtract.h"
+
 #define LOCTEXT_NAMESPACE "DMMaterialStageBlendSubtract"
 
 UDMMaterialStageBlendSubtract::UDMMaterialStageBlendSubtract()
@@ -14,4 +16,13 @@ UDMMaterialStageBlendSubtract::UDMMaterialStageBlendSubtract()
 {
 }
 
+void UDMMaterialStageBlendSubtract::BlendOpacityLayer(const TSharedRef<FDMMaterialBuildState>& InBuildState,
+	UMaterialExpression* InBaseLayerOpacityExpression, int32 InBaseOutputIndex, int32 InBaseOutputChannel,
+	UMaterialExpression* InMyLayerOpacityExpression, int32 InMyOutputIndex, int32 InMyOutputChannel,
+	TArray<UMaterialExpression*>& OutAddedExpressions, int32& OutOutputIndex, int32& OutOutputChannel) const
+{
+	CreateBlendOpacityLayer<UMaterialExpressionSubtract>(InBuildState, InBaseLayerOpacityExpression, InBaseOutputIndex, InBaseOutputChannel,
+		InMyLayerOpacityExpression, InMyOutputIndex, InMyOutputChannel, OutAddedExpressions, OutOutputIndex, OutOutputChannel);
+}
+
 #undef LOCTEXT_NAMESPACE
diff --git a/Engine/Plugins/Experimental/DynamicMaterial/Source/DynamicMaterialEditor/Public/Components/DMMaterialStageBlend.h b/Engine/Plugins/Experimental/DynamicMaterial/Source/DynamicMaterialEditor/Public/Components/DMMaterialStageBlend.h
--- a/Engine/Plugins/Experimental/DynamicMaterial/Source/DynamicMaterialEditor/Public/Components/DMMaterialStageBlend.h
+++ b/Engine/Plugins/Experimental/DynamicMaterial/Source/DynamicMaterialEditor/Public/Components/DMMaterialStageBlend.h
@@ -3,7 +3,10 @@
 #pragma once
 
 #include "DMMaterialStageThroughput.h"
+#include "Model/DMMaterialBuildState.h"
+#include "Templates/SharedPointer.h"
 #include "UObject/StrongObjectPtr.h"
+#include "Utils/DMUtils.h"
 #include "DMMaterialStageBlend.generated.h"
 
 class FMenuBuilder;
@@ -45,6 +48,20 @@ public:
 	UFUNCTION(BlueprintPure, Category = "Material Designer")
 	const FText& GetBlendDescription() const;
 
+	template<class InNodeType>
+	static void CreateBlendOpacityLayer(const TSharedRef<FDMMaterialBuildState>& InBuildState,
+		UMaterialExpression* InBaseLayerOpacityExpression, int32 InBaseOutputIndex, int32 InBaseOutputChannel,
+		UMaterialExpression* InMyLayerOpacityExpression, int32 InMyOutputIndex, int32 InMyOutputChannel,
+		TArray<UMaterialExpression*>& OutAddedExpressions, int32& OutOutputIndex, int32& OutOutputChannel);
+
+	/**
+	 * When building the opacity (mask) base input, combines the previous layer's alpha output with this ones.
+	 */
+	virtual void BlendOpacityLayer(const TSharedRef<FDMMaterialBuildState>& InBuildState, 
+		UMaterialExpression* InBaseLayerOpacityExpression, int32 InBaseOutputIndex, int32 InBaseOutputChannel,
+		UMaterialExpression* InMyLayerOpacityExpression, int32 InMyOutputIndex, int32 InMyOutputChannel,
+		TArray<UMaterialExpression*>& OutAddedExpressions, int32& OutOutputIndex, int32& OutOutputChannel) const;
+
 	//~ Begin UDMMaterialStageThroughput
 	DYNAMICMATERIALEDITOR_API virtual bool CanInputAcceptType(int32 InputIndex, EDMValueType ValueType) const override;
 	DYNAMICMATERIALEDITOR_API virtual void AddDefaultInput(int32 InInputIndex) const override;
@@ -114,3 +131,45 @@ protected:
 	DYNAMICMATERIALEDITOR_API virtual void GeneratePreviewMaterial(UMaterial* InPreviewMaterial) override;
 	//~ End UDMMaterialStageThroughput
 };
+
+template<class InNodeType>
+void UDMMaterialStageBlend::CreateBlendOpacityLayer(const TSharedRef<FDMMaterialBuildState>& InBuildState,
+	UMaterialExpression* InBaseLayerOpacityExpression, int32 InBaseOutputIndex, int32 InBaseOutputChannel,
+	UMaterialExpression* InMyLayerOpacityExpression, int32 InMyOutputIndex, int32 InMyOutputChannel,
+	TArray<UMaterialExpression*>& OutAddedExpressions, int32& OutOutputIndex, int32& OutOutputChannel)
+{
+	InNodeType* Operator = InBuildState->GetBuildUtils().CreateExpression<InNodeType>(UE_DM_NodeComment_Default);
+	check(Operator);
+
+	Operator->A.Expression = InBaseLayerOpacityExpression;
+	Operator->A.OutputIndex = InBaseOutputIndex;
+	Operator->A.Mask = 0;
+
+	if (InBaseOutputChannel != FDMMaterialStageConnectorChannel::WHOLE_CHANNEL)
+	{
+		Operator->A.Mask = 1;
+		Operator->A.MaskR = !!(InBaseOutputChannel & FDMMaterialStageConnectorChannel::FIRST_CHANNEL);
+		Operator->A.MaskG = !!(InBaseOutputChannel & FDMMaterialStageConnectorChannel::SECOND_CHANNEL);
+		Operator->A.MaskB = !!(InBaseOutputChannel & FDMMaterialStageConnectorChannel::THIRD_CHANNEL);
+		Operator->A.MaskA = !!(InBaseOutputChannel & FDMMaterialStageConnectorChannel::FOURTH_CHANNEL);
+	}
+
+	Operator->B.Expression = InMyLayerOpacityExpression;
+	Operator->B.OutputIndex = InMyOutputIndex;
+	Operator->B.Mask = 0;
+
+	if (InMyOutputChannel != FDMMaterialStageConnectorChannel::WHOLE_CHANNEL)
+	{
+		Operator->B.Mask = 1;
+		Operator->B.MaskR = !!(InMyOutputChannel & FDMMaterialStageConnectorChannel::FIRST_CHANNEL);
+		Operator->B.MaskG = !!(InMyOutputChannel & FDMMaterialStageConnectorChannel::SECOND_CHANNEL);
+		Operator->B.MaskB = !!(InMyOutputChannel & FDMMaterialStageConnectorChannel::THIRD_CHANNEL);
+		Operator->B.MaskA = !!(InMyOutputChannel & FDMMaterialStageConnectorChannel::FOURTH_CHANNEL);
+	}
+
+	OutAddedExpressions.Add(Operator);
+
+	// If we have to combine, it will use the Max node's output info
+	OutOutputIndex = 0;
+	OutOutputChannel = FDMMaterialStageConnectorChannel::WHOLE_CHANNEL;
+}
diff --git a/Engine/Plugins/Experimental/DynamicMaterial/Source/DynamicMaterialEditor/Public/Components/MaterialStageBlends/DMMSBAdd.h b/Engine/Plugins/Experimental/DynamicMaterial/Source/DynamicMaterialEditor/Public/Components/MaterialStageBlends/DMMSBAdd.h
--- a/Engine/Plugins/Experimental/DynamicMaterial/Source/DynamicMaterialEditor/Public/Components/MaterialStageBlends/DMMSBAdd.h
+++ b/Engine/Plugins/Experimental/DynamicMaterial/Source/DynamicMaterialEditor/Public/Components/MaterialStageBlends/DMMSBAdd.h
@@ -13,4 +13,11 @@ class UDMMaterialStageBlendAdd : public UDMMaterialStageBlendFunction
 
 public:
 	UDMMaterialStageBlendAdd();
+
+	//~ Begin UDMMaterialStageBlend
+	virtual void BlendOpacityLayer(const TSharedRef<FDMMaterialBuildState>& InBuildState,
+		UMaterialExpression* InBaseLayerOpacityExpression, int32 InBaseOutputIndex, int32 InBaseOutputChannel,
+		UMaterialExpression* InMyLayerOpacityExpression, int32 InMyOutputIndex, int32 InMyOutputChannel,
+		TArray<UMaterialExpression*>& OutAddedExpressions, int32& OutOutputIndex, int32& OutOutputChannel) const;
+	//~ End UDMMaterialStageBlend
 };
diff --git a/Engine/Plugins/Experimental/DynamicMaterial/Source/DynamicMaterialEditor/Public/Components/MaterialStageBlends/DMMSBDivide.h b/Engine/Plugins/Experimental/DynamicMaterial/Source/DynamicMaterialEditor/Public/Components/MaterialStageBlends/DMMSBDivide.h
--- a/Engine/Plugins/Experimental/DynamicMaterial/Source/DynamicMaterialEditor/Public/Components/MaterialStageBlends/DMMSBDivide.h
+++ b/Engine/Plugins/Experimental/DynamicMaterial/Source/DynamicMaterialEditor/Public/Components/MaterialStageBlends/DMMSBDivide.h
@@ -12,4 +12,11 @@ class UDMMaterialStageBlendDivide : public UDMMaterialStageBlendFunction
 
 public:
 	UDMMaterialStageBlendDivide();
+
+	//~ Begin UDMMaterialStageBlend
+	virtual void BlendOpacityLayer(const TSharedRef<FDMMaterialBuildState>& InBuildState,
+		UMaterialExpression* InBaseLayerOpacityExpression, int32 InBaseOutputIndex, int32 InBaseOutputChannel,
+		UMaterialExpression* InMyLayerOpacityExpression, int32 InMyOutputIndex, int32 InMyOutputChannel,
+		TArray<UMaterialExpression*>& OutAddedExpressions, int32& OutOutputIndex, int32& OutOutputChannel) const;
+	//~ End UDMMaterialStageBlend
 };
diff --git a/Engine/Plugins/Experimental/DynamicMaterial/Source/DynamicMaterialEditor/Public/Components/MaterialStageBlends/DMMSBMultiply.h b/Engine/Plugins/Experimental/DynamicMaterial/Source/DynamicMaterialEditor/Public/Components/MaterialStageBlends/DMMSBMultiply.h
--- a/Engine/Plugins/Experimental/DynamicMaterial/Source/DynamicMaterialEditor/Public/Components/MaterialStageBlends/DMMSBMultiply.h
+++ b/Engine/Plugins/Experimental/DynamicMaterial/Source/DynamicMaterialEditor/Public/Components/MaterialStageBlends/DMMSBMultiply.h
@@ -12,4 +12,11 @@ class UDMMaterialStageBlendMultiply : public UDMMaterialStageBlendFunction
 
 public:
 	UDMMaterialStageBlendMultiply();
+
+	//~ Begin UDMMaterialStageBlend
+	virtual void BlendOpacityLayer(const TSharedRef<FDMMaterialBuildState>& InBuildState,
+		UMaterialExpression* InBaseLayerOpacityExpression, int32 InBaseOutputIndex, int32 InBaseOutputChannel,
+		UMaterialExpression* InMyLayerOpacityExpression, int32 InMyOutputIndex, int32 InMyOutputChannel,
+		TArray<UMaterialExpression*>& OutAddedExpressions, int32& OutOutputIndex, int32& OutOutputChannel) const;
+	//~ End UDMMaterialStageBlend
 };
diff --git a/Engine/Plugins/Experimental/DynamicMaterial/Source/DynamicMaterialEditor/Public/Components/MaterialStageBlends/DMMSBSubtract.h b/Engine/Plugins/Experimental/DynamicMaterial/Source/DynamicMaterialEditor/Public/Components/MaterialStageBlends/DMMSBSubtract.h
--- a/Engine/Plugins/Experimental/DynamicMaterial/Source/DynamicMaterialEditor/Public/Components/MaterialStageBlends/DMMSBSubtract.h
+++ b/Engine/Plugins/Experimental/DynamicMaterial/Source/DynamicMaterialEditor/Public/Components/MaterialStageBlends/DMMSBSubtract.h
@@ -12,4 +12,11 @@ class UDMMaterialStageBlendSubtract : public UDMMaterialStageBlendFunction
 
 public:
 	UDMMaterialStageBlendSubtract();
+
+	//~ Begin UDMMaterialStageBlend
+	virtual void BlendOpacityLayer(const TSharedRef<FDMMaterialBuildState>& InBuildState,
+		UMaterialExpression* InBaseLayerOpacityExpression, int32 InBaseOutputIndex, int32 InBaseOutputChannel,
+		UMaterialExpression* InMyLayerOpacityExpression, int32 InMyOutputIndex, int32 InMyOutputChannel,
+		TArray<UMaterialExpression*>& OutAddedExpressions, int32& OutOutputIndex, int32& OutOutputChannel) const;
+	//~ End UDMMaterialStageBlend
 };
diff --git a/Engine/Plugins/Experimental/DynamicMaterial/Source/DynamicMaterialEditor/Public/Components/MaterialStageExpressions/DMMSETextureSampleBase.h b/Engine/Plugins/Experimental/DynamicMaterial/Source/DynamicMaterialEditor/Public/Components/MaterialStageExpressions/DMMSETextureSampleBase.h
--- a/Engine/Plugins/Experimental/DynamicMaterial/Source/DynamicMaterialEditor/Public/Components/MaterialStageExpressions/DMMSETextureSampleBase.h
+++ b/Engine/Plugins/Experimental/DynamicMaterial/Source/DynamicMaterialEditor/Public/Components/MaterialStageExpressions/DMMSETextureSampleBase.h
@@ -48,7 +48,7 @@ public:
 
 protected:
 	UPROPERTY(EditInstanceOnly, BlueprintReadWrite, Setter = SetClampTextureEnabled, BlueprintSetter = SetClampTextureEnabled, Category = "Material Designer",
-		meta=(ToolTip="Forces a material rebuild.", LowPriority, NotKeyframeable, AllowPrivateAccess = "true"))
+		DisplayName = "Clamp Texture UV", meta = (ToolTip = "Forces a material rebuild.", LowPriority, NotKeyframeable, AllowPrivateAccess = "true"))
 	bool bClampTexture;
 
 	DYNAMICMATERIALEDITOR_API UDMMaterialStageExpressionTextureSampleBase(const FText& InName, TSubclassOf<UMaterialExpression> InClass);
