Custom Details View:
- Made custom widgets use a fill horizotnal alignment.

Material Designer:
- Fixed object properties to display full row widgets when no display name is provided.

diff --git a/Engine/Plugins/Experimental/CustomDetailsView/Source/CustomDetailsView/Private/Items/CustomDetailsViewItemBase.cpp b/Engine/Plugins/Experimental/CustomDetailsView/Source/CustomDetailsView/Private/Items/CustomDetailsViewItemBase.cpp
--- a/Engine/Plugins/Experimental/CustomDetailsView/Source/CustomDetailsView/Private/Items/CustomDetailsViewItemBase.cpp
+++ b/Engine/Plugins/Experimental/CustomDetailsView/Source/CustomDetailsView/Private/Items/CustomDetailsViewItemBase.cpp
@@ -578,6 +578,7 @@ void FCustomDetailsViewItemBase::AddValueWidget(const TSharedRef<SSplitter>& InS
 	if (TSharedPtr<SWidget> OverrideWidget = GetOverrideWidget(ECustomDetailsViewWidgetType::Value))
 	{
 		ValueWidgetInner = OverrideWidget;
+		DetailWidgetRow.ValueWidget.HorizontalAlignment = EHorizontalAlignment::HAlign_Fill;
 	}
 
 	if (!ValueWidgetInner.IsValid())
diff --git a/Engine/Plugins/Experimental/DynamicMaterial/Source/DynamicMaterialEditor/Private/UI/Widgets/Editor/SDMObjectEditorWidgetBase.cpp b/Engine/Plugins/Experimental/DynamicMaterial/Source/DynamicMaterialEditor/Private/UI/Widgets/Editor/SDMObjectEditorWidgetBase.cpp
--- a/Engine/Plugins/Experimental/DynamicMaterial/Source/DynamicMaterialEditor/Private/UI/Widgets/Editor/SDMObjectEditorWidgetBase.cpp
+++ b/Engine/Plugins/Experimental/DynamicMaterial/Source/DynamicMaterialEditor/Private/UI/Widgets/Editor/SDMObjectEditorWidgetBase.cpp
@@ -267,7 +267,7 @@ void SDMObjectEditorWidgetBase::AddCustomRow(const TSharedRef<ICustomDetailsView
 	const FDMPropertyHandle& InPropertyRow)
 {
 	const FText RowDisplayName = InPropertyRow.NameOverride.Get(FText::GetEmpty());
-	const bool bIsWholeRow = !RowDisplayName.IsEmpty();
+	const bool bIsWholeRow = RowDisplayName.IsEmpty();
 
 	TSharedPtr<ICustomDetailsViewCustomItem> Item = InDetailsView->CreateCustomItem(
 		InParent,
