Material Designer: Layer duplication fix
- When duplicating layers, cached parameter names are no longer copied.

diff --git a/Engine/Plugins/Experimental/DynamicMaterial/Source/DynamicMaterial/Public/Components/DMMaterialValue.h b/Engine/Plugins/Experimental/DynamicMaterial/Source/DynamicMaterial/Public/Components/DMMaterialValue.h
--- a/Engine/Plugins/Experimental/DynamicMaterial/Source/DynamicMaterial/Public/Components/DMMaterialValue.h
+++ b/Engine/Plugins/Experimental/DynamicMaterial/Source/DynamicMaterial/Public/Components/DMMaterialValue.h
@@ -185,7 +185,7 @@ protected:
 	UPROPERTY(VisibleInstanceOnly, BlueprintReadOnly, Instanced, Category = "Material Designer")
 	TObjectPtr<UDMMaterialParameter> Parameter;
 
-	UPROPERTY(VisibleInstanceOnly, BlueprintReadOnly, Category = "Material Designer")
+	UPROPERTY(VisibleInstanceOnly, BlueprintReadOnly, DuplicateTransient, TextExportTransient, Category = "Material Designer")
 	FName CachedParameterName;
 
 #if WITH_EDITORONLY_DATA
diff --git a/Engine/Plugins/Experimental/DynamicMaterial/Source/DynamicMaterial/Public/Components/DMTextureUV.h b/Engine/Plugins/Experimental/DynamicMaterial/Source/DynamicMaterial/Public/Components/DMTextureUV.h
--- a/Engine/Plugins/Experimental/DynamicMaterial/Source/DynamicMaterial/Public/Components/DMTextureUV.h
+++ b/Engine/Plugins/Experimental/DynamicMaterial/Source/DynamicMaterial/Public/Components/DMTextureUV.h
@@ -236,13 +236,13 @@ protected:
 	bool bMirrorOnY_PreUndo = bMirrorOnY;
 #endif
 
-	UPROPERTY(VisibleInstanceOnly, BlueprintReadOnly, TextExportTransient, Category = "Material Designer")
+	UPROPERTY(VisibleInstanceOnly, BlueprintReadOnly, DuplicateTransient, TextExportTransient, Category = "Material Designer")
 	TMap<int32, TObjectPtr<UDMMaterialParameter>> MaterialParameters;
 
-	UPROPERTY(VisibleInstanceOnly, BlueprintReadOnly, TextExportTransient, Category = "Material Designer")
+	UPROPERTY(VisibleInstanceOnly, BlueprintReadOnly, DuplicateTransient, TextExportTransient, Category = "Material Designer")
 	TMap<int32, FName> CachedParameterNames;
 
-	UPROPERTY(VisibleInstanceOnly, BlueprintReadOnly, TextExportTransient, Category = "Material Designer")
+	UPROPERTY(VisibleInstanceOnly, BlueprintReadOnly, Category = "Material Designer")
 	TSet<int32> ExposedParameters;
 
 #if WITH_EDITOR
diff --git a/Engine/Plugins/Experimental/DynamicMaterial/Source/DynamicMaterialEditor/Private/Components/DMMaterialLayer.cpp b/Engine/Plugins/Experimental/DynamicMaterial/Source/DynamicMaterialEditor/Private/Components/DMMaterialLayer.cpp
--- a/Engine/Plugins/Experimental/DynamicMaterial/Source/DynamicMaterialEditor/Private/Components/DMMaterialLayer.cpp
+++ b/Engine/Plugins/Experimental/DynamicMaterial/Source/DynamicMaterialEditor/Private/Components/DMMaterialLayer.cpp
@@ -667,7 +667,7 @@ bool UDMMaterialLayerObject::CanMoveLayerBelow(UDMMaterialLayerObject* InLayer)
 FString UDMMaterialLayerObject::SerializeToString() const
 {
 	TSharedPtr<FJsonObject> LayerJson = MakeShared<FJsonObject>();
-	FJsonObjectConverter::UStructToJsonObject(UDMMaterialLayerObject::StaticClass(), this, LayerJson.ToSharedRef());
+	FJsonObjectConverter::UStructToJsonObject(UDMMaterialLayerObject::StaticClass(), this, LayerJson.ToSharedRef(), 0, CPF_Deprecated | CPF_Transient | CPF_TextExportTransient);
 
 	FString SerializedString;
 	const TSharedRef<TJsonWriter<>> Writer = TJsonWriterFactory<>::Create(&SerializedString);
@@ -962,24 +962,6 @@ UDMMaterialLayerObject* UDMMaterialLayerObject::DeserializeFromString(UDMMateria
 
 	FJsonObjectConverter::JsonObjectToUStruct(LayerJson.ToSharedRef(), StaticClass(), PastedLayer);
 
-	PastedLayer->ForEachValidStage(
-		EDMMaterialLayerStage::All,
-		[PastedLayer, MaterialModel](UDMMaterialStage* InStage)
-		{
-			if (GUndo)
-			{
-				InStage->Modify();
-			}
-
-			InStage->PostEditorDuplicate(MaterialModel, PastedLayer);
-		}
-	);
-
-	if (UDMMaterialEffectStack* EffectStack = PastedLayer->GetEffectStack())
-	{
-		EffectStack->PostEditorDuplicate(MaterialModel, PastedLayer);
-	}
-
 	return PastedLayer;
 }
 
