Material Designer: Fixed actor property display
- Custom material property generators will now display correctly when using something other than the Primitive Component's material list.

diff --git a/Engine/Plugins/Experimental/DynamicMaterial/Source/DynamicMaterialEditor/Private/DMObjectMaterialProperty.cpp b/Engine/Plugins/Experimental/DynamicMaterial/Source/DynamicMaterialEditor/Private/DMObjectMaterialProperty.cpp
--- a/Engine/Plugins/Experimental/DynamicMaterial/Source/DynamicMaterialEditor/Private/DMObjectMaterialProperty.cpp
+++ b/Engine/Plugins/Experimental/DynamicMaterial/Source/DynamicMaterialEditor/Private/DMObjectMaterialProperty.cpp
@@ -57,6 +57,11 @@ UDynamicMaterialModelBase* FDMObjectMaterialProperty::GetMaterialModelBase() con
 }
 
 UDynamicMaterialInstance* FDMObjectMaterialProperty::GetMaterial() const
+{
+	return Cast<UDynamicMaterialInstance>(GetMaterialInterface());
+}
+
+UMaterialInterface* FDMObjectMaterialProperty::GetMaterialInterface() const
 {
 	UObject* Outer = OuterWeak.Get();
 
@@ -94,10 +99,7 @@ UDynamicMaterialInstance* FDMObjectMaterialProperty::GetMaterial() const
 			}
 		}
 
-		if (Material)
-		{
-			return Cast<UDynamicMaterialInstance>(Material);
-		}
+		return Material;
 	}
 	else if (Index != INDEX_NONE)
 	{
@@ -105,7 +107,7 @@ UDynamicMaterialInstance* FDMObjectMaterialProperty::GetMaterial() const
 		{
 			if (Index >= 0 && Index < Component->GetNumMaterials())
 			{
-				return Cast<UDynamicMaterialInstance>(Component->GetMaterial(Index));
+				return Component->GetMaterial(Index);
 			}
 		}
 	}
diff --git a/Engine/Plugins/Experimental/DynamicMaterial/Source/DynamicMaterialEditor/Private/UI/Widgets/SDMActorMaterialSelector.cpp b/Engine/Plugins/Experimental/DynamicMaterial/Source/DynamicMaterialEditor/Private/UI/Widgets/SDMActorMaterialSelector.cpp
--- a/Engine/Plugins/Experimental/DynamicMaterial/Source/DynamicMaterialEditor/Private/UI/Widgets/SDMActorMaterialSelector.cpp
+++ b/Engine/Plugins/Experimental/DynamicMaterial/Source/DynamicMaterialEditor/Private/UI/Widgets/SDMActorMaterialSelector.cpp
@@ -72,15 +72,11 @@ TSharedRef<SWidget> SDMActorMaterialSelector::CreateSelectorLayout()
 
 	const UObject* CurrentOuter = nullptr;
 
-	for (const FDMObjectMaterialProperty& ActorMaterialProperty : ActorProperties)
+	for (int32 Index = 0; Index < ActorProperties.Num(); ++Index)
 	{
-		if (!ActorMaterialProperty.IsValid())
-		{
-			continue;
-		}
+		const FDMObjectMaterialProperty& ActorMaterialProperty = ActorProperties[Index];
 
-		// Only show material elements on the selector
-		if (!ActorMaterialProperty.IsElement())
+		if (!ActorMaterialProperty.IsValid())
 		{
 			continue;
 		}
@@ -92,13 +88,6 @@ TSharedRef<SWidget> SDMActorMaterialSelector::CreateSelectorLayout()
 			continue;
 		}
 
-		UPrimitiveComponent* PrimitiveComponent = Cast<UPrimitiveComponent>(ActorMaterialProperty.GetOuter());
-
-		if (!PrimitiveComponent)
-		{
-			continue;
-		}
-
 		if (Outer != CurrentOuter)
 		{
 			ListOuter->AddSlot()
@@ -117,7 +106,7 @@ TSharedRef<SWidget> SDMActorMaterialSelector::CreateSelectorLayout()
 			.AutoHeight()
 			.HAlign(EHorizontalAlignment::HAlign_Left)
 			[
-				CreateActorMaterialPropertyEntry(PrimitiveComponent, ActorMaterialProperty.GetIndex())
+				CreateActorMaterialPropertyEntry(Index, Cast<UPrimitiveComponent>(ActorMaterialProperty.GetOuter()))
 			];
 	}
 
@@ -143,15 +132,18 @@ TSharedRef<SWidget> SDMActorMaterialSelector::CreateNoPropertiesLayout()
 		.Text(LOCTEXT("NoMaterialSlot", "\n\nThe selected actor contains no primitive components with material slots."));
 }
 
-TSharedRef<SWidget> SDMActorMaterialSelector::CreateActorMaterialPropertyEntry(UPrimitiveComponent* InPrimitiveComponent, 
-	int32 InActorPropertyIndex)
+TSharedRef<SWidget> SDMActorMaterialSelector::CreateActorMaterialPropertyEntry(int32 InActorPropertyIndex, UPrimitiveComponent* InPrimitiveComponent)
 {
 	const FDMObjectMaterialProperty& ActorMaterialProperty = ActorProperties[InActorPropertyIndex];
 
 	constexpr int32 ThumbnailSize = 48;
 
+	UMaterialInterface* Material = InPrimitiveComponent
+		? InPrimitiveComponent->GetMaterial(ActorMaterialProperty.GetIndex())
+		: ActorMaterialProperty.GetMaterialInterface();
+
 	TSharedRef<FAssetThumbnail> Thumbnail = MakeShared<FAssetThumbnail>(
-		InPrimitiveComponent->GetMaterial(ActorMaterialProperty.GetIndex()),
+		Material,
 		ThumbnailSize,
 		ThumbnailSize, 
 		UThumbnailManager::Get().GetSharedThumbnailPool()
diff --git a/Engine/Plugins/Experimental/DynamicMaterial/Source/DynamicMaterialEditor/Private/UI/Widgets/SDMActorMaterialSelector.h b/Engine/Plugins/Experimental/DynamicMaterial/Source/DynamicMaterialEditor/Private/UI/Widgets/SDMActorMaterialSelector.h
--- a/Engine/Plugins/Experimental/DynamicMaterial/Source/DynamicMaterialEditor/Private/UI/Widgets/SDMActorMaterialSelector.h
+++ b/Engine/Plugins/Experimental/DynamicMaterial/Source/DynamicMaterialEditor/Private/UI/Widgets/SDMActorMaterialSelector.h
@@ -36,7 +36,7 @@ protected:
 
 	TSharedRef<SWidget> CreateNoPropertiesLayout();
 
-	TSharedRef<SWidget> CreateActorMaterialPropertyEntry(UPrimitiveComponent* InPrimitiveComponent, int32 InActorPropertyIndex);
+	TSharedRef<SWidget> CreateActorMaterialPropertyEntry(int32 InActorPropertyIndex, UPrimitiveComponent* InPrimitiveComponent = nullptr);
 
 	FReply OnCreateMaterialButtonClicked(int32 InActorPropertyIndex);
 };
diff --git a/Engine/Plugins/Experimental/DynamicMaterial/Source/DynamicMaterialEditor/Public/DMObjectMaterialProperty.h b/Engine/Plugins/Experimental/DynamicMaterial/Source/DynamicMaterialEditor/Public/DMObjectMaterialProperty.h
--- a/Engine/Plugins/Experimental/DynamicMaterial/Source/DynamicMaterialEditor/Public/DMObjectMaterialProperty.h
+++ b/Engine/Plugins/Experimental/DynamicMaterial/Source/DynamicMaterialEditor/Public/DMObjectMaterialProperty.h
@@ -12,6 +12,7 @@ class FProperty;
 class UDynamicMaterialInstance;
 class UDynamicMaterialModel;
 class UDynamicMaterialModelBase;
+class UMaterialInterface;
 class UObject;
 class UPrimitiveComponent;
 
@@ -42,6 +43,8 @@ struct FDMObjectMaterialProperty
 	DYNAMICMATERIALEDITOR_API UDynamicMaterialModelBase* GetMaterialModelBase() const;
 
 	DYNAMICMATERIALEDITOR_API UDynamicMaterialInstance* GetMaterial() const;
+	
+	DYNAMICMATERIALEDITOR_API UMaterialInterface* GetMaterialInterface() const;
 
 	DYNAMICMATERIALEDITOR_API bool IsValid() const;
 
