Material Designer: Changed how transactions work when applying preview material to the source material.

diff --git a/Engine/Plugins/Experimental/DynamicMaterial/Source/DynamicMaterialEditor/Private/UI/Widgets/Editor/SDMToolBar.cpp b/Engine/Plugins/Experimental/DynamicMaterial/Source/DynamicMaterialEditor/Private/UI/Widgets/Editor/SDMToolBar.cpp
--- a/Engine/Plugins/Experimental/DynamicMaterial/Source/DynamicMaterialEditor/Private/UI/Widgets/Editor/SDMToolBar.cpp
+++ b/Engine/Plugins/Experimental/DynamicMaterial/Source/DynamicMaterialEditor/Private/UI/Widgets/Editor/SDMToolBar.cpp
@@ -880,6 +880,7 @@ FReply SDMToolBar::OnSaveClicked()
 {
 	if (TSharedPtr<SDMMaterialEditor> EditorWidget = GetEditorWidget())
 	{
+		FScopedTransaction Transaction(LOCTEXT("SaveOriginalMaterial", "Save Original Material"));
 		EditorWidget->SaveOriginal();
 	}
 
diff --git a/Engine/Plugins/Experimental/DynamicMaterial/Source/DynamicMaterialEditor/Private/UI/Widgets/SDMMaterialEditor.cpp b/Engine/Plugins/Experimental/DynamicMaterial/Source/DynamicMaterialEditor/Private/UI/Widgets/SDMMaterialEditor.cpp
--- a/Engine/Plugins/Experimental/DynamicMaterial/Source/DynamicMaterialEditor/Private/UI/Widgets/SDMMaterialEditor.cpp
+++ b/Engine/Plugins/Experimental/DynamicMaterial/Source/DynamicMaterialEditor/Private/UI/Widgets/SDMMaterialEditor.cpp
@@ -1032,8 +1032,10 @@ void SDMMaterialEditor::ApplyToOriginal()
 
 	UDynamicMaterialInstance* OriginalMaterialInstance = Original->GetDynamicMaterialInstance();
 
-	FScopedTransaction Transaction(LOCTEXT("ApplyMaterialModel", "Apply Material Designer Changes"));
-	Original->Modify();
+	if (GUndo)
+	{
+		Original->Modify();
+	}
 
 	UDMMaterialModelFunctionLibrary::MirrorMaterialModel(Preview, Original);
 
