Material Designer: Moved actor namd and property selector to the end of the title row.

diff --git a/Engine/Plugins/Experimental/DynamicMaterial/Source/DynamicMaterialEditor/Private/UI/Widgets/Editor/SDMToolBar.cpp b/Engine/Plugins/Experimental/DynamicMaterial/Source/DynamicMaterialEditor/Private/UI/Widgets/Editor/SDMToolBar.cpp
--- a/Engine/Plugins/Experimental/DynamicMaterial/Source/DynamicMaterialEditor/Private/UI/Widgets/Editor/SDMToolBar.cpp
+++ b/Engine/Plugins/Experimental/DynamicMaterial/Source/DynamicMaterialEditor/Private/UI/Widgets/Editor/SDMToolBar.cpp
@@ -137,66 +137,3 @@
 			+ SWrapBox::Slot()
 			.HAlign(HAlign_Left)
 			.VAlign(VAlign_Center)
-			.Padding(5.0f, 0.0f, 0.0f, 0.0f)
-			[
-				SAssignNew(InstanceWidget, STextBlock)
-				.TextStyle(FDynamicMaterialEditorStyle::Get(), "ActorName")
-				.Text(LOCTEXT("Instance", "(Inst)"))
-				.Visibility(EVisibility::Collapsed)
-			]
-
-			+ SWrapBox::Slot()
-			.HAlign(HAlign_Left)
-			.VAlign(VAlign_Center)
-			.Padding(5.0f, 0.0f, 0.0f, 0.0f)
-			[
-				SAssignNew(ActorRowWidget, SHorizontalBox)
-				.Visibility(EVisibility::Collapsed)
-				+ SHorizontalBox::Slot()
-				.AutoWidth()
-				.Padding(0.f, 0.f, 5.f, 0.f)
-				.VAlign(EVerticalAlignment::VAlign_Center)
-				[
-					SNew(SImage)
-					.Image(ActorBrush)
-					.DesiredSizeOverride(LargeIconToolBarButtonSize)
-				]
-				+ SHorizontalBox::Slot()
-				.AutoWidth()
-				.Padding(0.f, 0.f, 0.f, 0.f)
-				.VAlign(EVerticalAlignment::VAlign_Center)
-				[
-					SAssignNew(ActorNameWidget, STextBlock)
-					.TextStyle(FDynamicMaterialEditorStyle::Get(), "ActorName")
-				]	
-				+ SHorizontalBox::Slot()
-				.AutoWidth()
-				.VAlign(EVerticalAlignment::VAlign_Center)
-				.Padding(5.0f, 0.0f, 0.0f, 0.0f)
-				[
-					SAssignNew(PropertySelectorContainer, SBox)
-					[
-						CreateSlotsComboBoxWidget()
-					]
-				]	
-				+ SHorizontalBox::Slot()
-				.AutoWidth()
-				.VAlign(EVerticalAlignment::VAlign_Center)
-				.Padding(5.0f, 0.0f, 0.0f, 0.0f)
-				[
-					SNew(SButton)
-					.ContentPadding(LargeIconToolBarButtonContentPadding)
-					.ButtonStyle(FDynamicMaterialEditorStyle::Get(), "HoverHintOnly")
-					.ToolTipText(LOCTEXT("MaterialDesignerUseTooltip", "Replace the material in this slot with the one selected in the content browser."))
-					.OnClicked(this, &SDMToolBar::OnUseClicked)
-					[
-						SNew(SImage)
-						.Image(FAppStyle::GetBrush(TEXT("Icons.Use")))
-						.DesiredSizeOverride(LargeIconToolBarButtonSize)
-					]
-				]
-			]
-
-			+ SWrapBox::Slot()
-			.HAlign(HAlign_Left)
-			.VAlign(VAlign_Center)
@@ -272,6 +209,69 @@
 					]
 				]
 			]
+
+			+ SWrapBox::Slot()
+			.HAlign(HAlign_Left)
+			.VAlign(VAlign_Center)
+			.Padding(5.0f, 0.0f, 0.0f, 0.0f)
+			[
+				SAssignNew(InstanceWidget, STextBlock)
+				.TextStyle(FDynamicMaterialEditorStyle::Get(), "ActorName")
+				.Text(LOCTEXT("Instance", "(Inst)"))
+				.Visibility(EVisibility::Collapsed)
+			]
+
+			+ SWrapBox::Slot()
+			.HAlign(HAlign_Left)
+			.VAlign(VAlign_Center)
+			.Padding(5.0f, 0.0f, 0.0f, 0.0f)
+			[
+				SAssignNew(ActorRowWidget, SHorizontalBox)
+				.Visibility(EVisibility::Collapsed)
+				+ SHorizontalBox::Slot()
+				.AutoWidth()
+				.Padding(0.f, 0.f, 5.f, 0.f)
+				.VAlign(EVerticalAlignment::VAlign_Center)
+				[
+					SNew(SImage)
+					.Image(ActorBrush)
+					.DesiredSizeOverride(LargeIconToolBarButtonSize)
+				]
+				+ SHorizontalBox::Slot()
+				.AutoWidth()
+				.Padding(0.f, 0.f, 0.f, 0.f)
+				.VAlign(EVerticalAlignment::VAlign_Center)
+				[
+					SAssignNew(ActorNameWidget, STextBlock)
+					.TextStyle(FDynamicMaterialEditorStyle::Get(), "ActorName")
+				]	
+				+ SHorizontalBox::Slot()
+				.AutoWidth()
+				.VAlign(EVerticalAlignment::VAlign_Center)
+				.Padding(5.0f, 0.0f, 0.0f, 0.0f)
+				[
+					SAssignNew(PropertySelectorContainer, SBox)
+					[
+						CreateSlotsComboBoxWidget()
+					]
+				]	
+				+ SHorizontalBox::Slot()
+				.AutoWidth()
+				.VAlign(EVerticalAlignment::VAlign_Center)
+				.Padding(5.0f, 0.0f, 0.0f, 0.0f)
+				[
+					SNew(SButton)
+					.ContentPadding(LargeIconToolBarButtonContentPadding)
+					.ButtonStyle(FDynamicMaterialEditorStyle::Get(), "HoverHintOnly")
+					.ToolTipText(LOCTEXT("MaterialDesignerUseTooltip", "Replace the material in this slot with the one selected in the content browser."))
+					.OnClicked(this, &SDMToolBar::OnUseClicked)
+					[
+						SNew(SImage)
+						.Image(FAppStyle::GetBrush(TEXT("Icons.Use")))
+						.DesiredSizeOverride(LargeIconToolBarButtonSize)
+					]
+				]
+			]
 		]
 
 		+ SHorizontalBox::Slot()
