Material Designer: Fixed cyclic call to CopyParameters
- A details panel refresh will no longer trigger parameter copying.

diff --git a/Engine/Plugins/Experimental/DynamicMaterial/Source/DynamicMaterialEditor/Private/UI/Widgets/Editor/SDMMaterialComponentEditor.cpp b/Engine/Plugins/Experimental/DynamicMaterial/Source/DynamicMaterialEditor/Private/UI/Widgets/Editor/SDMMaterialComponentEditor.cpp
--- a/Engine/Plugins/Experimental/DynamicMaterial/Source/DynamicMaterialEditor/Private/UI/Widgets/Editor/SDMMaterialComponentEditor.cpp
+++ b/Engine/Plugins/Experimental/DynamicMaterial/Source/DynamicMaterialEditor/Private/UI/Widgets/Editor/SDMMaterialComponentEditor.cpp
@@ -115,16 +115,20 @@ void SDMMaterialComponentEditor::OnComponentUpdated(UDMMaterialComponent* InComp
 		}
 	}
 
-	if (UDynamicMaterialEditorSettings* Settings = UDynamicMaterialEditorSettings::Get())
+	// Don't copy parameters if it's just a details panel refresh.
+	if (InUpdateType != EDMUpdateType::RefreshDetailView)
 	{
-		if (Settings->ShouldAutomaticallyCopyParametersToSourceMaterial())
+		if (UDynamicMaterialEditorSettings* Settings = UDynamicMaterialEditorSettings::Get())
 		{
-			if (InComponent)
+			if (Settings->ShouldAutomaticallyCopyParametersToSourceMaterial())
 			{
-				if (UDMMaterialComponent* OriginalComponent = GetOriginalComponent(InComponent))
+				if (InComponent)
 				{
-					IDMParameterContainer::CopyParametersBetween(InComponent, OriginalComponent);
-					return;
+					if (UDMMaterialComponent* OriginalComponent = GetOriginalComponent(InComponent))
+					{
+						IDMParameterContainer::CopyParametersBetween(InComponent, OriginalComponent);
+						return;
+					}
 				}
 			}
 		}
