MotionDesign : PropertyAnimator
- Added support for property animator within material designer (original material) by adapting Custom Details View

diff --git a/Engine/Plugins/Experimental/CustomDetailsView/Source/CustomDetailsView/Private/Items/CustomDetailsViewDetailTreeNodeItem.cpp b/Engine/Plugins/Experimental/CustomDetailsView/Source/CustomDetailsView/Private/Items/CustomDetailsViewDetailTreeNodeItem.cpp
--- a/Engine/Plugins/Experimental/CustomDetailsView/Source/CustomDetailsView/Private/Items/CustomDetailsViewDetailTreeNodeItem.cpp
+++ b/Engine/Plugins/Experimental/CustomDetailsView/Source/CustomDetailsView/Private/Items/CustomDetailsViewDetailTreeNodeItem.cpp
@@ -159,9 +159,12 @@ void FCustomDetailsViewDetailTreeNodeItem::SetCreateChildItemDelegate(FOnCustomD
 	ChildItemDelegate = InDelegate;
 }
 
-void FCustomDetailsViewDetailTreeNodeItem::AddExtensionWidget(const TSharedRef<SSplitter>& InSplitter
-	, const FDetailColumnSizeData& InColumnSizeData
-	, const FCustomDetailsViewArgs& InViewArgs)
+void FCustomDetailsViewDetailTreeNodeItem::SetCustomizeItemMenuContext(FOnCustomDetailsViewCustomizeItemMenuContext InDelegate)
+{
+	ContextMenuDelegate = InDelegate;
+}
+
+void FCustomDetailsViewDetailTreeNodeItem::AddExtensionWidget(const TSharedRef<SSplitter>& InSplitter, const FDetailColumnSizeData& InColumnSizeData, const FCustomDetailsViewArgs& InViewArgs)
 {
 	if (GetOverrideWidget(ECustomDetailsViewWidgetType::Extensions).IsValid())
 	{
@@ -359,6 +362,12 @@ TSharedPtr<SWidget> FCustomDetailsViewDetailTreeNodeItem::GenerateContextMenuWid
 		RefreshChildren();
 	});
 
+	// Customize context for this menu
+	if (const TSharedPtr<ICustomDetailsView> DetailsView = CustomDetailsViewWeak.Pin())
+	{
+		ContextMenuDelegate.ExecuteIfBound(DetailsView.ToSharedRef(), AsShared(), RowMenuContext, RowMenuContext->PropertyHandles);
+	}
+
 	FToolMenuContext ToolMenuContext(RowMenuContext);
 	ToolMenuContext.AddObject(GetMutableDefault<UCustomDetailsViewMenuContext>());
 
diff --git a/Engine/Plugins/Experimental/CustomDetailsView/Source/CustomDetailsView/Private/Items/CustomDetailsViewDetailTreeNodeItem.h b/Engine/Plugins/Experimental/CustomDetailsView/Source/CustomDetailsView/Private/Items/CustomDetailsViewDetailTreeNodeItem.h
--- a/Engine/Plugins/Experimental/CustomDetailsView/Source/CustomDetailsView/Private/Items/CustomDetailsViewDetailTreeNodeItem.h
+++ b/Engine/Plugins/Experimental/CustomDetailsView/Source/CustomDetailsView/Private/Items/CustomDetailsViewDetailTreeNodeItem.h
@@ -36,6 +36,7 @@ public:
 	virtual void CreateGlobalExtensionButtons(TArray<FPropertyRowExtensionButton>& OutExtensionButtons) override;
 	virtual bool CreateResetToDefaultButton(FPropertyRowExtensionButton& OutButton) override;
 	virtual void SetCreateChildItemDelegate(FOnCustomDetailsViewGenerateChildItem InDelegate) override;
+	virtual void SetCustomizeItemMenuContext(FOnCustomDetailsViewCustomizeItemMenuContext InDelegate) override;
 	//~ End ICustomDetailsViewItem
 
 	//~ Begin FCustomDetailsViewItemBase
@@ -103,4 +104,7 @@ protected:
 
 	/** Used for custom child property rows. */
 	FOnCustomDetailsViewGenerateChildItem ChildItemDelegate;
+
+	/** Used to customize the context menu */
+	FOnCustomDetailsViewCustomizeItemMenuContext ContextMenuDelegate;
 };
diff --git a/Engine/Plugins/Experimental/CustomDetailsView/Source/CustomDetailsView/Private/Items/CustomDetailsViewItemBase.h b/Engine/Plugins/Experimental/CustomDetailsView/Source/CustomDetailsView/Private/Items/CustomDetailsViewItemBase.h
--- a/Engine/Plugins/Experimental/CustomDetailsView/Source/CustomDetailsView/Private/Items/CustomDetailsViewItemBase.h
+++ b/Engine/Plugins/Experimental/CustomDetailsView/Source/CustomDetailsView/Private/Items/CustomDetailsViewItemBase.h
@@ -60,6 +60,7 @@ public:
 	virtual void CreateGlobalExtensionButtons(TArray<FPropertyRowExtensionButton>& OutExtensionButtons) override;
 	virtual TSharedRef<SWidget> CreateExtensionButtonWidget(const TArray<FPropertyRowExtensionButton>& InExtensionButtons) override;
 	virtual void SetCreateChildItemDelegate(FOnCustomDetailsViewGenerateChildItem InDelegate) override {} // No default implementation
+	virtual void SetCustomizeItemMenuContext(FOnCustomDetailsViewCustomizeItemMenuContext InDelegate) override {} // No default implementation
 	//~ End ICustomDetailsViewItem
 
 	virtual void AddWholeRowWidget(const TSharedRef<SSplitter>& InSplitter
diff --git a/Engine/Plugins/Experimental/CustomDetailsView/Source/CustomDetailsView/Public/Items/ICustomDetailsViewItem.h b/Engine/Plugins/Experimental/CustomDetailsView/Source/CustomDetailsView/Public/Items/ICustomDetailsViewItem.h
--- a/Engine/Plugins/Experimental/CustomDetailsView/Source/CustomDetailsView/Public/Items/ICustomDetailsViewItem.h
+++ b/Engine/Plugins/Experimental/CustomDetailsView/Source/CustomDetailsView/Public/Items/ICustomDetailsViewItem.h
@@ -21,6 +21,14 @@ DECLARE_DELEGATE_RetVal_ThreeParams(
 	const TSharedRef<IDetailTreeNode>& InChildNode
 )
 
+DECLARE_DELEGATE_FourParams(
+	FOnCustomDetailsViewCustomizeItemMenuContext,
+	const TSharedRef<ICustomDetailsView>& InCustomDetailsView,
+	const TSharedPtr<ICustomDetailsViewItem>& InItem,
+	UObject* InMenuContext,
+	TArray<TSharedPtr<IPropertyHandle>>& InPropertyHandles
+)
+
 class ICustomDetailsViewItem : public TSharedFromThis<ICustomDetailsViewItem>
 {
 public:
@@ -115,6 +123,11 @@ public:
 	 */
 	virtual void SetCreateChildItemDelegate(FOnCustomDetailsViewGenerateChildItem InDelegate) = 0;
 
+	/**
+	 * Allows to customize the context menu of this item
+	 */
+	virtual void SetCustomizeItemMenuContext(FOnCustomDetailsViewCustomizeItemMenuContext InDelegate) = 0;
+
 protected:
 	/** Adds the children of this node. */
 	virtual void GatherChildren(const TSharedRef<ICustomDetailsViewItem>& InParentItem,
diff --git a/Engine/Plugins/Experimental/DynamicMaterial/Source/DynamicMaterialEditor/Private/UI/Widgets/Editor/SDMObjectEditorWidgetBase.cpp b/Engine/Plugins/Experimental/DynamicMaterial/Source/DynamicMaterialEditor/Private/UI/Widgets/Editor/SDMObjectEditorWidgetBase.cpp
--- a/Engine/Plugins/Experimental/DynamicMaterial/Source/DynamicMaterialEditor/Private/UI/Widgets/Editor/SDMObjectEditorWidgetBase.cpp
+++ b/Engine/Plugins/Experimental/DynamicMaterial/Source/DynamicMaterialEditor/Private/UI/Widgets/Editor/SDMObjectEditorWidgetBase.cpp
@@ -218,6 +218,8 @@ void SDMObjectEditorWidgetBase::AddDetailTreeRow(const TSharedRef<ICustomDetails
 {
 	TSharedRef<ICustomDetailsViewItem> Item = InDetailsView->CreateDetailTreeItem(InParent, InPropertyRow.PreviewHandle.DetailTreeNode.ToSharedRef());
 
+	CustomizeItemContextMenu(Item, InPropertyRow);
+
 	if (InPropertyRow.NameOverride.IsSet())
 	{
 		Item->SetOverrideWidget(
@@ -281,6 +283,8 @@ void SDMObjectEditorWidgetBase::AddCustomRow(const TSharedRef<ICustomDetailsView
 		return;
 	}
 
+	CustomizeItemContextMenu(Item->AsItem(), InPropertyRow);
+
 	if (bIsWholeRow)
 	{
 		Item->SetWholeRowWidget(InPropertyRow.ValueWidget.ToSharedRef());
@@ -555,6 +559,8 @@ TSharedPtr<ICustomDetailsViewItem> SDMObjectEditorWidgetBase::CreateChildItem_Im
 {
 	TSharedRef<ICustomDetailsViewItem> Item = InDetailsView->CreateDetailTreeItem(InParent, InChildNode);
 
+	CustomizeItemContextMenu(Item, InOriginalRow);
+
 	if (!InOriginalRow.bEnabled)
 	{
 		Item->SetEnabledOverride(false);
@@ -585,4 +591,20 @@ TSharedPtr<ICustomDetailsViewItem> SDMObjectEditorWidgetBase::CreateChildItem_Im
 	return Item;
 }
 
+void SDMObjectEditorWidgetBase::CustomizeItemContextMenu(const TSharedRef<ICustomDetailsViewItem>& InItem, const FDMPropertyHandle& InPropertyRow)
+{
+	if (InPropertyRow.PreviewHandle.DetailTreeNode.IsValid())
+	{
+		InItem->SetCustomizeItemMenuContext(FOnCustomDetailsViewCustomizeItemMenuContext::CreateSPLambda(this,
+			[this, InPropertyRow](const TSharedRef<ICustomDetailsView>&, const TSharedPtr<ICustomDetailsViewItem>&, UObject*, TArray<TSharedPtr<IPropertyHandle>>& InPropertyHandles)
+			{
+				if (InPropertyRow.OriginalHandle.PropertyHandle.IsValid())
+				{
+					InPropertyHandles.Add(InPropertyRow.OriginalHandle.PropertyHandle);
+				}
+			}
+		));
+	}
+}
+
 #undef LOCTEXT_NAMESPACE
diff --git a/Engine/Plugins/Experimental/DynamicMaterial/Source/DynamicMaterialEditor/Private/UI/Widgets/Editor/SDMObjectEditorWidgetBase.h b/Engine/Plugins/Experimental/DynamicMaterial/Source/DynamicMaterialEditor/Private/UI/Widgets/Editor/SDMObjectEditorWidgetBase.h
--- a/Engine/Plugins/Experimental/DynamicMaterial/Source/DynamicMaterialEditor/Private/UI/Widgets/Editor/SDMObjectEditorWidgetBase.h
+++ b/Engine/Plugins/Experimental/DynamicMaterial/Source/DynamicMaterialEditor/Private/UI/Widgets/Editor/SDMObjectEditorWidgetBase.h
@@ -97,4 +97,6 @@ protected:
 
 	virtual TSharedPtr<ICustomDetailsViewItem> CreateChildItem_Impl(const TSharedRef<ICustomDetailsView>& InDetailsView, const TSharedRef<ICustomDetailsViewItem>& InParent,
 		const TSharedRef<IDetailTreeNode>& InChildNode, const TSharedRef<IPropertyHandle>& InPropertyHandle, const FDMPropertyHandle& InOriginalRow);
+
+	void CustomizeItemContextMenu(const TSharedRef<ICustomDetailsViewItem>& InItem, const FDMPropertyHandle& InPropertyRow);
 };
diff --git a/Engine/Plugins/Experimental/PropertyAnimatorCore/Source/PropertyAnimatorCoreEditor/Private/Subsystems/PropertyAnimatorCoreEditorSubsystem.cpp b/Engine/Plugins/Experimental/PropertyAnimatorCore/Source/PropertyAnimatorCoreEditor/Private/Subsystems/PropertyAnimatorCoreEditorSubsystem.cpp
--- a/Engine/Plugins/Experimental/PropertyAnimatorCore/Source/PropertyAnimatorCoreEditor/Private/Subsystems/PropertyAnimatorCoreEditorSubsystem.cpp
+++ b/Engine/Plugins/Experimental/PropertyAnimatorCore/Source/PropertyAnimatorCoreEditor/Private/Subsystems/PropertyAnimatorCoreEditorSubsystem.cpp
@@ -468,7 +468,17 @@ void UPropertyAnimatorCoreEditorSubsystem::FillAnimatorRowContextSection(UToolMe
 
 	constexpr bool bFindMemberProperty = true;
 	TArray<FPropertyAnimatorCoreData> Properties;
-	if (!GetPropertiesFromHandle(Context->PropertyHandles[0], Properties, bFindMemberProperty))
+	for (const TSharedPtr<IPropertyHandle>& PropertyHandle : Context->PropertyHandles)
+	{
+		if (GetPropertiesFromHandle(PropertyHandle, Properties, bFindMemberProperty))
+		{
+			break;
+		}
+
+		Properties.Reset();
+	}
+
+	if (Properties.IsEmpty())
 	{
 		return;
 	}
@@ -562,7 +572,7 @@ bool UPropertyAnimatorCoreEditorSubsystem::GetPropertiesFromHandle(const TShared
 		}
 	}
 
-	return true;
+	return PropertyData.GetOwningActor() != nullptr;
 }
 
 #undef LOCTEXT_NAMESPACE
