Material Designer: Changed order of change stage source and layer add menus.

diff --git a/Engine/Plugins/Experimental/DynamicMaterial/Source/DynamicMaterialEditor/Private/UI/Menus/DMMaterialSlotLayerMenus.cpp b/Engine/Plugins/Experimental/DynamicMaterial/Source/DynamicMaterialEditor/Private/UI/Menus/DMMaterialSlotLayerMenus.cpp
--- a/Engine/Plugins/Experimental/DynamicMaterial/Source/DynamicMaterialEditor/Private/UI/Menus/DMMaterialSlotLayerMenus.cpp
+++ b/Engine/Plugins/Experimental/DynamicMaterial/Source/DynamicMaterialEditor/Private/UI/Menus/DMMaterialSlotLayerMenus.cpp
@@ -168,37 +168,6 @@ void FDMMaterialSlotLayerMenus::AddAddLayerSection(FToolMenuSection& InSection)
 		return;
 	}
 
-	InSection.AddMenuEntry(
-		TEXT("Noise"),
-		LOCTEXT("AddNoise", "Noise"),
-		LOCTEXT("AddNoiseTooltip", "Add a new Material Layer with a noise pattern."),
-		GetDefault<UDMMaterialStageExpressionWorldPositionNoise>()->GetComponentIcon(),
-		FUIAction(FExecuteAction::CreateWeakLambda(
-			Slot,
-			[Slot]
-			{
-				UDMMaterialSlotFunctionLibrary::AddNewLayer_Expression(Slot, TSubclassOf<UDMMaterialStageExpression>(UDMMaterialStageExpressionWorldPositionNoise::StaticClass()));
-			}
-		))
-	);
-
-	if (ModelEditorOnlyData->GetDomain() == EMaterialDomain::MD_PostProcess)
-	{
-		InSection.AddMenuEntry(
-			TEXT("PostProcess"),
-			LOCTEXT("AddSceneTexture", "Post Process"),
-			LOCTEXT("AddSceneTextureTooltip", "Add a new Material Layer that represents the Scene Texture for a post process material."),
-			GetDefault<UDMMaterialStageExpressionSceneTexture>()->GetComponentIcon(),
-			FUIAction(FExecuteAction::CreateWeakLambda(
-				Slot,
-				[Slot]
-				{
-					UDMMaterialSlotFunctionLibrary::AddNewLayer_SceneTexture(Slot);
-				}
-			))
-		);
-	}
-
 	InSection.AddMenuEntry(
 		TEXT("SolidColor"),
 		LOCTEXT("AddColor", "Solid Color"),
@@ -210,7 +179,7 @@ void FDMMaterialSlotLayerMenus::AddAddLayerSection(FToolMenuSection& InSection)
 			{
 				UDMMaterialSlotFunctionLibrary::AddNewLayer_NewLocalValue(Slot, EDMValueType::VT_Float3_RGB);
 			}
-		))
+				))
 	);
 
 	InSection.AddMenuEntry(
@@ -238,6 +207,37 @@ void FDMMaterialSlotLayerMenus::AddAddLayerSection(FToolMenuSection& InSection)
 			{
 				UDMMaterialSlotFunctionLibrary::AddNewLayer_Expression(Slot, TSubclassOf<UDMMaterialStageExpression>(UDMMaterialStageExpressionTextureSampleEdgeColor::StaticClass()));
 			}
+				))
+	);
+
+	if (ModelEditorOnlyData->GetDomain() == EMaterialDomain::MD_PostProcess)
+	{
+		InSection.AddMenuEntry(
+			TEXT("PostProcess"),
+			LOCTEXT("AddSceneTexture", "Post Process"),
+			LOCTEXT("AddSceneTextureTooltip", "Add a new Material Layer that represents the Scene Texture for a post process material."),
+			GetDefault<UDMMaterialStageExpressionSceneTexture>()->GetComponentIcon(),
+			FUIAction(FExecuteAction::CreateWeakLambda(
+				Slot,
+				[Slot]
+				{
+					UDMMaterialSlotFunctionLibrary::AddNewLayer_SceneTexture(Slot);
+				}
+					))
+		);
+	}
+
+	InSection.AddMenuEntry(
+		TEXT("Noise"),
+		LOCTEXT("AddNoise", "Noise"),
+		LOCTEXT("AddNoiseTooltip", "Add a new Material Layer with a noise pattern."),
+		GetDefault<UDMMaterialStageExpressionWorldPositionNoise>()->GetComponentIcon(),
+		FUIAction(FExecuteAction::CreateWeakLambda(
+			Slot,
+			[Slot]
+			{
+				UDMMaterialSlotFunctionLibrary::AddNewLayer_Expression(Slot, TSubclassOf<UDMMaterialStageExpression>(UDMMaterialStageExpressionWorldPositionNoise::StaticClass()));
+			}
 		))
 	);
 
diff --git a/Engine/Plugins/Experimental/DynamicMaterial/Source/DynamicMaterialEditor/Private/UI/Menus/DMMaterialStageSourceMenus.cpp b/Engine/Plugins/Experimental/DynamicMaterial/Source/DynamicMaterialEditor/Private/UI/Menus/DMMaterialStageSourceMenus.cpp
--- a/Engine/Plugins/Experimental/DynamicMaterial/Source/DynamicMaterialEditor/Private/UI/Menus/DMMaterialStageSourceMenus.cpp
+++ b/Engine/Plugins/Experimental/DynamicMaterial/Source/DynamicMaterialEditor/Private/UI/Menus/DMMaterialStageSourceMenus.cpp
@@ -1464,37 +1464,6 @@
 
 	const TArray<TStrongObjectPtr<UClass>>& Gradients = UDMMaterialStageGradient::GetAvailableGradients();
 
-	InSection.AddMenuEntry("Noise",
-		LOCTEXT("ChangeSourceNoise", "Noise"),
-		LOCTEXT("ChangeSourceNoiseTooltip", "Change the source of this stage to a Noise Renderer."),
-		GetDefault<UDMMaterialStageExpressionWorldPositionNoise>()->GetComponentIcon(),
-		FUIAction(
-			FExecuteAction::CreateStatic(
-				&ChangeSourceToNoiseFromContext,
-				MenuContext
-			)
-		)
-	);
-
-	if (ModelEditorOnlyData->GetDomain() == EMaterialDomain::MD_PostProcess)
-	{
-		InSection.AddMenuEntry("SceneTexture",
-			LOCTEXT("AddSceneTexture", "Post Process"),
-			LOCTEXT("ChangeSourceSceneTextureTooltip", "Change the source of this stage to Scene Texture in post process materials."),
-			GetDefault<UDMMaterialStageExpressionSceneTexture>()->GetComponentIcon(),
-			FUIAction(
-				FExecuteAction::CreateStatic(
-					&ChangeSourceToSolidColorRGBFromContext,
-					MenuContext
-				),
-				FCanExecuteAction::CreateStatic(
-					&CanChangeSourceToSceneTextureFromContext,
-					MenuContext
-				)
-			)
-		);
-	}
-
 	InSection.AddMenuEntry("SolidColor",
 		LOCTEXT("ChangeSourceColorRGB", "Solid Color"),
 		LOCTEXT("ChangeSourceColorRGBTooltip", "Change the source of this stage to a Solid Color."),
@@ -1531,6 +1500,37 @@
 		)
 	);
 
+	if (ModelEditorOnlyData->GetDomain() == EMaterialDomain::MD_PostProcess)
+	{
+		InSection.AddMenuEntry("SceneTexture",
+			LOCTEXT("AddSceneTexture", "Post Process"),
+			LOCTEXT("ChangeSourceSceneTextureTooltip", "Change the source of this stage to Scene Texture in post process materials."),
+			GetDefault<UDMMaterialStageExpressionSceneTexture>()->GetComponentIcon(),
+			FUIAction(
+				FExecuteAction::CreateStatic(
+					&ChangeSourceToSolidColorRGBFromContext,
+					MenuContext
+				),
+				FCanExecuteAction::CreateStatic(
+					&CanChangeSourceToSceneTextureFromContext,
+					MenuContext
+				)
+			)
+		);
+	}
+
+	InSection.AddMenuEntry("Noise",
+		LOCTEXT("ChangeSourceNoise", "Noise"),
+		LOCTEXT("ChangeSourceNoiseTooltip", "Change the source of this stage to a Noise Renderer."),
+		GetDefault<UDMMaterialStageExpressionWorldPositionNoise>()->GetComponentIcon(),
+		FUIAction(
+			FExecuteAction::CreateStatic(
+				&ChangeSourceToNoiseFromContext,
+				MenuContext
+			)
+		)
+	);
+
 	if constexpr (UE::DynamicMaterialEditor::bGlobalValuesEnabled)
 	{
 		const TArray<UDMMaterialValue*>& Values = MaterialModel->GetValues();
diff --git a/Engine/Plugins/Experimental/DynamicMaterialMediaStreamBridge/Source/DynamicMaterialMediaStreamBridgeEditor/Private/DMMediaStreamStageSourceMenuExtender.cpp b/Engine/Plugins/Experimental/DynamicMaterialMediaStreamBridge/Source/DynamicMaterialMediaStreamBridgeEditor/Private/DMMediaStreamStageSourceMenuExtender.cpp
--- a/Engine/Plugins/Experimental/DynamicMaterialMediaStreamBridge/Source/DynamicMaterialMediaStreamBridgeEditor/Private/DMMediaStreamStageSourceMenuExtender.cpp
+++ b/Engine/Plugins/Experimental/DynamicMaterialMediaStreamBridge/Source/DynamicMaterialMediaStreamBridgeEditor/Private/DMMediaStreamStageSourceMenuExtender.cpp
@@ -97,7 +97,7 @@ void FDMMediaStreamStageSourceMenuExtender::ExtendMenu_ChangeSource(FToolMenuSec
 		)
 	);
 
-	NewEntry.InsertPosition = FToolMenuInsert(TEXT("Noise"), EToolMenuInsertType::Before);
+	NewEntry.InsertPosition = FToolMenuInsert(TEXT("Noise"), EToolMenuInsertType::After);
 }
 
 void FDMMediaStreamStageSourceMenuExtender::ExtendMenu_AddLayer(FToolMenuSection& InSection)
