Motion Design: fixed unacceptable word

diff --git a/Engine/Plugins/VirtualProduction/DynamicMaterial/Source/DynamicMaterialEditor/Private/Components/DMMaterialProperty.cpp b/Engine/Plugins/VirtualProduction/DynamicMaterial/Source/DynamicMaterialEditor/Private/Components/DMMaterialProperty.cpp
--- a/Engine/Plugins/VirtualProduction/DynamicMaterial/Source/DynamicMaterialEditor/Private/Components/DMMaterialProperty.cpp
+++ b/Engine/Plugins/VirtualProduction/DynamicMaterial/Source/DynamicMaterialEditor/Private/Components/DMMaterialProperty.cpp
@@ -793,7 +793,7 @@ bool UDMMaterialProperty::IsValidForModel(const UDynamicMaterialModelEditorOnlyD
 		/* Tesselation enabled */ InMaterialModel.IsNaniteTessellationEnabled(),
 		/* BlendableOutputAlpha (Post Process Alpha) */ false,
 		/* Uses Distortion */ false,
-		/* Shading model from master material */ false,
+		/* Shading model from main material */ false,
 		/* Outputting translucency velocity */ InMaterialModel.GetBlendMode() != BLEND_Opaque,
 		/* Thin surface */ false,
 		/* Is supported (substrate check) */ true
