Material Designer: Changed order of change stage source and layer add menus. [Redux]

diff --git a/Engine/Plugins/VirtualProduction/DynamicMaterial/Source/DynamicMaterialEditor/Private/UI/Menus/DMMaterialStageSourceMenus.cpp b/Engine/Plugins/VirtualProduction/DynamicMaterial/Source/DynamicMaterialEditor/Private/UI/Menus/DMMaterialStageSourceMenus.cpp
--- a/Engine/Plugins/VirtualProduction/DynamicMaterial/Source/DynamicMaterialEditor/Private/UI/Menus/DMMaterialStageSourceMenus.cpp
+++ b/Engine/Plugins/VirtualProduction/DynamicMaterial/Source/DynamicMaterialEditor/Private/UI/Menus/DMMaterialStageSourceMenus.cpp
@@ -1464,25 +1464,25 @@ void FDMMaterialStageSourceMenus::CreateChangeMaterialStageSource(FToolMenuSecti
 
 	const TArray<TStrongObjectPtr<UClass>>& Gradients = UDMMaterialStageGradient::GetAvailableGradients();
 
-	InSection.AddMenuEntry("SolidColor",
-		LOCTEXT("ChangeSourceColorRGB", "Solid Color"),
-		LOCTEXT("ChangeSourceColorRGBTooltip", "Change the source of this stage to a Solid Color."),
-		GetDefault<UDMMaterialValueFloat3RGB>()->GetComponentIcon(),
+	InSection.AddMenuEntry("TextureSample",
+		LOCTEXT("TextureSample", "Texture"),
+		LOCTEXT("ChangeSourceTextureSampleTooltip", "Change the source of this stage to a texture."),
+		GetDefault<UDMMaterialStageExpressionTextureSample>()->GetComponentIcon(),
 		FUIAction(
 			FExecuteAction::CreateStatic(
-				&ChangeSourceToSolidColorRGBFromContext,
+				&ChangeSourceToTextureSampleFromContext,
 				MenuContext
 			)
 		)
 	);
 
-	InSection.AddMenuEntry("TextureSample",
-		LOCTEXT("TextureSample", "Texture"),
-		LOCTEXT("ChangeSourceTextureSampleTooltip", "Change the source of this stage to a texture."),
-		GetDefault<UDMMaterialStageExpressionTextureSample>()->GetComponentIcon(),
+	InSection.AddMenuEntry("SolidColor",
+		LOCTEXT("ChangeSourceColorRGB", "Solid Color"),
+		LOCTEXT("ChangeSourceColorRGBTooltip", "Change the source of this stage to a Solid Color."),
+		GetDefault<UDMMaterialValueFloat3RGB>()->GetComponentIcon(),
 		FUIAction(
 			FExecuteAction::CreateStatic(
-				&ChangeSourceToTextureSampleFromContext,
+				&ChangeSourceToSolidColorRGBFromContext,
 				MenuContext
 			)
 		)
