Material Designer: Changing channel overrides now correctly updates the material.

diff --git a/Engine/Plugins/VirtualProduction/DynamicMaterial/Source/DynamicMaterialEditor/Private/Components/DMMaterialStageBlend.cpp b/Engine/Plugins/VirtualProduction/DynamicMaterial/Source/DynamicMaterialEditor/Private/Components/DMMaterialStageBlend.cpp
--- a/Engine/Plugins/VirtualProduction/DynamicMaterial/Source/DynamicMaterialEditor/Private/Components/DMMaterialStageBlend.cpp
+++ b/Engine/Plugins/VirtualProduction/DynamicMaterial/Source/DynamicMaterialEditor/Private/Components/DMMaterialStageBlend.cpp
@@ -435,6 +435,18 @@ void UDMMaterialStageBlend::PostEditUndo()
 	PullBaseChannelOverride();
 }
 
+void UDMMaterialStageBlend::PostEditChangeProperty(FPropertyChangedEvent& InPropertyChangedEvent)
+{
+	Super::PostEditChangeProperty(InPropertyChangedEvent);
+
+	static const FName BaseChannelOverrideName = GET_MEMBER_NAME_CHECKED(UDMMaterialStageBlend, BaseChannelOverride);
+
+	if (InPropertyChangedEvent.GetPropertyName() == BaseChannelOverrideName)
+	{
+		PushBaseChannelOverride();
+	}
+}
+
 void UDMMaterialStageBlend::NotifyPostChange(const FPropertyChangedEvent& InPropertyChangedEvent, FEditPropertyChain* InPropertyThatChanged)
 {
 	if (!IsComponentValid())
diff --git a/Engine/Plugins/VirtualProduction/DynamicMaterial/Source/DynamicMaterialEditor/Private/Components/DMMaterialStageThroughputLayerBlend.cpp b/Engine/Plugins/VirtualProduction/DynamicMaterial/Source/DynamicMaterialEditor/Private/Components/DMMaterialStageThroughputLayerBlend.cpp
--- a/Engine/Plugins/VirtualProduction/DynamicMaterial/Source/DynamicMaterialEditor/Private/Components/DMMaterialStageThroughputLayerBlend.cpp
+++ b/Engine/Plugins/VirtualProduction/DynamicMaterial/Source/DynamicMaterialEditor/Private/Components/DMMaterialStageThroughputLayerBlend.cpp
@@ -749,6 +749,18 @@ void UDMMaterialStageThroughputLayerBlend::NotifyPostChange(const FPropertyChang
 	Super::NotifyPostChange(InPropertyChangedEvent, InPropertyThatChanged);
 }
 
+void UDMMaterialStageThroughputLayerBlend::PostEditChangeProperty(FPropertyChangedEvent& InPropertyChangedEvent)
+{
+	Super::PostEditChangeProperty(InPropertyChangedEvent);
+
+	static const FName MaskChannelOverrideName = GET_MEMBER_NAME_CHECKED(UDMMaterialStageThroughputLayerBlend, MaskChannelOverride);
+
+	if (InPropertyChangedEvent.GetPropertyName() == MaskChannelOverrideName)
+	{
+		PushMaskChannelOverride();
+	}
+}
+
 void UDMMaterialStageThroughputLayerBlend::ConnectOutputToInput(const TSharedRef<FDMMaterialBuildState>& InBuildState, int32 InInputIdx,
 	int32 InExpressionInputIndex, UMaterialExpression* InSourceExpression, int32 InSourceOutputIndex, int32 InSourceOutputChannel)
 {
diff --git a/Engine/Plugins/VirtualProduction/DynamicMaterial/Source/DynamicMaterialEditor/Public/Components/DMMaterialStageBlend.h b/Engine/Plugins/VirtualProduction/DynamicMaterial/Source/DynamicMaterialEditor/Public/Components/DMMaterialStageBlend.h
--- a/Engine/Plugins/VirtualProduction/DynamicMaterial/Source/DynamicMaterialEditor/Public/Components/DMMaterialStageBlend.h
+++ b/Engine/Plugins/VirtualProduction/DynamicMaterial/Source/DynamicMaterialEditor/Public/Components/DMMaterialStageBlend.h
@@ -90,6 +90,7 @@ public:
 
 	//~ Begin UObject
 	DYNAMICMATERIALEDITOR_API virtual void PostEditUndo() override;
+	DYNAMICMATERIALEDITOR_API virtual void PostEditChangeProperty(FPropertyChangedEvent& InPropertyChangedEvent) override;
 	//~ End UObject
 
 	//~ Begin FNotifyHook
diff --git a/Engine/Plugins/VirtualProduction/DynamicMaterial/Source/DynamicMaterialEditor/Public/Components/DMMaterialStageThroughputLayerBlend.h b/Engine/Plugins/VirtualProduction/DynamicMaterial/Source/DynamicMaterialEditor/Public/Components/DMMaterialStageThroughputLayerBlend.h
--- a/Engine/Plugins/VirtualProduction/DynamicMaterial/Source/DynamicMaterialEditor/Public/Components/DMMaterialStageThroughputLayerBlend.h
+++ b/Engine/Plugins/VirtualProduction/DynamicMaterial/Source/DynamicMaterialEditor/Public/Components/DMMaterialStageThroughputLayerBlend.h
@@ -53,7 +53,9 @@ public:
 	//~ End UDMMaterialStageSource
 
 	//~ Begin FNotifyHook
-	DYNAMICMATERIALEDITOR_API virtual void NotifyPostChange(const FPropertyChangedEvent& InPropertyChangedEvent, class FEditPropertyChain* InPropertyThatChanged) override;
+	DYNAMICMATERIALEDITOR_API virtual void NotifyPostChange(const FPropertyChangedEvent& InPropertyChangedEvent, 
+		class FEditPropertyChain* InPropertyThatChanged) override;
+	DYNAMICMATERIALEDITOR_API virtual void PostEditChangeProperty(FPropertyChangedEvent& InPropertyChangedEvent) override;
 	//~ End FNotifyHook
 
 	//~ Begin UDMMaterialStageThroughput
