Material Designer: Fixed inability to key or expose layer opacity.

diff --git a/Engine/Plugins/VirtualProduction/DynamicMaterial/Source/DynamicMaterialEditor/Private/UI/Widgets/Editor/SDMMaterialComponentEditor.cpp b/Engine/Plugins/VirtualProduction/DynamicMaterial/Source/DynamicMaterialEditor/Private/UI/Widgets/Editor/SDMMaterialComponentEditor.cpp
--- a/Engine/Plugins/VirtualProduction/DynamicMaterial/Source/DynamicMaterialEditor/Private/UI/Widgets/Editor/SDMMaterialComponentEditor.cpp
+++ b/Engine/Plugins/VirtualProduction/DynamicMaterial/Source/DynamicMaterialEditor/Private/UI/Widgets/Editor/SDMMaterialComponentEditor.cpp
@@ -196,7 +196,14 @@ TArray<FDMPropertyHandle> SDMMaterialComponentEditor::GetPropertyRows()
 
 	FDynamicMaterialEditorModule::GeneratorComponentPropertyRows(Params);
 
-	for (const FDMPropertyHandle& PropertyRow : PropertyRows)
+	BindPropertyRowUpdateDelegates(PropertyRows);
+
+	return PropertyRows;
+}
+
+void SDMMaterialComponentEditor::BindPropertyRowUpdateDelegates(TConstArrayView<FDMPropertyHandle> InPropertyRows)
+{
+	for (const FDMPropertyHandle& PropertyRow : InPropertyRows)
 	{
 		TArray<UObject*> Outers;
 
@@ -233,8 +240,6 @@ TArray<FDMPropertyHandle> SDMMaterialComponentEditor::GetPropertyRows()
 			}
 		}
 	}
-
-	return PropertyRows;
 }
 
 void SDMMaterialComponentEditor::AddDetailTreeRowExtensionWidgets(const TSharedRef<ICustomDetailsView>& InDetailsView,
diff --git a/Engine/Plugins/VirtualProduction/DynamicMaterial/Source/DynamicMaterialEditor/Private/UI/Widgets/Editor/SDMMaterialComponentEditor.h b/Engine/Plugins/VirtualProduction/DynamicMaterial/Source/DynamicMaterialEditor/Private/UI/Widgets/Editor/SDMMaterialComponentEditor.h
--- a/Engine/Plugins/VirtualProduction/DynamicMaterial/Source/DynamicMaterialEditor/Private/UI/Widgets/Editor/SDMMaterialComponentEditor.h
+++ b/Engine/Plugins/VirtualProduction/DynamicMaterial/Source/DynamicMaterialEditor/Private/UI/Widgets/Editor/SDMMaterialComponentEditor.h
@@ -33,6 +33,8 @@ protected:
 
 	UDMMaterialComponent* GetOriginalComponent(UDMMaterialComponent* InPreviewComponent) const;
 
+	void BindPropertyRowUpdateDelegates(TConstArrayView<FDMPropertyHandle> InPropertyRows);
+
 	//~ Begin SDMObjectEditorWidgetBase
 	virtual TSharedRef<ICustomDetailsViewItem> GetDefaultCategory(const TSharedRef<ICustomDetailsView>& InDetailsView,
 		const FCustomDetailsViewItemId& InRootId) override;
diff --git a/Engine/Plugins/VirtualProduction/DynamicMaterial/Source/DynamicMaterialEditor/Private/UI/Widgets/Editor/SDMMaterialSlotEditor.cpp b/Engine/Plugins/VirtualProduction/DynamicMaterial/Source/DynamicMaterialEditor/Private/UI/Widgets/Editor/SDMMaterialSlotEditor.cpp
--- a/Engine/Plugins/VirtualProduction/DynamicMaterial/Source/DynamicMaterialEditor/Private/UI/Widgets/Editor/SDMMaterialSlotEditor.cpp
+++ b/Engine/Plugins/VirtualProduction/DynamicMaterial/Source/DynamicMaterialEditor/Private/UI/Widgets/Editor/SDMMaterialSlotEditor.cpp
@@ -27,6 +27,8 @@
 #include "Engine/Texture.h"
 #include "Engine/World.h"
 #include "HAL/PlatformApplicationMisc.h"
+#include "IDynamicMaterialEditorModule.h"
+#include "IDetailPropertyRow.h"
 #include "Items/ICustomDetailsViewItem.h"
 #include "Materials/MaterialFunctionInterface.h"
 #include "Model/DynamicMaterialModel.h"
@@ -55,6 +57,79 @@
 
 #define LOCTEXT_NAMESPACE "SDMMaterialSlotEditor"
 
+namespace UE::DynamicMaterialEditor::Private
+{
+class SDMLayerOpacityEditor : public SDMMaterialComponentEditor
+{
+	SLATE_DECLARE_WIDGET(SDMLayerOpacityEditor, SDMMaterialComponentEditor)
+
+	SLATE_BEGIN_ARGS(SDMLayerOpacityEditor)
+		{}
+	SLATE_END_ARGS()
+
+public:
+	SDMLayerOpacityEditor()
+	{
+		bShowCategories = false;
+	}
+
+	virtual ~SDMLayerOpacityEditor() override = default;
+
+	void Construct(const FArguments& InArgs, const TSharedRef<SDMMaterialEditor>& InEditorWidget, UDMMaterialComponent* InMaterialComponent)
+	{
+		SDMMaterialComponentEditor::Construct(
+			SDMMaterialComponentEditor::FArguments(),
+			InEditorWidget,
+			InMaterialComponent
+		);
+	}
+
+protected:
+	//~ Begin SDMObjectEditorWidgetBase
+	virtual TArray<FDMPropertyHandle> GetPropertyRows() override
+	{
+		TSharedPtr<SDMMaterialEditor> EditorWidget = GetEditorWidget();
+
+		if (!EditorWidget)
+		{
+			return {};
+		}
+
+		TArray<FDMPropertyHandle> PropertyRows;
+		TSet<UObject*> ProcessedObjects;
+
+		UDMMaterialValue* Value = Cast<UDMMaterialValue>(GetComponent());
+
+		FDMComponentPropertyRowGeneratorParams Params(PropertyRows, ProcessedObjects);
+		Params.NotifyHook = this;
+		Params.Owner = this;
+		Params.Object = Value;
+		Params.PreviewMaterialModelBase = EditorWidget->GetPreviewMaterialModelBase();
+		Params.OriginalMaterialModelBase = EditorWidget->GetOriginalMaterialModelBase();
+
+		FDMPropertyHandle& Handle = Params.PropertyRows->Add_GetRef(
+			FDMWidgetLibrary::Get().GetPropertyHandle(Params.CreatePropertyHandleParams(UDMMaterialValue::ValueName))
+		);
+
+		Handle.NameOverride = LOCTEXT("LayerOpacity", "Layer Opacity");
+
+		Handle.ResetToDefaultOverride = FResetToDefaultOverride::Create(
+			FIsResetToDefaultVisible::CreateUObject(Value, &UDMMaterialValue::CanResetToDefault),
+			FResetToDefaultHandler::CreateUObject(Value, &UDMMaterialValue::ResetToDefault)
+		);
+
+		BindPropertyRowUpdateDelegates(PropertyRows);
+
+		return PropertyRows;
+	}
+	//~ End SDMObjectEditorWidgetBase
+};
+
+void SDMLayerOpacityEditor::PrivateRegisterAttributes(FSlateAttributeDescriptor::FInitializer&)
+{
+}
+}
+
 void SDMMaterialSlotEditor::PrivateRegisterAttributes(FSlateAttributeDescriptor::FInitializer&)
 {
 }
@@ -676,116 +751,38 @@ TSharedRef<SWidget> SDMMaterialSlotEditor::CreateSlot_LayerOpacity()
 	LayerOpacityValueWeak.Reset();
 	LayerOpacityItem = nullptr;
 
-	if (LayerViewSlot.IsValid())
+	if (TSharedPtr<SDMMaterialEditor> MaterialEditor = EditorWidgetWeak.Pin())
 	{
-		if (const UDMMaterialLayerObject* SelectedLayer = LayerViewSlot->GetSelectedLayer())
+		if (LayerViewSlot.IsValid())
 		{
-			if (UDMMaterialStage* ValidStage = SelectedLayer->GetFirstValidStage(EDMMaterialLayerStage::All))
+			if (const UDMMaterialLayerObject* SelectedLayer = LayerViewSlot->GetSelectedLayer())
 			{
-				if (UDMMaterialStageInputValue* SelectedOpacityStageInputValue = UDMMaterialStageFunctionLibrary::FindDefaultStageOpacityInputValue(ValidStage))
+				if (UDMMaterialStage* ValidStage = SelectedLayer->GetFirstValidStage(EDMMaterialLayerStage::All))
 				{
-					if (UDMMaterialValueFloat1* OpacityValue = Cast<UDMMaterialValueFloat1>(SelectedOpacityStageInputValue->GetValue()))
+					if (UDMMaterialStageInputValue* SelectedOpacityStageInputValue = UDMMaterialStageFunctionLibrary::FindDefaultStageOpacityInputValue(ValidStage))
 					{
-						LayerOpacityValueWeak = OpacityValue;
-
-						UWorld* World = OpacityValue->GetWorld();
-						TSharedPtr<IDetailKeyframeHandler> KeyframeHandler = nullptr;
-
-						if (!World)
+						if (UDMMaterialValueFloat1* OpacityValue = Cast<UDMMaterialValueFloat1>(SelectedOpacityStageInputValue->GetValue()))
 						{
-							if (TSharedPtr<SDMMaterialEditor> EditorWidget = EditorWidgetWeak.Pin())
-							{
-								if (UDynamicMaterialModelBase* OriginalMaterialModelBase = EditorWidget->GetOriginalMaterialModelBase())
-								{
-									World = OriginalMaterialModelBase->GetWorld();
-								}
-							}
-						}
+							using namespace UE::DynamicMaterialEditor::Private;
 
-						if (World)
-						{
-							if (const UDMWorldSubsystem* const WorldSubsystem = World->GetSubsystem<UDMWorldSubsystem>())
-							{
-								KeyframeHandler = WorldSubsystem->GetKeyframeHandler();
-							}
-						}
-
-						if (!OpacityValue->GetOnUpdate().IsBoundToObject(this))
-						{
-							OpacityValue->GetOnUpdate().AddSP(this, &SDMMaterialSlotEditor::OnOpacityUpdated);
+							LayerOpacityValueWeak = OpacityValue;
+							LayerOpacityItem = SNew(SDMLayerOpacityEditor, MaterialEditor.ToSharedRef(), OpacityValue);
 						}
-
-						TSharedPtr<SDMMaterialEditor> EditorWidget = GetEditorWidget();
-
-						FDMPropertyHandleGenerateParams Params;
-						Params.Widget = this;
-						Params.NotifyHook = this;
-						Params.Object = OpacityValue;
-						Params.PreviewMaterialModelBase = EditorWidget.IsValid() ? EditorWidget->GetPreviewMaterialModelBase() : nullptr;
-						Params.OriginalMaterialModelBase = EditorWidget.IsValid() ? EditorWidget->GetOriginalMaterialModelBase() : nullptr;
-						Params.PropertyName = UDMMaterialValue::ValueName;
-
-						FDMPropertyHandle PropertyHandle = FDMWidgetLibrary::Get().GetPropertyHandle(Params);
-
-						FCustomDetailsViewArgs Args;
-						Args.KeyframeHandler = KeyframeHandler;
-						Args.bAllowGlobalExtensions = true;
-						Args.bAllowResetToDefault = true;
-						Args.bShowCategories = false;
-
-						TSharedRef<ICustomDetailsView> DetailsView = ICustomDetailsViewModule::Get().CreateCustomDetailsView(Args);
-						FCustomDetailsViewItemId RootId = DetailsView->GetRootItem()->GetItemId();
-
-						LayerOpacityItem = DetailsView->CreateDetailTreeItem(DetailsView->GetRootItem(), PropertyHandle.PreviewHandle.DetailTreeNode.ToSharedRef());
-
-						LayerOpacityItem->SetResetToDefaultOverride(FResetToDefaultOverride::Create(
-							FIsResetToDefaultVisible::CreateUObject(OpacityValue, &UDMMaterialValue::CanResetToDefault),
-							FResetToDefaultHandler::CreateUObject(OpacityValue, &UDMMaterialValue::ResetToDefault)
-						));
-
-						LayerOpacityItem->MakeWidget(nullptr, SharedThis(this));
 					}
 				}
 			}
 		}
 	}
 
-	TSharedPtr<SWidget> ValueWidget = LayerOpacityItem.IsValid() ? LayerOpacityItem->GetWidget(ECustomDetailsViewWidgetType::Value) : nullptr;
-	TSharedPtr<SWidget> ExtensionWidget = LayerOpacityItem.IsValid() ? LayerOpacityItem->GetWidget(ECustomDetailsViewWidgetType::Extensions) : nullptr;
-
-	using namespace UE::DynamicMaterialEditor::Private;
-
-	if (ValueWidget.IsValid())
+	if (!LayerOpacityItem.IsValid())
 	{
-		if (TSharedPtr<SWidget> FoundPropertyValueWidget = FDMWidgetLibrary::Get().FindWidgetInHierarchy(ValueWidget.ToSharedRef(), FDMWidgetLibrary::PropertyValueWidget))
-		{
-			if (TSharedPtr<SWidget> InnerPropertyWidget = FDMWidgetLibrary::Get().GetInnerPropertyValueWidget(FoundPropertyValueWidget.ToSharedRef()))
-			{
-				ValueWidget = InnerPropertyWidget;
-			}
-		}
+		return SNullWidget::NullWidget;
 	}
 
-	return SNew(SHorizontalBox)
-		.ToolTipText(LOCTEXT("MaterialDesignerInstanceLayerOpacityTooltip", "Change the Opacity of the selected Material Layer."))
-		
-		+ SHorizontalBox::Slot()
-		.AutoWidth()
-		[
-			SNew(SBox)
-			.HeightOverride(32.f)
-			.HAlign(HAlign_Left)
-			.VAlign(VAlign_Center)
-			.Padding(5.f, 3.f, 5.f, 3.f)
-			[
-				SNew(STextBlock)
-				.Font(IDetailLayoutBuilder::GetDetailFont())
-				.Text(LOCTEXT("MaterialDesignerInstanceLayerOpacity", "Opacity"))
-			]
-		]
 
-		+ SHorizontalBox::Slot()
-		.FillWidth(0.5f)
+	return SNew(SBox)
+		.HAlign(HAlign_Fill)
+		.ToolTipText(LOCTEXT("MaterialDesignerInstanceLayerOpacityTooltip", "Change the Opacity of the selected Material Layer."))
 		[
 			SNew(SBox)
 			.HeightOverride(32.f)
@@ -794,22 +791,7 @@ TSharedRef<SWidget> SDMMaterialSlotEditor::CreateSlot_LayerOpacity()
 			.VAlign(EVerticalAlignment::VAlign_Center)
 			.Padding(0.f, 3.f, 0.f, 3.f)
 			[
-				ValueWidget.IsValid() ? ValueWidget.ToSharedRef() : SNullWidget::NullWidget
-			]
-		]
-
-		+ SHorizontalBox::Slot()
-		.AutoWidth()
-		[
-			SNew(SBox)
-			.WidthOverride(50.f)
-			.HeightOverride(32.f)
-			.HAlign(HAlign_Fill)
-			.VAlign(VAlign_Center)
-			.VAlign(EVerticalAlignment::VAlign_Center)
-			.Padding(0.f, 3.f, 5.f, 3.f)
-			[
-				ExtensionWidget.IsValid() ? ExtensionWidget.ToSharedRef() : SNullWidget::NullWidget
+				LayerOpacityItem.ToSharedRef()
 			]
 		];
 }
diff --git a/Engine/Plugins/VirtualProduction/DynamicMaterial/Source/DynamicMaterialEditor/Private/UI/Widgets/Editor/SDMMaterialSlotEditor.h b/Engine/Plugins/VirtualProduction/DynamicMaterial/Source/DynamicMaterialEditor/Private/UI/Widgets/Editor/SDMMaterialSlotEditor.h
--- a/Engine/Plugins/VirtualProduction/DynamicMaterial/Source/DynamicMaterialEditor/Private/UI/Widgets/Editor/SDMMaterialSlotEditor.h
+++ b/Engine/Plugins/VirtualProduction/DynamicMaterial/Source/DynamicMaterialEditor/Private/UI/Widgets/Editor/SDMMaterialSlotEditor.h
@@ -122,7 +122,7 @@ protected:
 	TDMWidgetSlot<SWidget> LayerSettingsSlot;
 
 	TWeakObjectPtr<UDMMaterialValueFloat1> LayerOpacityValueWeak;
-	TSharedPtr<ICustomDetailsViewItem> LayerOpacityItem;
+	TSharedPtr<SWidget> LayerOpacityItem;
 
 	FOnLayerSelectionChanged OnLayerSelectionChanged;
 	FOnStageSelectionChanged OnStageSelectionChanged;
diff --git a/Engine/Plugins/VirtualProduction/DynamicMaterial/Source/DynamicMaterialEditor/Private/UI/Widgets/Editor/SDMObjectEditorWidgetBase.cpp b/Engine/Plugins/VirtualProduction/DynamicMaterial/Source/DynamicMaterialEditor/Private/UI/Widgets/Editor/SDMObjectEditorWidgetBase.cpp
--- a/Engine/Plugins/VirtualProduction/DynamicMaterial/Source/DynamicMaterialEditor/Private/UI/Widgets/Editor/SDMObjectEditorWidgetBase.cpp
+++ b/Engine/Plugins/VirtualProduction/DynamicMaterial/Source/DynamicMaterialEditor/Private/UI/Widgets/Editor/SDMObjectEditorWidgetBase.cpp
@@ -88,7 +88,8 @@ TSharedRef<SWidget> SDMObjectEditorWidgetBase::CreateWidget()
 	Args.KeyframeHandler = KeyframeHandler;
 	Args.bAllowGlobalExtensions = true;
 	Args.bAllowResetToDefault = true;
-	Args.bShowCategories = true;
+	Args.bShowCategories = bShowCategories;
+	Args.RightColumnMinWidth = 75.f;
 	Args.OnExpansionStateChanged.AddSP(this, &SDMObjectEditorWidgetBase::OnExpansionStateChanged);
 
 	TSharedRef<ICustomDetailsView> DetailsView = ICustomDetailsViewModule::Get().CreateCustomDetailsView(Args);
@@ -122,15 +123,24 @@ TSharedRef<SWidget> SDMObjectEditorWidgetBase::CreateWidget()
 				break;
 		}
 
-		TSharedRef<ICustomDetailsViewItem> CategoryItem = GetCategoryForRow(DetailsView, RootId, PropertyRow);
+		TSharedPtr<ICustomDetailsViewItem> Parent;
+
+		if (bShowCategories)
+		{
+			Parent = GetCategoryForRow(DetailsView, RootId, PropertyRow);
+		}
+		else
+		{
+			Parent = DetailsView->GetRootItem();
+		}
 
 		if (bHasValidCustomWidget)
 		{
-			AddCustomRow(DetailsView, CategoryItem, Position, PropertyRow);
+			AddCustomRow(DetailsView, Parent.ToSharedRef(), Position, PropertyRow);
 		}
 		else if (PropertyRow.PreviewHandle.DetailTreeNode)
 		{
-			AddDetailTreeRow(DetailsView, CategoryItem, Position, PropertyRow);
+			AddDetailTreeRow(DetailsView, Parent.ToSharedRef(), Position, PropertyRow);
 		}		
 	}
 
diff --git a/Engine/Plugins/VirtualProduction/DynamicMaterial/Source/DynamicMaterialEditor/Private/UI/Widgets/Editor/SDMObjectEditorWidgetBase.h b/Engine/Plugins/VirtualProduction/DynamicMaterial/Source/DynamicMaterialEditor/Private/UI/Widgets/Editor/SDMObjectEditorWidgetBase.h
--- a/Engine/Plugins/VirtualProduction/DynamicMaterial/Source/DynamicMaterialEditor/Private/UI/Widgets/Editor/SDMObjectEditorWidgetBase.h
+++ b/Engine/Plugins/VirtualProduction/DynamicMaterial/Source/DynamicMaterialEditor/Private/UI/Widgets/Editor/SDMObjectEditorWidgetBase.h
@@ -60,6 +60,7 @@ protected:
 	bool bConstructing;
 	TArray<FName> Categories;
 	TSharedPtr<ICustomDetailsViewItem> DefaultCategoryItem;
+	bool bShowCategories = true;
 
 	TSharedRef<SWidget> CreateWidget();
 
