Material Designer: Shaders module is now runtime instead of editor.

diff --git a/Engine/Plugins/VirtualProduction/DynamicMaterial/DynamicMaterial.uplugin b/Engine/Plugins/VirtualProduction/DynamicMaterial/DynamicMaterial.uplugin
--- a/Engine/Plugins/VirtualProduction/DynamicMaterial/DynamicMaterial.uplugin
+++ b/Engine/Plugins/VirtualProduction/DynamicMaterial/DynamicMaterial.uplugin
@@ -36,8 +36,8 @@
 			"LoadingPhase": "Default"
 		},
 		{
-			"Name": "DynamicMaterialShadersEditor",
-			"Type": "Editor",
+			"Name": "DynamicMaterialShaders",
+			"Type": "Runtime",
 			"LoadingPhase": "PostConfigInit"
 		}
 	],
diff --git a/Engine/Plugins/VirtualProduction/DynamicMaterial/Source/DynamicMaterialEditor/DynamicMaterialEditor.Build.cs b/Engine/Plugins/VirtualProduction/DynamicMaterial/Source/DynamicMaterialEditor/DynamicMaterialEditor.Build.cs
--- a/Engine/Plugins/VirtualProduction/DynamicMaterial/Source/DynamicMaterialEditor/DynamicMaterialEditor.Build.cs
+++ b/Engine/Plugins/VirtualProduction/DynamicMaterial/Source/DynamicMaterialEditor/DynamicMaterialEditor.Build.cs
@@ -30,7 +30,7 @@ public class DynamicMaterialEditor : ModuleRules
 				"ContentBrowserData",
 				"CustomDetailsView",
 				"DeveloperSettings",
-				"DynamicMaterialShadersEditor",
+				"DynamicMaterialShaders",
 				"DynamicMaterialTextureSet",
 				"DynamicMaterialTextureSetEditor",
 				"EditorSubsystem",
diff --git a/Engine/Plugins/VirtualProduction/DynamicMaterial/Source/DynamicMaterialShaders/DynamicMaterialShaders.Build.cs b/Engine/Plugins/VirtualProduction/DynamicMaterial/Source/DynamicMaterialShaders/DynamicMaterialShaders.Build.cs
new file mode 100644
--- /dev/null
+++ b/Engine/Plugins/VirtualProduction/DynamicMaterial/Source/DynamicMaterialShaders/DynamicMaterialShaders.Build.cs
@@ -0,0 +1,32 @@
+// Copyright Epic Games, Inc. All Rights Reserved.
+
+using UnrealBuildTool;
+
+public class DynamicMaterialShaders : ModuleRules
+{
+	public DynamicMaterialShaders(ReadOnlyTargetRules Target) : base(Target)
+	{
+		PCHUsage = ModuleRules.PCHUsageMode.UseExplicitOrSharedPCHs;
+
+		PublicDependencyModuleNames.AddRange(
+			new string[]
+			{
+				"Core",
+				"CoreUObject",
+				"RenderCore"
+			}
+		);
+
+		PrivateDependencyModuleNames.AddRange(
+			new string[]
+			{
+				"Engine",
+				"Projects",
+				"Renderer",
+				"RHI"
+			}
+		);
+
+		ShortName = "DynMatEdShdrs";
+	}
+}
diff --git a/Engine/Plugins/VirtualProduction/DynamicMaterial/Source/DynamicMaterialShaders/Private/DMAlphaOneMinusPS.cpp b/Engine/Plugins/VirtualProduction/DynamicMaterial/Source/DynamicMaterialShaders/Private/DMAlphaOneMinusPS.cpp
new file mode 100644
--- /dev/null
+++ b/Engine/Plugins/VirtualProduction/DynamicMaterial/Source/DynamicMaterialShaders/Private/DMAlphaOneMinusPS.cpp
@@ -0,0 +1,26 @@
+// Copyright Epic Games, Inc. All Rights Reserved.
+
+#include "DMAlphaOneMinusPS.h"
+#include "DataDrivenShaderPlatformInfo.h"
+#include "DynamicMaterialShadersModule.h"
+#include "RenderGraphBuilder.h"
+#include "RHIStaticStates.h"
+#include "UObject/Package.h"
+
+const FString FDMAlphaOneMinusPS::ShaderPath = FString(UE::DynamicMaterialShaders::Internal::VirtualShaderMountPoint) + TEXT("/DMAlphaOneMinusPS.usf");
+
+bool FDMAlphaOneMinusPS::ShouldCompilePermutation(const FGlobalShaderPermutationParameters& InParameters)
+{
+	return IsFeatureLevelSupported(InParameters.Platform, ERHIFeatureLevel::ES3_1);
+}
+
+FDMAlphaOneMinusPS::FParameters* FDMAlphaOneMinusPS::AllocateAndSetParameters(FRDGBuilder& InGraphBuilder, FRDGTextureRef InRGBATexture, FRDGTextureRef InOutputTexture)
+{
+	FParameters* Parameters = InGraphBuilder.AllocParameters<FParameters>();
+	Parameters->InputTexture = InRGBATexture;
+	Parameters->InputSampler = TStaticSamplerState<SF_Point>::GetRHI();
+	Parameters->RenderTargets[0] = FRenderTargetBinding{InOutputTexture, ERenderTargetLoadAction::ENoAction};
+	return Parameters;
+}
+
+IMPLEMENT_SHADER_TYPE(, FDMAlphaOneMinusPS, *FDMAlphaOneMinusPS::ShaderPath, TEXT("MainPS"), SF_Pixel);
diff --git a/Engine/Plugins/VirtualProduction/DynamicMaterial/Source/DynamicMaterialShaders/Private/DynamicMaterialShadersModule.cpp b/Engine/Plugins/VirtualProduction/DynamicMaterial/Source/DynamicMaterialShaders/Private/DynamicMaterialShadersModule.cpp
new file mode 100644
--- /dev/null
+++ b/Engine/Plugins/VirtualProduction/DynamicMaterial/Source/DynamicMaterialShaders/Private/DynamicMaterialShadersModule.cpp
@@ -0,0 +1,20 @@
+// Copyright Epic Games, Inc. All Rights Reserved.
+
+#include "DynamicMaterialShadersModule.h"
+#include "DMAlphaOneMinusPS.h"
+#include "Interfaces/IPluginManager.h"
+#include "Modules/ModuleManager.h"
+
+void FDynamicMaterialShadersModule::StartupModule()
+{
+	const TSharedPtr<IPlugin> Plugin = IPluginManager::Get().FindPlugin(UE_PLUGIN_NAME);
+	check(Plugin.IsValid());
+
+	const FString PluginShaderDir = FPaths::Combine(Plugin->GetBaseDir(), TEXT("Shaders"));
+	AddShaderSourceDirectoryMapping(UE::DynamicMaterialShaders::Internal::VirtualShaderMountPoint, PluginShaderDir);
+
+	// Force init
+	FDMAlphaOneMinusPS::GetStaticType();
+}
+
+IMPLEMENT_MODULE(FDynamicMaterialShadersModule, DynamicMaterialShaders)
diff --git a/Engine/Plugins/VirtualProduction/DynamicMaterial/Source/DynamicMaterialShaders/Private/DynamicMaterialShadersModule.h b/Engine/Plugins/VirtualProduction/DynamicMaterial/Source/DynamicMaterialShaders/Private/DynamicMaterialShadersModule.h
new file mode 100644
--- /dev/null
+++ b/Engine/Plugins/VirtualProduction/DynamicMaterial/Source/DynamicMaterialShaders/Private/DynamicMaterialShadersModule.h
@@ -0,0 +1,17 @@
+// Copyright Epic Games, Inc. All Rights Reserved.
+
+#pragma once
+
+#include "Modules/ModuleInterface.h"
+
+namespace UE::DynamicMaterialShaders::Internal
+{
+	static inline const TCHAR* VirtualShaderMountPoint = TEXT("/Plugin/MaterialDesigner");
+}
+
+class FDynamicMaterialShadersModule : public IModuleInterface
+{
+	//~ Begin IModuleInterface interface
+	virtual void StartupModule() override;
+	//~ End IModuleInterface interface
+};
diff --git a/Engine/Plugins/VirtualProduction/DynamicMaterial/Source/DynamicMaterialShaders/Public/DMAlphaOneMinusPS.h b/Engine/Plugins/VirtualProduction/DynamicMaterial/Source/DynamicMaterialShaders/Public/DMAlphaOneMinusPS.h
new file mode 100644
--- /dev/null
+++ b/Engine/Plugins/VirtualProduction/DynamicMaterial/Source/DynamicMaterialShaders/Public/DMAlphaOneMinusPS.h
@@ -0,0 +1,29 @@
+// Copyright Epic Games, Inc. All Rights Reserved.
+
+#pragma once
+
+#include "GlobalShader.h"
+#include "ShaderParameterStruct.h"
+
+class FRDGBuilder;
+struct FGlobalShaderPermutationParameters;
+using FRDGTextureRef = class FRDGTexture*;
+
+class FDMAlphaOneMinusPS : public FGlobalShader
+{
+	DECLARE_EXPORTED_GLOBAL_SHADER(FDMAlphaOneMinusPS, DYNAMICMATERIALSHADERS_API);
+	SHADER_USE_PARAMETER_STRUCT(FDMAlphaOneMinusPS, FGlobalShader);
+
+	BEGIN_SHADER_PARAMETER_STRUCT(FParameters, )
+		SHADER_PARAMETER_RDG_TEXTURE(Texture2D, InputTexture)
+		SHADER_PARAMETER_SAMPLER(SamplerState, InputSampler)
+		RENDER_TARGET_BINDING_SLOTS()
+	END_SHADER_PARAMETER_STRUCT()
+
+public:
+	static const FString ShaderPath;
+
+	static DYNAMICMATERIALSHADERS_API bool ShouldCompilePermutation(const FGlobalShaderPermutationParameters& InParameters);
+
+	DYNAMICMATERIALSHADERS_API FParameters* AllocateAndSetParameters(FRDGBuilder& InGraphBuilder, FRDGTextureRef InRGBATexture, FRDGTextureRef InOutputTexture);
+};
diff --git a/Engine/Plugins/VirtualProduction/DynamicMaterial/Source/DynamicMaterialShadersEditor/DynamicMaterialShadersEditor.Build.cs b/Engine/Plugins/VirtualProduction/DynamicMaterial/Source/DynamicMaterialShadersEditor/DynamicMaterialShadersEditor.Build.cs
deleted file mode 100644
diff --git a/Engine/Plugins/VirtualProduction/DynamicMaterial/Source/DynamicMaterialShadersEditor/Private/DMAlphaOneMinusPS.cpp b/Engine/Plugins/VirtualProduction/DynamicMaterial/Source/DynamicMaterialShadersEditor/Private/DMAlphaOneMinusPS.cpp
deleted file mode 100644
diff --git a/Engine/Plugins/VirtualProduction/DynamicMaterial/Source/DynamicMaterialShadersEditor/Private/DynamicMaterialShadersEditorModule.cpp b/Engine/Plugins/VirtualProduction/DynamicMaterial/Source/DynamicMaterialShadersEditor/Private/DynamicMaterialShadersEditorModule.cpp
deleted file mode 100644
diff --git a/Engine/Plugins/VirtualProduction/DynamicMaterial/Source/DynamicMaterialShadersEditor/Private/DynamicMaterialShadersEditorModule.h b/Engine/Plugins/VirtualProduction/DynamicMaterial/Source/DynamicMaterialShadersEditor/Private/DynamicMaterialShadersEditorModule.h
deleted file mode 100644
diff --git a/Engine/Plugins/VirtualProduction/DynamicMaterial/Source/DynamicMaterialShadersEditor/Public/DMAlphaOneMinusPS.h b/Engine/Plugins/VirtualProduction/DynamicMaterial/Source/DynamicMaterialShadersEditor/Public/DMAlphaOneMinusPS.h
deleted file mode 100644
