Material Designer: Text and UMG Renderers and Render Targets now correctly copy their properties from the preview to the source material.

diff --git a/Engine/Plugins/VirtualProduction/DynamicMaterial/Source/DynamicMaterial/Private/Components/MaterialValues/DMMaterialValueRenderTarget.cpp b/Engine/Plugins/VirtualProduction/DynamicMaterial/Source/DynamicMaterial/Private/Components/MaterialValues/DMMaterialValueRenderTarget.cpp
--- a/Engine/Plugins/VirtualProduction/DynamicMaterial/Source/DynamicMaterial/Private/Components/MaterialValues/DMMaterialValueRenderTarget.cpp
+++ b/Engine/Plugins/VirtualProduction/DynamicMaterial/Source/DynamicMaterial/Private/Components/MaterialValues/DMMaterialValueRenderTarget.cpp
@@ -164,6 +164,17 @@ void UDMMaterialValueRenderTarget::Update(UDMMaterialComponent* InSource, EDMUpd
 }
 
 #if WITH_EDITOR
+void UDMMaterialValueRenderTarget::CopyParametersFrom_Implementation(UObject* InOther)
+{
+	// Intentionally not calling Super.
+	// The render target texture should not be copied as it's unique per instance.
+
+	UDMMaterialValueRenderTarget* OtherRenderTarget = CastChecked<UDMMaterialValueRenderTarget>(InOther);
+	OtherRenderTarget->SetTextureSize(TextureSize);
+	OtherRenderTarget->SetTextureFormat(TextureFormat);
+	OtherRenderTarget->SetClearColor(ClearColor);
+}
+
 UDMMaterialValueDynamic* UDMMaterialValueRenderTarget::ToDynamic(UDynamicMaterialModelDynamic* InMaterialModelDynamic)
 {
 	return UDMMaterialValueDynamic::CreateValueDynamic<UDMMaterialValueRenderTargetDynamic>(InMaterialModelDynamic, this);
@@ -346,11 +357,6 @@ void UDMMaterialValueRenderTarget::UpdateRenderTarget()
 	}
 }
 
-void UDMMaterialValueRenderTarget::CopyParametersFrom_Implementation(UObject* InOther)
-{
-	// Overridden to do nothing
-}
-
 UDMMaterialComponent* UDMMaterialValueRenderTarget::GetSubComponentByPath(FDMComponentPath& InPath, const FDMComponentPathSegment& InPathSegment) const
 {
 	if (InPathSegment.GetToken() == RendererPathToken)
diff --git a/Engine/Plugins/VirtualProduction/DynamicMaterial/Source/DynamicMaterial/Private/Components/RenderTargetRenderers/DMRenderTargetTextRenderer.cpp b/Engine/Plugins/VirtualProduction/DynamicMaterial/Source/DynamicMaterial/Private/Components/RenderTargetRenderers/DMRenderTargetTextRenderer.cpp
--- a/Engine/Plugins/VirtualProduction/DynamicMaterial/Source/DynamicMaterial/Private/Components/RenderTargetRenderers/DMRenderTargetTextRenderer.cpp
+++ b/Engine/Plugins/VirtualProduction/DynamicMaterial/Source/DynamicMaterial/Private/Components/RenderTargetRenderers/DMRenderTargetTextRenderer.cpp
@@ -82,6 +82,35 @@ UDMRenderTargetTextRenderer::UDMRenderTargetTextRenderer()
 }
 
 #if WITH_EDITOR
+void UDMRenderTargetTextRenderer::CopyParametersFrom_Implementation(UObject* InOther)
+{
+	UDMRenderTargetTextRenderer* OtherTextRenderer = CastChecked<UDMRenderTargetTextRenderer>(InOther);
+	OtherTextRenderer->SetFontInfo(FontInfo);
+	OtherTextRenderer->SetText(Text);
+	OtherTextRenderer->SetTextColor(TextColor);
+	OtherTextRenderer->SetBackgroundColor(GetBackgroundColor());
+	OtherTextRenderer->SetHasHighlight(bHasHighlight);
+	OtherTextRenderer->SetHighlightColor(HighlightColor);
+	OtherTextRenderer->SetHasShadow(bHasShadow);
+	OtherTextRenderer->SetShadowColor(ShadowColor);
+	OtherTextRenderer->SetShadowOffset(ShadowOffset);
+	OtherTextRenderer->SetAutoWrapText(bAutoWrapText);
+	OtherTextRenderer->SetWrapTextAt(WrapTextAt);
+	OtherTextRenderer->SetWrappingPolicy(WrappingPolicy);
+	OtherTextRenderer->SetJustify(Justify);
+	OtherTextRenderer->SetTransformPolicy(TransformPolicy);
+	OtherTextRenderer->SetFlowDirection(FlowDirection);
+	OtherTextRenderer->SetShapingMethod(ShapingMethod);
+	OtherTextRenderer->SetStrikeBrush(StrikeBrush);
+	OtherTextRenderer->SetLineHeight(LineHeight);
+	OtherTextRenderer->SetPaddingLeft(PaddingLeft);
+	OtherTextRenderer->SetPaddingRight(PaddingRight);
+	OtherTextRenderer->SetPaddingTop(PaddingTop);
+	OtherTextRenderer->SetPaddingBottom(PaddingBottom);
+	OtherTextRenderer->SetOverrideRenderTargetSize(bOverrideRenderTargetSize);
+	OtherTextRenderer->SetStrikeBrush(StrikeBrush);
+}
+
 TSharedPtr<FJsonValue> UDMRenderTargetTextRenderer::JsonSerialize() const
 {
 	TSharedRef<FJsonObject> Object = MakeShared<FJsonObject>();
@@ -263,7 +292,11 @@ void UDMRenderTargetTextRenderer::PostEditChangeProperty(FPropertyChangedEvent&
 		|| PropertyName == FDMRenderTargetTextRenderer::AutoWrapTextName)
 	{
 		// Cause details panel refresh.
-		Update(this, EDMUpdateType::Structure);
+		Update(this, EDMUpdateType::Value | EDMUpdateType::RefreshDetailView);
+	}
+	else
+	{
+		Update(this, EDMUpdateType::Value);
 	}
 
 	if (FDMRenderTargetTextRenderer::PropertyNames.Contains(PropertyName))
diff --git a/Engine/Plugins/VirtualProduction/DynamicMaterial/Source/DynamicMaterial/Private/Components/RenderTargetRenderers/DMRenderTargetUMGWidgetRenderer.cpp b/Engine/Plugins/VirtualProduction/DynamicMaterial/Source/DynamicMaterial/Private/Components/RenderTargetRenderers/DMRenderTargetUMGWidgetRenderer.cpp
--- a/Engine/Plugins/VirtualProduction/DynamicMaterial/Source/DynamicMaterial/Private/Components/RenderTargetRenderers/DMRenderTargetUMGWidgetRenderer.cpp
+++ b/Engine/Plugins/VirtualProduction/DynamicMaterial/Source/DynamicMaterial/Private/Components/RenderTargetRenderers/DMRenderTargetUMGWidgetRenderer.cpp
@@ -28,6 +28,12 @@ void UDMRenderTargetUMGWidgetRenderer::SetWidgetClass(TSubclassOf<UWidget> InWid
 }
 
 #if WITH_EDITOR
+void UDMRenderTargetUMGWidgetRenderer::CopyParametersFrom_Implementation(UObject* InOther)
+{
+	UDMRenderTargetUMGWidgetRenderer* OtherUMGRenderer = CastChecked<UDMRenderTargetUMGWidgetRenderer>(InOther);
+	SetWidgetClass(OtherUMGRenderer->GetWidgetClass());
+}
+
 TSharedPtr<FJsonValue> UDMRenderTargetUMGWidgetRenderer::JsonSerialize() const
 {
 	return FDMJsonUtils::Serialize(WidgetClass.Get());
diff --git a/Engine/Plugins/VirtualProduction/DynamicMaterial/Source/DynamicMaterial/Public/Components/MaterialValues/DMMaterialValueRenderTarget.h b/Engine/Plugins/VirtualProduction/DynamicMaterial/Source/DynamicMaterial/Public/Components/MaterialValues/DMMaterialValueRenderTarget.h
--- a/Engine/Plugins/VirtualProduction/DynamicMaterial/Source/DynamicMaterial/Public/Components/MaterialValues/DMMaterialValueRenderTarget.h
+++ b/Engine/Plugins/VirtualProduction/DynamicMaterial/Source/DynamicMaterial/Public/Components/MaterialValues/DMMaterialValueRenderTarget.h
@@ -62,7 +62,11 @@ public:
 	/** Will trigger the end of frame update if it is currently queued or is invalid */
 	void FlushCreateRenderTarget();
 
-#if WITH_EDITOR
+#if WITH_EDITOR	
+	//~ Begin IDMParameterContainer
+	DYNAMICMATERIAL_API virtual void CopyParametersFrom_Implementation(UObject* InOther) override;
+	//~ End IDMParameterContainer
+
 	//~ Begin IDMJsonSerializable
 	DYNAMICMATERIAL_API virtual TSharedPtr<FJsonValue> JsonSerialize() const override;
 	DYNAMICMATERIAL_API virtual bool JsonDeserialize(const TSharedPtr<FJsonValue>& InJsonValue) override;
@@ -117,10 +121,6 @@ protected:
 
 	void UpdateRenderTarget();
 
-	//~ Begin IDMParameterContainer
-	virtual void CopyParametersFrom_Implementation(UObject* InOther) override;
-	//~ End IDMParameterContainer
-
 	//~ Begin UDMMaterialComponent
 	DYNAMICMATERIAL_API virtual UDMMaterialComponent* GetSubComponentByPath(FDMComponentPath& InPath, const FDMComponentPathSegment& InPathSegment) const override;
 #if WITH_EDITOR
diff --git a/Engine/Plugins/VirtualProduction/DynamicMaterial/Source/DynamicMaterial/Public/Components/MaterialValuesDynamic/DMMaterialValueRenderTargetDynamic.h b/Engine/Plugins/VirtualProduction/DynamicMaterial/Source/DynamicMaterial/Public/Components/MaterialValuesDynamic/DMMaterialValueRenderTargetDynamic.h
--- a/Engine/Plugins/VirtualProduction/DynamicMaterial/Source/DynamicMaterial/Public/Components/MaterialValuesDynamic/DMMaterialValueRenderTargetDynamic.h
+++ b/Engine/Plugins/VirtualProduction/DynamicMaterial/Source/DynamicMaterial/Public/Components/MaterialValuesDynamic/DMMaterialValueRenderTargetDynamic.h
@@ -14,6 +14,10 @@ class UDMMaterialValueRenderTargetDynamic : public UDMMaterialValueDynamic
 public:
 	DYNAMICMATERIAL_API UDMMaterialValueRenderTargetDynamic();
 
+	//~ Begin IDMParameterContainer
+	DYNAMICMATERIAL_API virtual void CopyParametersFrom_Implementation(UObject* InOther) override;
+	//~ End IDMParameterContainer
+
 #if WITH_EDITOR
 	//~ Begin IDMJsonSerializable
 	DYNAMICMATERIAL_API virtual TSharedPtr<FJsonValue> JsonSerialize() const override;
@@ -29,9 +33,4 @@ public:
 	DYNAMICMATERIAL_API virtual void CopyDynamicPropertiesTo(UDMMaterialComponent* InDestinationComponent) const override;
 #endif
 	// ~End UDMMaterialValueDynamic
-
-protected:
-	//~ Begin IDMParameterContainer
-	virtual void CopyParametersFrom_Implementation(UObject* InOther) override;
-	//~ End IDMParameterContainer
 };
diff --git a/Engine/Plugins/VirtualProduction/DynamicMaterial/Source/DynamicMaterial/Public/Components/RenderTargetRenderers/DMRenderTargetTextRenderer.h b/Engine/Plugins/VirtualProduction/DynamicMaterial/Source/DynamicMaterial/Public/Components/RenderTargetRenderers/DMRenderTargetTextRenderer.h
--- a/Engine/Plugins/VirtualProduction/DynamicMaterial/Source/DynamicMaterial/Public/Components/RenderTargetRenderers/DMRenderTargetTextRenderer.h
+++ b/Engine/Plugins/VirtualProduction/DynamicMaterial/Source/DynamicMaterial/Public/Components/RenderTargetRenderers/DMRenderTargetTextRenderer.h
@@ -4,6 +4,7 @@
 
 #include "Components/RenderTargetRenderers/DMRenderTargetWidgetRendererBase.h"
 #include "Delegates/IDelegateInstance.h"
+#include "IDMParameterContainer.h"
 #include "Fonts/FontCache.h"
 #include "Framework/Text/TextLayout.h"
 #include "StructUtils/InstancedStruct.h"
@@ -34,7 +35,7 @@ struct FDMTextLine
  * Renderer that renders an STextBlock widget and exposes all its parameters/properties.
  */
 UCLASS(MinimalAPI, BlueprintType, ClassGroup = "Material Designer", meta = (DisplayName = "Material Designer Render Target Text Renderer"))
-class UDMRenderTargetTextRenderer : public UDMRenderTargetWidgetRendererBase
+class UDMRenderTargetTextRenderer : public UDMRenderTargetWidgetRendererBase, public IDMParameterContainer
 {
 	GENERATED_BODY()
 
@@ -182,6 +183,10 @@ public:
 	DYNAMICMATERIAL_API void SetOverrideRenderTargetSize(bool bInOverride);
 
 #if WITH_EDITOR
+	//~ Begin IDMParameterContainer
+	DYNAMICMATERIAL_API virtual void CopyParametersFrom_Implementation(UObject* InOther) override;
+	//~ End IDMParameterContainer
+
 	//~ Begin IDMJsonSerializable
 	DYNAMICMATERIAL_API virtual TSharedPtr<FJsonValue> JsonSerialize() const override;
 	DYNAMICMATERIAL_API virtual bool JsonDeserialize(const TSharedPtr<FJsonValue>& InJsonValue) override;
diff --git a/Engine/Plugins/VirtualProduction/DynamicMaterial/Source/DynamicMaterial/Public/Components/RenderTargetRenderers/DMRenderTargetUMGWidgetRenderer.h b/Engine/Plugins/VirtualProduction/DynamicMaterial/Source/DynamicMaterial/Public/Components/RenderTargetRenderers/DMRenderTargetUMGWidgetRenderer.h
--- a/Engine/Plugins/VirtualProduction/DynamicMaterial/Source/DynamicMaterial/Public/Components/RenderTargetRenderers/DMRenderTargetUMGWidgetRenderer.h
+++ b/Engine/Plugins/VirtualProduction/DynamicMaterial/Source/DynamicMaterial/Public/Components/RenderTargetRenderers/DMRenderTargetUMGWidgetRenderer.h
@@ -3,13 +3,14 @@
 #pragma once
 
 #include "Components/RenderTargetRenderers/DMRenderTargetWidgetRendererBase.h"
+#include "IDMParameterContainer.h"
 #include "DMRenderTargetUMGWidgetRenderer.generated.h"
 
 class FWidgetRenderer;
 class UWidget;
 
 UCLASS(MinimalAPI, BlueprintType, ClassGroup = "Material Designer", meta = (DisplayName = "Material Designer Render Target Widget Renderer"))
-class UDMRenderTargetUMGWidgetRenderer : public UDMRenderTargetWidgetRendererBase
+class UDMRenderTargetUMGWidgetRenderer : public UDMRenderTargetWidgetRendererBase, public IDMParameterContainer
 {
 	GENERATED_BODY()
 
@@ -23,6 +24,10 @@ public:
 	DYNAMICMATERIAL_API void SetWidgetClass(TSubclassOf<UWidget> InWidgetClass);
 
 #if WITH_EDITOR
+	//~ Begin IDMParameterContainer
+	DYNAMICMATERIAL_API virtual void CopyParametersFrom_Implementation(UObject* InOther) override;
+	//~ End IDMParameterContainer
+
 	//~ Begin IDMJsonSerializable
 	DYNAMICMATERIAL_API virtual TSharedPtr<FJsonValue> JsonSerialize() const override;
 	DYNAMICMATERIAL_API virtual bool JsonDeserialize(const TSharedPtr<FJsonValue>& InJsonValue) override;
