Material Designer: Fixed ability to change structure of material instance

diff --git a/Engine/Plugins/VirtualProduction/DynamicMaterial/Source/DynamicMaterial/Private/Components/MaterialValuesDynamic/DMMaterialValueRenderTargetDynamic.cpp b/Engine/Plugins/VirtualProduction/DynamicMaterial/Source/DynamicMaterial/Private/Components/MaterialValuesDynamic/DMMaterialValueRenderTargetDynamic.cpp
--- a/Engine/Plugins/VirtualProduction/DynamicMaterial/Source/DynamicMaterial/Private/Components/MaterialValuesDynamic/DMMaterialValueRenderTargetDynamic.cpp
+++ b/Engine/Plugins/VirtualProduction/DynamicMaterial/Source/DynamicMaterial/Private/Components/MaterialValuesDynamic/DMMaterialValueRenderTargetDynamic.cpp
@@ -34,6 +34,7 @@ void UDMMaterialValueRenderTargetDynamic::CopyDynamicPropertiesTo(UDMMaterialCom
 void UDMMaterialValueRenderTargetDynamic::CopyParametersFrom_Implementation(UObject* InOther)
 {
 	// Overridden to do nothing
+	// The render target texture should not be copied as it's unique per instance.
 }
 
 void UDMMaterialValueRenderTargetDynamic::SetMIDParameter(UMaterialInstanceDynamic* InMID) const
diff --git a/Engine/Plugins/VirtualProduction/DynamicMaterial/Source/DynamicMaterialEditor/Private/UI/Widgets/Editor/SlotEditor/SDMMaterialLayerBlendMode.cpp b/Engine/Plugins/VirtualProduction/DynamicMaterial/Source/DynamicMaterialEditor/Private/UI/Widgets/Editor/SlotEditor/SDMMaterialLayerBlendMode.cpp
--- a/Engine/Plugins/VirtualProduction/DynamicMaterial/Source/DynamicMaterialEditor/Private/UI/Widgets/Editor/SlotEditor/SDMMaterialLayerBlendMode.cpp
+++ b/Engine/Plugins/VirtualProduction/DynamicMaterial/Source/DynamicMaterialEditor/Private/UI/Widgets/Editor/SlotEditor/SDMMaterialLayerBlendMode.cpp
@@ -32,12 +32,14 @@
 #include "Components/MaterialStageBlends/DMMSBVividLight.h"
 #include "DetailLayoutBuilder.h"
 #include "DynamicMaterialEditorStyle.h"
+#include "Model/DynamicMaterialModel.h"
 #include "ToolMenu.h"
 #include "ToolMenuDelegates.h"
 #include "ToolMenus.h"
 #include "UI/Widgets/Editor/SDMMaterialSlotEditor.h"
 #include "UI/Widgets/Editor/SlotEditor/SDMMaterialSlotLayerItem.h"
 #include "UI/Widgets/Editor/SlotEditor/SDMMaterialSlotLayerView.h"
+#include "UI/Widgets/SDMMaterialEditor.h"
 #include "Widgets/Input/SComboButton.h"
 #include "Widgets/SNullWidget.h"
 #include "Widgets/Text/STextBlock.h"
@@ -139,9 +141,9 @@ void UDMSourceBlendModeContextObject::SetBlendModeWidget(const TSharedPtr<SDMMat
 	BlendModeWidgetWeak = InBlendModeWidget;
 }
 
-void SDMMaterialLayerBlendMode::Construct(const FArguments& InArgs, const TSharedRef<SDMMaterialSlotLayerItem> InSlotEditor)
+void SDMMaterialLayerBlendMode::Construct(const FArguments& InArgs, const TSharedRef<SDMMaterialSlotLayerItem> InLayerItem)
 {
-	LayerItemWidgetWeak = InSlotEditor;
+	LayerItemWidgetWeak = InLayerItem;
 	SelectedItem = InArgs._SelectedItem;
 
 	SetCanTick(false);
@@ -325,14 +327,14 @@ bool SDMMaterialLayerBlendMode::InBlendModeSelected(UClass* InBlendClass)
 
 bool SDMMaterialLayerBlendMode::IsSelectorEnabled() const
 {
-	TSharedPtr<SDMMaterialSlotLayerItem> LayerItemWiget = LayerItemWidgetWeak.Pin();
+	TSharedPtr<SDMMaterialSlotLayerItem> LayerItemWidget = LayerItemWidgetWeak.Pin();
 
-	if (!LayerItemWiget.IsValid())
+	if (!LayerItemWidget.IsValid())
 	{
 		return false;
 	}
 
-	UDMMaterialLayerObject* Layer = LayerItemWiget->GetLayer();
+	UDMMaterialLayerObject* Layer = LayerItemWidget->GetLayer();
 
 	if (!Layer)
 	{
@@ -340,7 +342,29 @@ bool SDMMaterialLayerBlendMode::IsSelectorEnabled() const
 	}
 
 	// We only want to change if the stage is enabled.
-	return !!Layer->GetStage(EDMMaterialLayerStage::Base, /* Check Enabled */ true);
+	if (!Layer->GetStage(EDMMaterialLayerStage::Base, /* Check Enabled */ true))
+	{
+		return false;
+	}
+
+	if (TSharedPtr<SDMMaterialSlotLayerView> SlotLayerView = LayerItemWidget->GetSlotLayerView())
+	{
+		if (TSharedPtr<SDMMaterialSlotEditor> SlotEditorWidget = SlotLayerView->GetSlotEditorWidget())
+		{
+			if (TSharedPtr<SDMMaterialEditor> EditorWidget = SlotEditorWidget->GetEditorWidget())
+			{
+				if (UDynamicMaterialModelBase* ModelBase = EditorWidget->GetOriginalMaterialModelBase())
+				{
+					if (ModelBase->IsA<UDynamicMaterialModel>())
+					{
+						return true;
+					}
+				}
+			}
+		}
+	}
+
+	return false;
 }
 
 void SDMMaterialLayerBlendMode::MakeSourceBlendMenu(UToolMenu* InToolMenu)
diff --git a/Engine/Plugins/VirtualProduction/DynamicMaterial/Source/DynamicMaterialEditor/Private/UI/Widgets/Editor/SlotEditor/SDMMaterialLayerBlendMode.h b/Engine/Plugins/VirtualProduction/DynamicMaterial/Source/DynamicMaterialEditor/Private/UI/Widgets/Editor/SlotEditor/SDMMaterialLayerBlendMode.h
--- a/Engine/Plugins/VirtualProduction/DynamicMaterial/Source/DynamicMaterialEditor/Private/UI/Widgets/Editor/SlotEditor/SDMMaterialLayerBlendMode.h
+++ b/Engine/Plugins/VirtualProduction/DynamicMaterial/Source/DynamicMaterialEditor/Private/UI/Widgets/Editor/SlotEditor/SDMMaterialLayerBlendMode.h
@@ -51,7 +51,7 @@ public:
 
 	virtual ~SDMMaterialLayerBlendMode() override = default;
 
-	void Construct(const FArguments& InArgs, const TSharedRef<SDMMaterialSlotLayerItem> InSlotEditor);
+	void Construct(const FArguments& InArgs, const TSharedRef<SDMMaterialSlotLayerItem> InLayerItem);
 
 protected:
 	static TArray<TStrongObjectPtr<UClass>> SupportedBlendClasses;
