Material Designer: Materials now have a stable generated material name.
- Prevents instances from failing to load when the parent material recompiles.

diff --git a/Engine/Plugins/VirtualProduction/DynamicMaterial/Source/DynamicMaterialEditor/Private/Model/DynamicMaterialModelEditorOnlyData.cpp b/Engine/Plugins/VirtualProduction/DynamicMaterial/Source/DynamicMaterialEditor/Private/Model/DynamicMaterialModelEditorOnlyData.cpp
--- a/Engine/Plugins/VirtualProduction/DynamicMaterial/Source/DynamicMaterialEditor/Private/Model/DynamicMaterialModelEditorOnlyData.cpp
+++ b/Engine/Plugins/VirtualProduction/DynamicMaterial/Source/DynamicMaterialEditor/Private/Model/DynamicMaterialModelEditorOnlyData.cpp
@@ -346,7 +346,7 @@ void UDynamicMaterialModelEditorOnlyData::CreateMaterial()
 		MaterialModel->DynamicMaterial = Cast<UMaterial>(MaterialFactory->FactoryCreateNew(
 			UMaterial::StaticClass(),
 			MaterialModel,
-			NAME_None,
+			TEXT("GeneratedMaterial"),
 			RF_DuplicateTransient | RF_TextExportTransient | RF_Public,
 			nullptr,
 			GWarn
