Material Designer: Now using correct outer for instanced mateirals

diff --git a/Engine/Plugins/VirtualProduction/DynamicMaterial/Source/DynamicMaterialEditor/Private/UI/Widgets/Editor/SDMToolBar.cpp b/Engine/Plugins/VirtualProduction/DynamicMaterial/Source/DynamicMaterialEditor/Private/UI/Widgets/Editor/SDMToolBar.cpp
--- a/Engine/Plugins/VirtualProduction/DynamicMaterial/Source/DynamicMaterialEditor/Private/UI/Widgets/Editor/SDMToolBar.cpp
+++ b/Engine/Plugins/VirtualProduction/DynamicMaterial/Source/DynamicMaterialEditor/Private/UI/Widgets/Editor/SDMToolBar.cpp
@@ -1060,7 +1060,7 @@ FReply SDMToolBar::OnConvertToInstanceClicked()
 
 	UDynamicMaterialInstance* NewInstance = Cast<UDynamicMaterialInstance>(InstanceFactory->FactoryCreateNew(
 		UDynamicMaterialInstance::StaticClass(),
-		OldInstance->GetOuter(),
+		Property.GetOuter(),
 		NAME_None,
 		RF_Transactional,
 		nullptr,
diff --git a/Engine/Plugins/VirtualProduction/DynamicMaterial/Source/DynamicMaterialEditor/Private/UI/Widgets/SDMMaterialEditor.cpp b/Engine/Plugins/VirtualProduction/DynamicMaterial/Source/DynamicMaterialEditor/Private/UI/Widgets/SDMMaterialEditor.cpp
--- a/Engine/Plugins/VirtualProduction/DynamicMaterial/Source/DynamicMaterialEditor/Private/UI/Widgets/SDMMaterialEditor.cpp
+++ b/Engine/Plugins/VirtualProduction/DynamicMaterial/Source/DynamicMaterialEditor/Private/UI/Widgets/SDMMaterialEditor.cpp
@@ -1283,7 +1283,7 @@ TSharedRef<SDMToolBar> SDMMaterialEditor::CreateSlot_ToolBar()
 
 TSharedRef<SDMMaterialGlobalSettingsEditor> SDMMaterialEditor::CreateSlot_GlobalSettingsEditor()
 {
-	return SNew(SDMMaterialGlobalSettingsEditor, SharedThis(this), GetPreviewMaterialModelBase());
+	return SNew(SDMMaterialGlobalSettingsEditor, SharedThis(this), GetOriginalMaterialModelBase());
 }
 
 TSharedRef<SDMMaterialProperties> SDMMaterialEditor::CreateSlot_MaterialProperties()
