Material Designer: Fixed hitch when selecting actors with material designer materials when the ui isn't open.

diff --git a/Engine/Plugins/VirtualProduction/DynamicMaterial/Source/DynamicMaterialEditor/Private/LevelEditor/DMLevelEditorIntegrationInstance.cpp b/Engine/Plugins/VirtualProduction/DynamicMaterial/Source/DynamicMaterialEditor/Private/LevelEditor/DMLevelEditorIntegrationInstance.cpp
--- a/Engine/Plugins/VirtualProduction/DynamicMaterial/Source/DynamicMaterialEditor/Private/LevelEditor/DMLevelEditorIntegrationInstance.cpp
+++ b/Engine/Plugins/VirtualProduction/DynamicMaterial/Source/DynamicMaterialEditor/Private/LevelEditor/DMLevelEditorIntegrationInstance.cpp
@@ -57,9 +57,9 @@ FDMLevelEditorIntegrationInstance::~FDMLevelEditorIntegrationInstance()
 	UnregisterWithTabManager();
 }
 
-const TSharedPtr<SDMMaterialDesigner>& FDMLevelEditorIntegrationInstance::GetMaterialDesigner() const
+TSharedPtr<SDMMaterialDesigner> FDMLevelEditorIntegrationInstance::GetMaterialDesigner() const
 {
-	return MaterialDesigner;
+	return MaterialDesignerWeak.Pin();
 }
 
 TSharedPtr<SDockTab> FDMLevelEditorIntegrationInstance::InvokeTab() const
@@ -149,10 +149,6 @@ FDMLevelEditorIntegrationInstance::FDMLevelEditorIntegrationInstance(const TShar
 {
 	LevelEditorWeak = InLevelEditor;
 
-	MaterialDesigner = StaticCastSharedRef<SDMMaterialDesigner>(
-		FDynamicMaterialEditorModule::CreateEditor(nullptr, nullptr)
-	);
-
 	RegisterSelectionChange();
 	RegisterWithTabManager();
 }
@@ -225,11 +221,21 @@ void FDMLevelEditorIntegrationInstance::RegisterWithTabManager()
 		FOnSpawnTab::CreateLambda(
 			[this](const FSpawnTabArgs&)
 			{
+				TSharedRef<SDMMaterialDesigner> NewMaterialDesigner = StaticCastSharedRef<SDMMaterialDesigner>(
+					FDynamicMaterialEditorModule::CreateEditor(nullptr, nullptr)
+				);
+				MaterialDesignerWeak = NewMaterialDesigner;
+
+				if (AActor* LastSelectedActor = LastSelectedActorWeak.Get())
+				{
+					NewMaterialDesigner->OnActorSelected(LastSelectedActor);
+				}
+
 				return SNew(SDockTab)
 					.TabRole(ETabRole::PanelTab)
 					.Content()
 					[
-						MaterialDesigner.ToSharedRef()
+						NewMaterialDesigner
 					];
 			}
 		)
@@ -287,11 +293,23 @@ void FDMLevelEditorIntegrationInstance::OnActorSelected(AActor* InActor)
 		return;
 	}
 
-	MaterialDesigner->OnActorSelected(InActor);
+	LastSelectedActorWeak = InActor;
+
+	if (TSharedPtr<SDMMaterialDesigner> MaterialDesigner = GetMaterialDesigner())
+	{
+		MaterialDesigner->OnActorSelected(InActor);
+	}
 }
 
 void FDMLevelEditorIntegrationInstance::OnObjectSelectionChanged(const UTypedElementSelectionSet* InSelectionSet)
 {
+	TSharedPtr<SDMMaterialDesigner> MaterialDesigner = GetMaterialDesigner();
+
+	if (!MaterialDesigner.IsValid())
+	{
+		return;
+	}
+
 	UDynamicMaterialModelBase* NewSelectedMaterialModelBase = nullptr;
 
 	for (UDynamicMaterialModelBase* MaterialModelBase : InSelectionSet->GetSelectedObjects<UDynamicMaterialModelBase>())
diff --git a/Engine/Plugins/VirtualProduction/DynamicMaterial/Source/DynamicMaterialEditor/Private/LevelEditor/DMLevelEditorIntegrationInstance.h b/Engine/Plugins/VirtualProduction/DynamicMaterial/Source/DynamicMaterialEditor/Private/LevelEditor/DMLevelEditorIntegrationInstance.h
--- a/Engine/Plugins/VirtualProduction/DynamicMaterial/Source/DynamicMaterialEditor/Private/LevelEditor/DMLevelEditorIntegrationInstance.h
+++ b/Engine/Plugins/VirtualProduction/DynamicMaterial/Source/DynamicMaterialEditor/Private/LevelEditor/DMLevelEditorIntegrationInstance.h
@@ -28,7 +28,7 @@ public:
 
 	~FDMLevelEditorIntegrationInstance();
 
-	const TSharedPtr<SDMMaterialDesigner>& GetMaterialDesigner() const;
+	TSharedPtr<SDMMaterialDesigner> GetMaterialDesigner() const;
 
 	TSharedPtr<SDockTab> InvokeTab() const;
 
@@ -44,8 +44,9 @@ private:
 	TWeakPtr<ILevelEditor> LevelEditorWeak;
 	TWeakObjectPtr<UTypedElementSelectionSet> ActorSelectionSetWeak;
 	TWeakObjectPtr<UTypedElementSelectionSet> ObjectSelectionSetWeak;
-	TSharedPtr<SDMMaterialDesigner> MaterialDesigner;
+	TWeakPtr<SDMMaterialDesigner> MaterialDesignerWeak;
 	FString LastOpenAssetPartialPath;
+	TWeakObjectPtr<AActor> LastSelectedActorWeak;
 
 	FDMLevelEditorIntegrationInstance(const TSharedRef<ILevelEditor>& InLevelEditor);
 
