Material Designer: Fixed when crash loading values with classes that are not currently enabled.

diff --git a/Engine/Plugins/VirtualProduction/DynamicMaterial/Source/DynamicMaterialEditor/Private/Components/DMMaterialStageThroughput.cpp b/Engine/Plugins/VirtualProduction/DynamicMaterial/Source/DynamicMaterialEditor/Private/Components/DMMaterialStageThroughput.cpp
--- a/Engine/Plugins/VirtualProduction/DynamicMaterial/Source/DynamicMaterialEditor/Private/Components/DMMaterialStageThroughput.cpp
+++ b/Engine/Plugins/VirtualProduction/DynamicMaterial/Source/DynamicMaterialEditor/Private/Components/DMMaterialStageThroughput.cpp
@@ -12,6 +12,7 @@
 #include "Components/MaterialStageInputs/DMMSIValue.h"
 #include "Containers/ArrayView.h"
 #include "DMDefs.h"
+#include "DynamicMaterialEditorModule.h"
 #include "Materials/MaterialExpression.h"
 #include "Materials/MaterialExpressionAppendVector.h"
 #include "Materials/MaterialExpressionComponentMask.h"
@@ -19,6 +20,7 @@
 #include "Materials/MaterialExpressionConstant2Vector.h"
 #include "Materials/MaterialExpressionConstant3Vector.h"
 #include "Materials/MaterialExpressionMaterialFunctionCall.h"
+#include "Materials/MaterialExpressionScalarParameter.h"
 #include "Materials/MaterialExpressionTextureBase.h"
 #include "Model/DMMaterialBuildState.h"
 #include "Model/DMMaterialBuildUtils.h"
@@ -759,8 +761,22 @@ int32 UDMMaterialStageThroughput::ResolveInputChannel(const TSharedRef<FDMMateri
 		InputValue->GenerateExpressions(InBuildState);
 
 		OutExpressions.Append(InBuildState->GetStageSourceExpressions(InputValue));
-		check(!OutExpressions.IsEmpty());
-		LastExpression = OutExpressions.Last();
+		
+		if (OutExpressions.IsEmpty())
+		{
+			UE_LOG(LogDynamicMaterialEditor, Error, TEXT("No expressions generated!"));
+
+			UMaterialExpressionScalarParameter* NewExpression = InBuildState->GetBuildUtils().CreateExpression<UMaterialExpressionScalarParameter>(
+				UE_DM_NodeComment_Default
+			);
+
+			NewExpression->DefaultValue = 0.f;
+			LastExpression = NewExpression;
+		}
+		else
+		{
+			LastExpression = OutExpressions.Last();
+		}
 
 		InputSourceOutputConnectors = &(InputValue->GetOutputConnectors());
 		check(InputSourceOutputConnectors->IsValidIndex(OutChannel.OutputIndex));
diff --git a/Engine/Plugins/VirtualProduction/DynamicMaterial/Source/DynamicMaterialEditor/Private/Components/MaterialStageInputs/DMMSIValue.cpp b/Engine/Plugins/VirtualProduction/DynamicMaterial/Source/DynamicMaterialEditor/Private/Components/MaterialStageInputs/DMMSIValue.cpp
--- a/Engine/Plugins/VirtualProduction/DynamicMaterial/Source/DynamicMaterialEditor/Private/Components/MaterialStageInputs/DMMSIValue.cpp
+++ b/Engine/Plugins/VirtualProduction/DynamicMaterial/Source/DynamicMaterialEditor/Private/Components/MaterialStageInputs/DMMSIValue.cpp
@@ -475,7 +475,11 @@ void UDMMaterialStageInputValue::GenerateExpressions(const TSharedRef<FDMMateria
 	{
 		return;
 	}
-	check(Value);
+	
+	if (!Value)
+	{
+		return;
+	}
 
 	Value->GenerateExpression(InBuildState);
 	InBuildState->AddStageSourceExpressions(this, InBuildState->GetValueExpressions(Value));
