Material Designer: Texture edge color fixes
- Texture edge color Location input renamed Pixel UV.
- Pixel UV is now always visible in the details panel.
- Changing the Pixel UV from a standard location will now set the Edge Location to custom.
- Texture Sample for edge color sampling now uses the global wrap uv mode to enable true edge reading.
- Inputs in the details panel now use their input name over the owning class's display name.
- Applying values to the source component now dirties the level.

diff --git a/Engine/Plugins/VirtualProduction/DynamicMaterial/Source/DynamicMaterialEditor/Private/Components/MaterialStageExpressions/DMMSETextureSampleEdgeColor.cpp b/Engine/Plugins/VirtualProduction/DynamicMaterial/Source/DynamicMaterialEditor/Private/Components/MaterialStageExpressions/DMMSETextureSampleEdgeColor.cpp
--- a/Engine/Plugins/VirtualProduction/DynamicMaterial/Source/DynamicMaterialEditor/Private/Components/MaterialStageExpressions/DMMSETextureSampleEdgeColor.cpp
+++ b/Engine/Plugins/VirtualProduction/DynamicMaterial/Source/DynamicMaterialEditor/Private/Components/MaterialStageExpressions/DMMSETextureSampleEdgeColor.cpp
@@ -12,11 +12,6 @@
 
 #define LOCTEXT_NAMESPACE "DMMaterialStageExpressionTextureSample"
 
-namespace UE::DynamicMaterialEditor::TextureSampleEdgeColor::Private
-{
-	constexpr float EdgeGap = 0.025;
-}
-
 UDMMaterialStageExpressionTextureSampleEdgeColor::UDMMaterialStageExpressionTextureSampleEdgeColor()
 	: UDMMaterialStageExpression(
 		LOCTEXT("EdgeColor", "Edge Color"),
@@ -27,9 +22,7 @@ UDMMaterialStageExpressionTextureSampleEdgeColor::UDMMaterialStageExpressionText
 	bAllowNestedInputs = true;
 
 	InputConnectors.Add({1, LOCTEXT("Texture", "Texture"), EDMValueType::VT_Texture});
-	InputConnectors.Add({0, LOCTEXT("Location", "Location"), EDMValueType::VT_Float2});
-
-	using namespace UE::DynamicMaterialEditor::TextureSampleEdgeColor::Private;
+	InputConnectors.Add({0, LOCTEXT("PixelUV", "Pixel UV"), EDMValueType::VT_Float2});
 
 	EditableProperties.Add(GET_MEMBER_NAME_CHECKED(UDMMaterialStageExpressionTextureSampleEdgeColor, EdgeLocation));
 	EdgeLocation = EDMEdgeLocation::TopLeft;
@@ -42,16 +35,6 @@ bool UDMMaterialStageExpressionTextureSampleEdgeColor::CanChangeInputType(int32
 	return false;
 }
 
-bool UDMMaterialStageExpressionTextureSampleEdgeColor::IsInputVisible(int32 InInputIndex) const
-{
-	if (InInputIndex == 1 && EdgeLocation != EDMEdgeLocation::Custom)
-	{
-		return false;
-	}
-
-	return Super::IsInputVisible(InInputIndex);
-}
-
 void UDMMaterialStageExpressionTextureSampleEdgeColor::AddDefaultInput(int32 InInputIndex) const
 {
 	UDMMaterialStage* Stage = GetStage();
@@ -59,8 +42,6 @@ void UDMMaterialStageExpressionTextureSampleEdgeColor::AddDefaultInput(int32 InI
 
 	if (InInputIndex == 1)
 	{
-		using namespace UE::DynamicMaterialEditor::TextureSampleEdgeColor::Private;
-
 		UDMMaterialStageInputValue* InputValue = UDMMaterialStageInputValue::ChangeStageInput_NewLocalValue(Stage, 
 			1, FDMMaterialStageConnectorChannel::WHOLE_CHANNEL,
 			EDMValueType::VT_Float2, FDMMaterialStageConnectorChannel::WHOLE_CHANNEL);
@@ -69,7 +50,7 @@ void UDMMaterialStageExpressionTextureSampleEdgeColor::AddDefaultInput(int32 InI
 		UDMMaterialValueFloat2* InputFloat2 = Cast<UDMMaterialValueFloat2>(InputValue->GetValue());
 		check(InputFloat2);
 
-		InputFloat2->SetDefaultValue(FVector2D(EdgeGap, EdgeGap));
+		InputFloat2->SetDefaultValue(FVector2D::ZeroVector);
 		InputFloat2->ApplyDefaultValue();
 		return;
 	}
@@ -77,6 +58,147 @@ void UDMMaterialStageExpressionTextureSampleEdgeColor::AddDefaultInput(int32 InI
 	return Super::AddDefaultInput(InInputIndex);
 }
 
+void UDMMaterialStageExpressionTextureSampleEdgeColor::AddExpressionProperties(const TArray<UMaterialExpression*>& InExpressions) const
+{
+	Super::AddExpressionProperties(InExpressions);
+
+	if (InExpressions.IsEmpty())
+	{
+		return;
+	}
+
+	UMaterialExpressionTextureSample* TextureSample = Cast<UMaterialExpressionTextureSample>(InExpressions[0]);
+
+	if (!TextureSample)
+	{
+		return;
+	}
+
+	TextureSample->SamplerSource = ESamplerSourceMode::SSM_Clamp_WorldGroupSettings;
+}
+
+void UDMMaterialStageExpressionTextureSampleEdgeColor::Update(UDMMaterialComponent* InSource, EDMUpdateType InUpdateType)
+{
+	UDMMaterialStage* Stage = GetStage();
+	check(Stage);
+	check(Stage->GetInputs().IsValidIndex(1));
+
+	UDMMaterialStageInputValue* InputValue = CastChecked<UDMMaterialStageInputValue>(Stage->GetInputs()[1]);
+	UDMMaterialValueFloat2* Float2 = CastChecked<UDMMaterialValueFloat2>(InputValue->GetValue());
+
+	const FVector2D& Float2Value = Float2->GetValue();
+
+	const bool bLeftX = FMath::IsNearlyEqual(Float2Value.X, 0.f);
+	const bool bCenterX = FMath::IsNearlyEqual(Float2Value.X, 0.5f);
+	const bool bRightX = FMath::IsNearlyEqual(Float2Value.X, 1.f);
+
+	const bool bTopY = FMath::IsNearlyEqual(Float2Value.Y, 0.f);
+	const bool bCenterY = FMath::IsNearlyEqual(Float2Value.Y, 0.5f);
+	const bool bBottomY = FMath::IsNearlyEqual(Float2Value.Y, 1.f);
+
+	const bool bIsValidX = bLeftX || bCenterX || bRightX;
+	const bool bIsValidY = bTopY || bCenterY || bBottomY;
+
+	if (!bIsValidX || !bIsValidY)
+	{
+		EdgeLocation = EDMEdgeLocation::Custom;
+
+		Super::Update(InSource, InUpdateType | EDMUpdateType::Value | EDMUpdateType::RefreshDetailView);
+		return;
+	}
+
+	bool bChangedEdgeLocation = false;
+
+	if (bTopY && bLeftX)
+	{
+		if (EdgeLocation != EDMEdgeLocation::TopLeft)
+		{
+			bChangedEdgeLocation = true;
+			EdgeLocation = EDMEdgeLocation::TopLeft;
+		}
+	}
+	else if (bTopY && bCenterX)
+	{
+		if (EdgeLocation != EDMEdgeLocation::Top)
+		{
+			bChangedEdgeLocation = true;
+			EdgeLocation = EDMEdgeLocation::Top;
+		}
+	}
+	else if (bTopY && bRightX)
+	{
+		if (EdgeLocation != EDMEdgeLocation::TopRight)
+		{
+			bChangedEdgeLocation = true;
+			EdgeLocation = EDMEdgeLocation::TopRight;
+		}
+	}
+	else if (bCenterY && bLeftX)
+	{
+		if (EdgeLocation != EDMEdgeLocation::Left)
+		{
+			bChangedEdgeLocation = true;
+			EdgeLocation = EDMEdgeLocation::Left;
+		}
+	}
+	else if (bCenterY && bCenterX)
+	{
+		if (EdgeLocation != EDMEdgeLocation::Center)
+		{
+			bChangedEdgeLocation = true;
+			EdgeLocation = EDMEdgeLocation::Center;
+		}
+	}
+	else if (bCenterY && bRightX)
+	{
+		if (EdgeLocation != EDMEdgeLocation::Right)
+		{
+			bChangedEdgeLocation = true;
+			EdgeLocation = EDMEdgeLocation::Right;
+		}
+	}
+	else if (bBottomY && bLeftX)
+	{
+		if (EdgeLocation != EDMEdgeLocation::BottomLeft)
+		{
+			bChangedEdgeLocation = true;
+			EdgeLocation = EDMEdgeLocation::BottomLeft;
+		}
+	}
+	else if (bBottomY && bCenterX)
+	{
+		if (EdgeLocation != EDMEdgeLocation::Bottom)
+		{
+			bChangedEdgeLocation = true;
+			EdgeLocation = EDMEdgeLocation::Bottom;
+		}
+	}
+	else if (bBottomY && bRightX)
+	{
+		if (EdgeLocation != EDMEdgeLocation::BottomRight)
+		{
+			bChangedEdgeLocation = true;
+			EdgeLocation = EDMEdgeLocation::BottomRight;
+		}
+	}
+
+	if (bChangedEdgeLocation)
+	{
+		if (PreEditEdgeLocation == EDMEdgeLocation::Custom)
+		{
+			Super::Update(InSource, InUpdateType | EDMUpdateType::Value | EDMUpdateType::RefreshDetailView);
+		}
+		else
+		{
+			Super::Update(InSource, InUpdateType | EDMUpdateType::Value);
+		}
+
+		return;
+	}
+
+	Super::Update(InSource, InUpdateType);
+}
+
 void UDMMaterialStageExpressionTextureSampleEdgeColor::PreEditChange(FEditPropertyChain& InPropertyAboutToChange)
 {
 	PreEditEdgeLocation = EdgeLocation;
@@ -94,40 +216,44 @@ void UDMMaterialStageExpressionTextureSampleEdgeColor::PostEditChangeProperty(FP
 	}
 }
 
-void UDMMaterialStageExpressionTextureSampleEdgeColor::OnEdgeLocationChanged()
+void UDMMaterialStageExpressionTextureSampleEdgeColor::CopyParametersFrom_Implementation(UObject* InOther)
 {
-	if (EdgeLocation != EDMEdgeLocation::Custom)
-	{
-		using namespace UE::DynamicMaterialEditor::TextureSampleEdgeColor::Private;
+	UDMMaterialStageExpressionTextureSampleEdgeColor* OtherEdgeColor = CastChecked<UDMMaterialStageExpressionTextureSampleEdgeColor>(InOther);
+	OtherEdgeColor->EdgeLocation = EdgeLocation;
+}
 
-		UDMMaterialStage* Stage = GetStage();
-		check(Stage);
-		check(Stage->GetInputs().IsValidIndex(1));
+void UDMMaterialStageExpressionTextureSampleEdgeColor::OnEdgeLocationChanged()
+{
+	UDMMaterialStage* Stage = GetStage();
+	check(Stage);
+	check(Stage->GetInputs().IsValidIndex(1));
 
-		UDMMaterialStageInputValue* InputValue = CastChecked<UDMMaterialStageInputValue>(Stage->GetInputs()[1]);
-		UDMMaterialValueFloat2* Float2 = CastChecked<UDMMaterialValueFloat2>(InputValue->GetValue());
+	UDMMaterialStageInputValue* InputValue = CastChecked<UDMMaterialStageInputValue>(Stage->GetInputs()[1]);
+	UDMMaterialValueFloat2* Float2 = CastChecked<UDMMaterialValueFloat2>(InputValue->GetValue());
 
+	if (EdgeLocation != EDMEdgeLocation::Custom)
+	{
 		FVector2D Float2Value = FVector2D::ZeroVector;
 
 		switch (EdgeLocation)
 		{
 			case EDMEdgeLocation::TopLeft:
-				Float2Value.X = EdgeGap;
-				Float2Value.Y = EdgeGap;
+				Float2Value.X = 0;
+				Float2Value.Y = 0;
 				break;
 
 			case EDMEdgeLocation::Top:
 				Float2Value.X = 0.5;
-				Float2Value.Y = EdgeGap;
+				Float2Value.Y = 0;
 				break;
 
 			case EDMEdgeLocation::TopRight:
-				Float2Value.X = 1.f - EdgeGap;
-				Float2Value.Y = EdgeGap;
+				Float2Value.X = 1.f;
+				Float2Value.Y = 0;
 				break;
 
 			case EDMEdgeLocation::Left:
-				Float2Value.X = EdgeGap;
+				Float2Value.X = 0;
 				Float2Value.Y = 0.5;
 				break;
 
@@ -137,23 +263,23 @@ void UDMMaterialStageExpressionTextureSampleEdgeColor::OnEdgeLocationChanged()
 				break;
 
 			case EDMEdgeLocation::Right:
-				Float2Value.X = 1.f - EdgeGap;
+				Float2Value.X = 1.f;
 				Float2Value.Y = 0.5;
 				break;
 
 			case EDMEdgeLocation::BottomLeft:
-				Float2Value.X = EdgeGap;
-				Float2Value.Y = 1.f - EdgeGap;
+				Float2Value.X = 0;
+				Float2Value.Y = 1.f;
 				break;
 
 			case EDMEdgeLocation::Bottom:
 				Float2Value.X = 0.5;
-				Float2Value.Y = 1.f - EdgeGap;
+				Float2Value.Y = 1.f;
 				break;
 
 			case EDMEdgeLocation::BottomRight:
-				Float2Value.X = 1.f - EdgeGap;
-				Float2Value.Y = 1.f - EdgeGap;
+				Float2Value.X = 1.f;
+				Float2Value.Y = 1.f;
 				break;
 
 			default:
@@ -162,11 +288,16 @@ void UDMMaterialStageExpressionTextureSampleEdgeColor::OnEdgeLocationChanged()
 		}
 
 		Float2->SetValue(Float2Value);
+		Float2->Update(this, EDMUpdateType::Value);
 	}
 
 	if (EdgeLocation == EDMEdgeLocation::Custom || PreEditEdgeLocation == EDMEdgeLocation::Custom)
 	{
-		Update(this, EDMUpdateType::RefreshDetailView);
+		Update(this, EDMUpdateType::Value | EDMUpdateType::RefreshDetailView);
+	}
+	else
+	{
+		Update(this, EDMUpdateType::Value);
 	}
 }
 
diff --git a/Engine/Plugins/VirtualProduction/DynamicMaterial/Source/DynamicMaterialEditor/Private/UI/PropertyGenerators/DMThroughputPropertyRowGenerator.cpp b/Engine/Plugins/VirtualProduction/DynamicMaterial/Source/DynamicMaterialEditor/Private/UI/PropertyGenerators/DMThroughputPropertyRowGenerator.cpp
--- a/Engine/Plugins/VirtualProduction/DynamicMaterial/Source/DynamicMaterialEditor/Private/UI/PropertyGenerators/DMThroughputPropertyRowGenerator.cpp
+++ b/Engine/Plugins/VirtualProduction/DynamicMaterial/Source/DynamicMaterialEditor/Private/UI/PropertyGenerators/DMThroughputPropertyRowGenerator.cpp
@@ -113,13 +113,13 @@ void FDMThroughputPropertyRowGenerator::AddComponentProperties(FDMComponentPrope
 					{
 						PropertyRow.NameOverride = FText::FromName(PropertyRow.ValueName);
 					}
-					else if (PropertyRow.PreviewHandle.PropertyHandle.IsValid())
+					else if (!InputConnectors[InputIdx].Name.IsEmpty())
 					{
-						PropertyRow.NameOverride = PropertyRow.PreviewHandle.PropertyHandle->GetPropertyDisplayName();
+						PropertyRow.NameOverride = InputConnectors[InputIdx].Name;
 					}
-					else
+					else if (PropertyRow.PreviewHandle.PropertyHandle.IsValid())
 					{
-						PropertyRow.NameOverride = InputConnectors[InputIdx].Name;
+						PropertyRow.NameOverride = PropertyRow.PreviewHandle.PropertyHandle->GetPropertyDisplayName();
 					}
 				}
 			}
diff --git a/Engine/Plugins/VirtualProduction/DynamicMaterial/Source/DynamicMaterialEditor/Private/UI/Widgets/Editor/SDMMaterialComponentEditor.cpp b/Engine/Plugins/VirtualProduction/DynamicMaterial/Source/DynamicMaterialEditor/Private/UI/Widgets/Editor/SDMMaterialComponentEditor.cpp
--- a/Engine/Plugins/VirtualProduction/DynamicMaterial/Source/DynamicMaterialEditor/Private/UI/Widgets/Editor/SDMMaterialComponentEditor.cpp
+++ b/Engine/Plugins/VirtualProduction/DynamicMaterial/Source/DynamicMaterialEditor/Private/UI/Widgets/Editor/SDMMaterialComponentEditor.cpp
@@ -126,6 +126,11 @@ void SDMMaterialComponentEditor::OnComponentUpdated(UDMMaterialComponent* InComp
 				{
 					if (UDMMaterialComponent* OriginalComponent = GetOriginalComponent(InComponent))
 					{
+						if (GUndo)
+						{
+							OriginalComponent->Modify();
+						}
+
 						IDMParameterContainer::CopyParametersBetween(InComponent, OriginalComponent);
 						return;
 					}
diff --git a/Engine/Plugins/VirtualProduction/DynamicMaterial/Source/DynamicMaterialEditor/Public/Components/MaterialStageExpressions/DMMSETextureSampleEdgeColor.h b/Engine/Plugins/VirtualProduction/DynamicMaterial/Source/DynamicMaterialEditor/Public/Components/MaterialStageExpressions/DMMSETextureSampleEdgeColor.h
--- a/Engine/Plugins/VirtualProduction/DynamicMaterial/Source/DynamicMaterialEditor/Public/Components/MaterialStageExpressions/DMMSETextureSampleEdgeColor.h
+++ b/Engine/Plugins/VirtualProduction/DynamicMaterial/Source/DynamicMaterialEditor/Public/Components/MaterialStageExpressions/DMMSETextureSampleEdgeColor.h
@@ -3,6 +3,7 @@
 #pragma once
 
 #include "Components/DMMaterialStageExpression.h"
+#include "IDMParameterContainer.h"
 #include "DMMSETextureSampleEdgeColor.generated.h"
 
 UENUM(BlueprintType)
@@ -21,7 +22,7 @@ enum class EDMEdgeLocation : uint8
 };
 
 UCLASS(MinimalAPI, BlueprintType, ClassGroup = "Material Designer")
-class UDMMaterialStageExpressionTextureSampleEdgeColor : public UDMMaterialStageExpression
+class UDMMaterialStageExpressionTextureSampleEdgeColor : public UDMMaterialStageExpression, public IDMParameterContainer
 {
 	GENERATED_BODY()
 
@@ -30,15 +31,26 @@ public:
 
 	//~ Begin UDMMaterialStageThroughput
 	DYNAMICMATERIALEDITOR_API virtual bool CanChangeInputType(int32 InInputIndex) const override;
-	DYNAMICMATERIALEDITOR_API virtual bool IsInputVisible(int32 InInputIndex) const override;
 	DYNAMICMATERIALEDITOR_API virtual void AddDefaultInput(int32 InInputIndex) const override;
 	//~ End UDMMaterialStageThroughput
 
+	//~ Begin UDMMaterialStageSource
+	DYNAMICMATERIALEDITOR_API virtual void AddExpressionProperties(const TArray<UMaterialExpression*>& InExpressions) const override;
+	//~ End UDMMaterialStageSource
+
+	//~ Begin UDMMaterialComponent
+	DYNAMICMATERIALEDITOR_API virtual void Update(UDMMaterialComponent* InSource, EDMUpdateType InUpdateType) override;
+	//~ End UDMMaterialComponent
+
 	//~ Begin UObject
 	DYNAMICMATERIALEDITOR_API virtual void PreEditChange(FEditPropertyChain& InPropertyAboutToChange) override;
 	DYNAMICMATERIALEDITOR_API virtual void PostEditChangeProperty(FPropertyChangedEvent& InPropertyChangedEvent) override;
 	//~ End UObject
 
+	//~ Begin IDMParameterContainer
+	virtual void CopyParametersFrom_Implementation(UObject* InOther) override;
+	//~ End IDMParameterContainer
+
 protected:
 	UPROPERTY(EditInstanceOnly, BlueprintReadOnly, Category = "Material Designer", meta=(AllowPrivateAccess = "true"))
 	EDMEdgeLocation EdgeLocation;
