Material Designer: Switching to a Texture Edge Color will no longer cause a crash.

diff --git a/Engine/Plugins/VirtualProduction/DynamicMaterial/Source/DynamicMaterialEditor/Private/Components/MaterialStageExpressions/DMMSETextureSampleEdgeColor.cpp b/Engine/Plugins/VirtualProduction/DynamicMaterial/Source/DynamicMaterialEditor/Private/Components/MaterialStageExpressions/DMMSETextureSampleEdgeColor.cpp
--- a/Engine/Plugins/VirtualProduction/DynamicMaterial/Source/DynamicMaterialEditor/Private/Components/MaterialStageExpressions/DMMSETextureSampleEdgeColor.cpp
+++ b/Engine/Plugins/VirtualProduction/DynamicMaterial/Source/DynamicMaterialEditor/Private/Components/MaterialStageExpressions/DMMSETextureSampleEdgeColor.cpp
@@ -81,7 +81,12 @@ void UDMMaterialStageExpressionTextureSampleEdgeColor::Update(UDMMaterialCompone
 {
 	UDMMaterialStage* Stage = GetStage();
 	check(Stage);
-	check(Stage->GetInputs().IsValidIndex(1));
+	
+	if (!Stage->GetInputs().IsValidIndex(1))
+	{
+		Super::Update(InSource, InUpdateType);
+		return;
+	}
 
 	UDMMaterialStageInputValue* InputValue = CastChecked<UDMMaterialStageInputValue>(Stage->GetInputs()[1]);
 	UDMMaterialValueFloat2* Float2 = CastChecked<UDMMaterialValueFloat2>(InputValue->GetValue());
