Run UnrealCodeFixup to add #include UE_INLINE_GENERATED_CPP_BY_NAME to files where possible

diff --git a/Engine/Plugins/VirtualProduction/DynamicMaterial/Source/DynamicMaterialEditor/Private/Components/MaterialStageBlends/DMMSBDarken.cpp b/Engine/Plugins/VirtualProduction/DynamicMaterial/Source/DynamicMaterialEditor/Private/Components/MaterialStageBlends/DMMSBDarken.cpp
--- a/Engine/Plugins/VirtualProduction/DynamicMaterial/Source/DynamicMaterialEditor/Private/Components/MaterialStageBlends/DMMSBDarken.cpp
+++ b/Engine/Plugins/VirtualProduction/DynamicMaterial/Source/DynamicMaterialEditor/Private/Components/MaterialStageBlends/DMMSBDarken.cpp
@@ -2,6 +2,8 @@
 
 #include "Components/MaterialStageBlends/DMMSBDarken.h"
 
+#include UE_INLINE_GENERATED_CPP_BY_NAME(DMMSBDarken)
+
 #define LOCTEXT_NAMESPACE "DMMaterialStageBlendDarken"
 
 UDMMaterialStageBlendDarken::UDMMaterialStageBlendDarken()
diff --git a/Engine/Plugins/VirtualProduction/DynamicMaterial/Source/DynamicMaterialEditor/Private/Components/MaterialStageBlends/DMMSBDarkenColor.cpp b/Engine/Plugins/VirtualProduction/DynamicMaterial/Source/DynamicMaterialEditor/Private/Components/MaterialStageBlends/DMMSBDarkenColor.cpp
--- a/Engine/Plugins/VirtualProduction/DynamicMaterial/Source/DynamicMaterialEditor/Private/Components/MaterialStageBlends/DMMSBDarkenColor.cpp
+++ b/Engine/Plugins/VirtualProduction/DynamicMaterial/Source/DynamicMaterialEditor/Private/Components/MaterialStageBlends/DMMSBDarkenColor.cpp
@@ -2,6 +2,8 @@
 
 #include "Components/MaterialStageBlends/DMMSBDarkenColor.h"
 
+#include UE_INLINE_GENERATED_CPP_BY_NAME(DMMSBDarkenColor)
+
 #define LOCTEXT_NAMESPACE "DMMaterialStageBlendDarkenColor"
 
 UDMMaterialStageBlendDarkenColor::UDMMaterialStageBlendDarkenColor()
diff --git a/Engine/Plugins/VirtualProduction/DynamicMaterial/Source/DynamicMaterialEditor/Private/Components/MaterialStageBlends/DMMSBDifference.cpp b/Engine/Plugins/VirtualProduction/DynamicMaterial/Source/DynamicMaterialEditor/Private/Components/MaterialStageBlends/DMMSBDifference.cpp
--- a/Engine/Plugins/VirtualProduction/DynamicMaterial/Source/DynamicMaterialEditor/Private/Components/MaterialStageBlends/DMMSBDifference.cpp
+++ b/Engine/Plugins/VirtualProduction/DynamicMaterial/Source/DynamicMaterialEditor/Private/Components/MaterialStageBlends/DMMSBDifference.cpp
@@ -2,6 +2,8 @@
 
 #include "Components/MaterialStageBlends/DMMSBDifference.h"
 
+#include UE_INLINE_GENERATED_CPP_BY_NAME(DMMSBDifference)
+
 #define LOCTEXT_NAMESPACE "DMMaterialStageBlendDifference"
 
 UDMMaterialStageBlendDifference::UDMMaterialStageBlendDifference()
diff --git a/Engine/Plugins/VirtualProduction/DynamicMaterial/Source/DynamicMaterialEditor/Private/Components/MaterialStageBlends/DMMSBDivide.cpp b/Engine/Plugins/VirtualProduction/DynamicMaterial/Source/DynamicMaterialEditor/Private/Components/MaterialStageBlends/DMMSBDivide.cpp
--- a/Engine/Plugins/VirtualProduction/DynamicMaterial/Source/DynamicMaterialEditor/Private/Components/MaterialStageBlends/DMMSBDivide.cpp
+++ b/Engine/Plugins/VirtualProduction/DynamicMaterial/Source/DynamicMaterialEditor/Private/Components/MaterialStageBlends/DMMSBDivide.cpp
@@ -4,6 +4,8 @@
 
 #include "Materials/MaterialExpressionDivide.h"
 
+#include UE_INLINE_GENERATED_CPP_BY_NAME(DMMSBDivide)
+
 #define LOCTEXT_NAMESPACE "DMMaterialStageBlendDivide"
 
 UDMMaterialStageBlendDivide::UDMMaterialStageBlendDivide()
diff --git a/Engine/Plugins/VirtualProduction/DynamicMaterial/Source/DynamicMaterialEditor/Private/Components/MaterialStageBlends/DMMSBExclusion.cpp b/Engine/Plugins/VirtualProduction/DynamicMaterial/Source/DynamicMaterialEditor/Private/Components/MaterialStageBlends/DMMSBExclusion.cpp
--- a/Engine/Plugins/VirtualProduction/DynamicMaterial/Source/DynamicMaterialEditor/Private/Components/MaterialStageBlends/DMMSBExclusion.cpp
+++ b/Engine/Plugins/VirtualProduction/DynamicMaterial/Source/DynamicMaterialEditor/Private/Components/MaterialStageBlends/DMMSBExclusion.cpp
@@ -2,6 +2,8 @@
 
 #include "Components/MaterialStageBlends/DMMSBExclusion.h"
 
+#include UE_INLINE_GENERATED_CPP_BY_NAME(DMMSBExclusion)
+
 #define LOCTEXT_NAMESPACE "DMMaterialStageBlendExclusion"
 
 UDMMaterialStageBlendExclusion::UDMMaterialStageBlendExclusion()
diff --git a/Engine/Plugins/VirtualProduction/DynamicMaterial/Source/DynamicMaterialEditor/Private/Components/MaterialStageBlends/DMMSBHardLight.cpp b/Engine/Plugins/VirtualProduction/DynamicMaterial/Source/DynamicMaterialEditor/Private/Components/MaterialStageBlends/DMMSBHardLight.cpp
--- a/Engine/Plugins/VirtualProduction/DynamicMaterial/Source/DynamicMaterialEditor/Private/Components/MaterialStageBlends/DMMSBHardLight.cpp
+++ b/Engine/Plugins/VirtualProduction/DynamicMaterial/Source/DynamicMaterialEditor/Private/Components/MaterialStageBlends/DMMSBHardLight.cpp
@@ -2,6 +2,8 @@
 
 #include "Components/MaterialStageBlends/DMMSBHardLight.h"
 
+#include UE_INLINE_GENERATED_CPP_BY_NAME(DMMSBHardLight)
+
 #define LOCTEXT_NAMESPACE "DMMaterialStageBlendHardLight"
 
 UDMMaterialStageBlendHardLight::UDMMaterialStageBlendHardLight()
diff --git a/Engine/Plugins/VirtualProduction/DynamicMaterial/Source/DynamicMaterialEditor/Private/Components/MaterialStageBlends/DMMSBHardMix.cpp b/Engine/Plugins/VirtualProduction/DynamicMaterial/Source/DynamicMaterialEditor/Private/Components/MaterialStageBlends/DMMSBHardMix.cpp
--- a/Engine/Plugins/VirtualProduction/DynamicMaterial/Source/DynamicMaterialEditor/Private/Components/MaterialStageBlends/DMMSBHardMix.cpp
+++ b/Engine/Plugins/VirtualProduction/DynamicMaterial/Source/DynamicMaterialEditor/Private/Components/MaterialStageBlends/DMMSBHardMix.cpp
@@ -2,6 +2,8 @@
 
 #include "Components/MaterialStageBlends/DMMSBHardMix.h"
 
+#include UE_INLINE_GENERATED_CPP_BY_NAME(DMMSBHardMix)
+
 #define LOCTEXT_NAMESPACE "DMMaterialStageBlendHardMix"
 
 UDMMaterialStageBlendHardMix::UDMMaterialStageBlendHardMix()
diff --git a/Engine/Plugins/VirtualProduction/DynamicMaterial/Source/DynamicMaterialEditor/Private/Components/MaterialStageBlends/DMMSBHue.cpp b/Engine/Plugins/VirtualProduction/DynamicMaterial/Source/DynamicMaterialEditor/Private/Components/MaterialStageBlends/DMMSBHue.cpp
--- a/Engine/Plugins/VirtualProduction/DynamicMaterial/Source/DynamicMaterialEditor/Private/Components/MaterialStageBlends/DMMSBHue.cpp
+++ b/Engine/Plugins/VirtualProduction/DynamicMaterial/Source/DynamicMaterialEditor/Private/Components/MaterialStageBlends/DMMSBHue.cpp
@@ -2,6 +2,8 @@
 
 #include "Components/MaterialStageBlends/DMMSBHue.h"
 
+#include UE_INLINE_GENERATED_CPP_BY_NAME(DMMSBHue)
+
 #define LOCTEXT_NAMESPACE "DMMaterialStageBlendHue"
 
 UDMMaterialStageBlendHue::UDMMaterialStageBlendHue()
diff --git a/Engine/Plugins/VirtualProduction/DynamicMaterial/Source/DynamicMaterialEditor/Private/Components/MaterialStageBlends/DMMSBLighten.cpp b/Engine/Plugins/VirtualProduction/DynamicMaterial/Source/DynamicMaterialEditor/Private/Components/MaterialStageBlends/DMMSBLighten.cpp
--- a/Engine/Plugins/VirtualProduction/DynamicMaterial/Source/DynamicMaterialEditor/Private/Components/MaterialStageBlends/DMMSBLighten.cpp
+++ b/Engine/Plugins/VirtualProduction/DynamicMaterial/Source/DynamicMaterialEditor/Private/Components/MaterialStageBlends/DMMSBLighten.cpp
@@ -2,6 +2,8 @@
 
 #include "Components/MaterialStageBlends/DMMSBLighten.h"
 
+#include UE_INLINE_GENERATED_CPP_BY_NAME(DMMSBLighten)
+
 #define LOCTEXT_NAMESPACE "DMMaterialStageBlendLighten"
 
 UDMMaterialStageBlendLighten::UDMMaterialStageBlendLighten()
diff --git a/Engine/Plugins/VirtualProduction/DynamicMaterial/Source/DynamicMaterialEditor/Private/Components/MaterialStageBlends/DMMSBLightenColor.cpp b/Engine/Plugins/VirtualProduction/DynamicMaterial/Source/DynamicMaterialEditor/Private/Components/MaterialStageBlends/DMMSBLightenColor.cpp
--- a/Engine/Plugins/VirtualProduction/DynamicMaterial/Source/DynamicMaterialEditor/Private/Components/MaterialStageBlends/DMMSBLightenColor.cpp
+++ b/Engine/Plugins/VirtualProduction/DynamicMaterial/Source/DynamicMaterialEditor/Private/Components/MaterialStageBlends/DMMSBLightenColor.cpp
@@ -2,6 +2,8 @@
 
 #include "Components/MaterialStageBlends/DMMSBLightenColor.h"
 
+#include UE_INLINE_GENERATED_CPP_BY_NAME(DMMSBLightenColor)
+
 #define LOCTEXT_NAMESPACE "DMMaterialStageBlendLightenColor"
 
 UDMMaterialStageBlendLightenColor::UDMMaterialStageBlendLightenColor()
diff --git a/Engine/Plugins/VirtualProduction/DynamicMaterial/Source/DynamicMaterialEditor/Private/Components/MaterialStageBlends/DMMSBLinearBurn.cpp b/Engine/Plugins/VirtualProduction/DynamicMaterial/Source/DynamicMaterialEditor/Private/Components/MaterialStageBlends/DMMSBLinearBurn.cpp
--- a/Engine/Plugins/VirtualProduction/DynamicMaterial/Source/DynamicMaterialEditor/Private/Components/MaterialStageBlends/DMMSBLinearBurn.cpp
+++ b/Engine/Plugins/VirtualProduction/DynamicMaterial/Source/DynamicMaterialEditor/Private/Components/MaterialStageBlends/DMMSBLinearBurn.cpp
@@ -2,6 +2,8 @@
 
 #include "Components/MaterialStageBlends/DMMSBLinearBurn.h"
 
+#include UE_INLINE_GENERATED_CPP_BY_NAME(DMMSBLinearBurn)
+
 #define LOCTEXT_NAMESPACE "DMMaterialStageBlendLinearBurn"
 
 UDMMaterialStageBlendLinearBurn::UDMMaterialStageBlendLinearBurn()
diff --git a/Engine/Plugins/VirtualProduction/DynamicMaterial/Source/DynamicMaterialEditor/Private/Components/MaterialStageBlends/DMMSBLinearDodge.cpp b/Engine/Plugins/VirtualProduction/DynamicMaterial/Source/DynamicMaterialEditor/Private/Components/MaterialStageBlends/DMMSBLinearDodge.cpp
--- a/Engine/Plugins/VirtualProduction/DynamicMaterial/Source/DynamicMaterialEditor/Private/Components/MaterialStageBlends/DMMSBLinearDodge.cpp
+++ b/Engine/Plugins/VirtualProduction/DynamicMaterial/Source/DynamicMaterialEditor/Private/Components/MaterialStageBlends/DMMSBLinearDodge.cpp
@@ -2,6 +2,8 @@
 
 #include "Components/MaterialStageBlends/DMMSBLinearDodge.h"
 
+#include UE_INLINE_GENERATED_CPP_BY_NAME(DMMSBLinearDodge)
+
 #define LOCTEXT_NAMESPACE "DMMaterialStageBlendLinearDodge"
 
 UDMMaterialStageBlendLinearDodge::UDMMaterialStageBlendLinearDodge()
diff --git a/Engine/Plugins/VirtualProduction/DynamicMaterial/Source/DynamicMaterialEditor/Private/Components/MaterialStageBlends/DMMSBLinearLight.cpp b/Engine/Plugins/VirtualProduction/DynamicMaterial/Source/DynamicMaterialEditor/Private/Components/MaterialStageBlends/DMMSBLinearLight.cpp
--- a/Engine/Plugins/VirtualProduction/DynamicMaterial/Source/DynamicMaterialEditor/Private/Components/MaterialStageBlends/DMMSBLinearLight.cpp
+++ b/Engine/Plugins/VirtualProduction/DynamicMaterial/Source/DynamicMaterialEditor/Private/Components/MaterialStageBlends/DMMSBLinearLight.cpp
@@ -2,6 +2,8 @@
 
 #include "Components/MaterialStageBlends/DMMSBLinearLight.h"
 
+#include UE_INLINE_GENERATED_CPP_BY_NAME(DMMSBLinearLight)
+
 #define LOCTEXT_NAMESPACE "DMMaterialStageBlendLinearLight"
 
 UDMMaterialStageBlendLinearLight::UDMMaterialStageBlendLinearLight()
diff --git a/Engine/Plugins/VirtualProduction/DynamicMaterial/Source/DynamicMaterialEditor/Private/Components/MaterialStageBlends/DMMSBLuminosity.cpp b/Engine/Plugins/VirtualProduction/DynamicMaterial/Source/DynamicMaterialEditor/Private/Components/MaterialStageBlends/DMMSBLuminosity.cpp
--- a/Engine/Plugins/VirtualProduction/DynamicMaterial/Source/DynamicMaterialEditor/Private/Components/MaterialStageBlends/DMMSBLuminosity.cpp
+++ b/Engine/Plugins/VirtualProduction/DynamicMaterial/Source/DynamicMaterialEditor/Private/Components/MaterialStageBlends/DMMSBLuminosity.cpp
@@ -2,6 +2,8 @@
 
 #include "Components/MaterialStageBlends/DMMSBLuminosity.h"
 
+#include UE_INLINE_GENERATED_CPP_BY_NAME(DMMSBLuminosity)
+
 #define LOCTEXT_NAMESPACE "DMMaterialStageBlendLuminosity"
 
 UDMMaterialStageBlendLuminosity::UDMMaterialStageBlendLuminosity()
diff --git a/Engine/Plugins/VirtualProduction/DynamicMaterial/Source/DynamicMaterialEditor/Private/Components/MaterialStageBlends/DMMSBMultiply.cpp b/Engine/Plugins/VirtualProduction/DynamicMaterial/Source/DynamicMaterialEditor/Private/Components/MaterialStageBlends/DMMSBMultiply.cpp
--- a/Engine/Plugins/VirtualProduction/DynamicMaterial/Source/DynamicMaterialEditor/Private/Components/MaterialStageBlends/DMMSBMultiply.cpp
+++ b/Engine/Plugins/VirtualProduction/DynamicMaterial/Source/DynamicMaterialEditor/Private/Components/MaterialStageBlends/DMMSBMultiply.cpp
@@ -4,6 +4,8 @@
 
 #include "Materials/MaterialExpressionMultiply.h"
 
+#include UE_INLINE_GENERATED_CPP_BY_NAME(DMMSBMultiply)
+
 #define LOCTEXT_NAMESPACE "DMMaterialStageBlendMultiply"
 
 UDMMaterialStageBlendMultiply::UDMMaterialStageBlendMultiply()
diff --git a/Engine/Plugins/VirtualProduction/DynamicMaterial/Source/DynamicMaterialEditor/Private/Components/MaterialStageBlends/DMMSBNormal.cpp b/Engine/Plugins/VirtualProduction/DynamicMaterial/Source/DynamicMaterialEditor/Private/Components/MaterialStageBlends/DMMSBNormal.cpp
--- a/Engine/Plugins/VirtualProduction/DynamicMaterial/Source/DynamicMaterialEditor/Private/Components/MaterialStageBlends/DMMSBNormal.cpp
+++ b/Engine/Plugins/VirtualProduction/DynamicMaterial/Source/DynamicMaterialEditor/Private/Components/MaterialStageBlends/DMMSBNormal.cpp
@@ -2,6 +2,8 @@
 
 #include "Components/MaterialStageBlends/DMMSBNormal.h"
 
+#include UE_INLINE_GENERATED_CPP_BY_NAME(DMMSBNormal)
+
 #define LOCTEXT_NAMESPACE "DMMaterialStageBlendNormal"
 
 UDMMaterialStageBlendNormal::UDMMaterialStageBlendNormal()
diff --git a/Engine/Plugins/VirtualProduction/DynamicMaterial/Source/DynamicMaterialEditor/Private/Components/MaterialStageBlends/DMMSBOverlay.cpp b/Engine/Plugins/VirtualProduction/DynamicMaterial/Source/DynamicMaterialEditor/Private/Components/MaterialStageBlends/DMMSBOverlay.cpp
--- a/Engine/Plugins/VirtualProduction/DynamicMaterial/Source/DynamicMaterialEditor/Private/Components/MaterialStageBlends/DMMSBOverlay.cpp
+++ b/Engine/Plugins/VirtualProduction/DynamicMaterial/Source/DynamicMaterialEditor/Private/Components/MaterialStageBlends/DMMSBOverlay.cpp
@@ -2,6 +2,8 @@
 
 #include "Components/MaterialStageBlends/DMMSBOverlay.h"
 
+#include UE_INLINE_GENERATED_CPP_BY_NAME(DMMSBOverlay)
+
 #define LOCTEXT_NAMESPACE "DMMaterialStageBlendOverlay"
 
 UDMMaterialStageBlendOverlay::UDMMaterialStageBlendOverlay()
diff --git a/Engine/Plugins/VirtualProduction/DynamicMaterial/Source/DynamicMaterialEditor/Private/Components/MaterialStageBlends/DMMSBPinLight.cpp b/Engine/Plugins/VirtualProduction/DynamicMaterial/Source/DynamicMaterialEditor/Private/Components/MaterialStageBlends/DMMSBPinLight.cpp
--- a/Engine/Plugins/VirtualProduction/DynamicMaterial/Source/DynamicMaterialEditor/Private/Components/MaterialStageBlends/DMMSBPinLight.cpp
+++ b/Engine/Plugins/VirtualProduction/DynamicMaterial/Source/DynamicMaterialEditor/Private/Components/MaterialStageBlends/DMMSBPinLight.cpp
@@ -2,6 +2,8 @@
 
 #include "Components/MaterialStageBlends/DMMSBPinLight.h"
 
+#include UE_INLINE_GENERATED_CPP_BY_NAME(DMMSBPinLight)
+
 #define LOCTEXT_NAMESPACE "DMMaterialStageBlendPinLight"
 
 UDMMaterialStageBlendPinLight::UDMMaterialStageBlendPinLight()
diff --git a/Engine/Plugins/VirtualProduction/DynamicMaterial/Source/DynamicMaterialEditor/Private/Components/MaterialStageBlends/DMMSBSaturation.cpp b/Engine/Plugins/VirtualProduction/DynamicMaterial/Source/DynamicMaterialEditor/Private/Components/MaterialStageBlends/DMMSBSaturation.cpp
--- a/Engine/Plugins/VirtualProduction/DynamicMaterial/Source/DynamicMaterialEditor/Private/Components/MaterialStageBlends/DMMSBSaturation.cpp
+++ b/Engine/Plugins/VirtualProduction/DynamicMaterial/Source/DynamicMaterialEditor/Private/Components/MaterialStageBlends/DMMSBSaturation.cpp
@@ -2,6 +2,8 @@
 
 #include "Components/MaterialStageBlends/DMMSBSaturation.h"
 
+#include UE_INLINE_GENERATED_CPP_BY_NAME(DMMSBSaturation)
+
 #define LOCTEXT_NAMESPACE "DMMaterialStageBlendSaturation"
 
 UDMMaterialStageBlendSaturation::UDMMaterialStageBlendSaturation()
diff --git a/Engine/Plugins/VirtualProduction/DynamicMaterial/Source/DynamicMaterialEditor/Private/Components/MaterialStageBlends/DMMSBScreen.cpp b/Engine/Plugins/VirtualProduction/DynamicMaterial/Source/DynamicMaterialEditor/Private/Components/MaterialStageBlends/DMMSBScreen.cpp
--- a/Engine/Plugins/VirtualProduction/DynamicMaterial/Source/DynamicMaterialEditor/Private/Components/MaterialStageBlends/DMMSBScreen.cpp
+++ b/Engine/Plugins/VirtualProduction/DynamicMaterial/Source/DynamicMaterialEditor/Private/Components/MaterialStageBlends/DMMSBScreen.cpp
@@ -2,6 +2,8 @@
 
 #include "Components/MaterialStageBlends/DMMSBScreen.h"
 
+#include UE_INLINE_GENERATED_CPP_BY_NAME(DMMSBScreen)
+
 #define LOCTEXT_NAMESPACE "DMMaterialStageBlendScreen"
 
 UDMMaterialStageBlendScreen::UDMMaterialStageBlendScreen()
diff --git a/Engine/Plugins/VirtualProduction/DynamicMaterial/Source/DynamicMaterialEditor/Private/Components/MaterialStageBlends/DMMSBSoftLight.cpp b/Engine/Plugins/VirtualProduction/DynamicMaterial/Source/DynamicMaterialEditor/Private/Components/MaterialStageBlends/DMMSBSoftLight.cpp
--- a/Engine/Plugins/VirtualProduction/DynamicMaterial/Source/DynamicMaterialEditor/Private/Components/MaterialStageBlends/DMMSBSoftLight.cpp
+++ b/Engine/Plugins/VirtualProduction/DynamicMaterial/Source/DynamicMaterialEditor/Private/Components/MaterialStageBlends/DMMSBSoftLight.cpp
@@ -2,6 +2,8 @@
 
 #include "Components/MaterialStageBlends/DMMSBSoftLight.h"
 
+#include UE_INLINE_GENERATED_CPP_BY_NAME(DMMSBSoftLight)
+
 #define LOCTEXT_NAMESPACE "DMMaterialStageBlendSoftLight"
 
 UDMMaterialStageBlendSoftLight::UDMMaterialStageBlendSoftLight()
diff --git a/Engine/Plugins/VirtualProduction/DynamicMaterial/Source/DynamicMaterialEditor/Private/Components/MaterialStageBlends/DMMSBSubtract.cpp b/Engine/Plugins/VirtualProduction/DynamicMaterial/Source/DynamicMaterialEditor/Private/Components/MaterialStageBlends/DMMSBSubtract.cpp
--- a/Engine/Plugins/VirtualProduction/DynamicMaterial/Source/DynamicMaterialEditor/Private/Components/MaterialStageBlends/DMMSBSubtract.cpp
+++ b/Engine/Plugins/VirtualProduction/DynamicMaterial/Source/DynamicMaterialEditor/Private/Components/MaterialStageBlends/DMMSBSubtract.cpp
@@ -4,6 +4,8 @@
 
 #include "Materials/MaterialExpressionSubtract.h"
 
+#include UE_INLINE_GENERATED_CPP_BY_NAME(DMMSBSubtract)
+
 #define LOCTEXT_NAMESPACE "DMMaterialStageBlendSubtract"
 
 UDMMaterialStageBlendSubtract::UDMMaterialStageBlendSubtract()
diff --git a/Engine/Plugins/VirtualProduction/DynamicMaterial/Source/DynamicMaterialEditor/Private/Components/MaterialStageBlends/DMMSBVividLight.cpp b/Engine/Plugins/VirtualProduction/DynamicMaterial/Source/DynamicMaterialEditor/Private/Components/MaterialStageBlends/DMMSBVividLight.cpp
--- a/Engine/Plugins/VirtualProduction/DynamicMaterial/Source/DynamicMaterialEditor/Private/Components/MaterialStageBlends/DMMSBVividLight.cpp
+++ b/Engine/Plugins/VirtualProduction/DynamicMaterial/Source/DynamicMaterialEditor/Private/Components/MaterialStageBlends/DMMSBVividLight.cpp
@@ -2,6 +2,8 @@
 
 #include "Components/MaterialStageBlends/DMMSBVividLight.h"
 
+#include UE_INLINE_GENERATED_CPP_BY_NAME(DMMSBVividLight)
+
 #define LOCTEXT_NAMESPACE "DMMaterialStageBlendVividLight"
 
 UDMMaterialStageBlendVividLight::UDMMaterialStageBlendVividLight()
diff --git a/Engine/Plugins/VirtualProduction/DynamicMaterial/Source/DynamicMaterialEditor/Private/Components/MaterialStageExpressions/DMMSEAbs.cpp b/Engine/Plugins/VirtualProduction/DynamicMaterial/Source/DynamicMaterialEditor/Private/Components/MaterialStageExpressions/DMMSEAbs.cpp
--- a/Engine/Plugins/VirtualProduction/DynamicMaterial/Source/DynamicMaterialEditor/Private/Components/MaterialStageExpressions/DMMSEAbs.cpp
+++ b/Engine/Plugins/VirtualProduction/DynamicMaterial/Source/DynamicMaterialEditor/Private/Components/MaterialStageExpressions/DMMSEAbs.cpp
@@ -3,6 +3,8 @@
 #include "Components/MaterialStageExpressions/DMMSEAbs.h"
 #include "Materials/MaterialExpressionAbs.h"
 
+#include UE_INLINE_GENERATED_CPP_BY_NAME(DMMSEAbs)
+
 #define LOCTEXT_NAMESPACE "DMMaterialStageExpressionAbs"
 
 UDMMaterialStageExpressionAbs::UDMMaterialStageExpressionAbs()
diff --git a/Engine/Plugins/VirtualProduction/DynamicMaterial/Source/DynamicMaterialEditor/Private/Components/MaterialStageExpressions/DMMSEActorPositionWS.cpp b/Engine/Plugins/VirtualProduction/DynamicMaterial/Source/DynamicMaterialEditor/Private/Components/MaterialStageExpressions/DMMSEActorPositionWS.cpp
--- a/Engine/Plugins/VirtualProduction/DynamicMaterial/Source/DynamicMaterialEditor/Private/Components/MaterialStageExpressions/DMMSEActorPositionWS.cpp
+++ b/Engine/Plugins/VirtualProduction/DynamicMaterial/Source/DynamicMaterialEditor/Private/Components/MaterialStageExpressions/DMMSEActorPositionWS.cpp
@@ -3,6 +3,8 @@
 #include "Components/MaterialStageExpressions/DMMSEActorPositionWS.h"
 #include "Materials/MaterialExpressionActorPositionWS.h"
 
+#include UE_INLINE_GENERATED_CPP_BY_NAME(DMMSEActorPositionWS)
+
 #define LOCTEXT_NAMESPACE "DMMaterialStageExpressionActorPositionWS"
 
 UDMMaterialStageExpressionActorPositionWS::UDMMaterialStageExpressionActorPositionWS()
diff --git a/Engine/Plugins/VirtualProduction/DynamicMaterial/Source/DynamicMaterialEditor/Private/Components/MaterialStageExpressions/DMMSEAdd.cpp b/Engine/Plugins/VirtualProduction/DynamicMaterial/Source/DynamicMaterialEditor/Private/Components/MaterialStageExpressions/DMMSEAdd.cpp
--- a/Engine/Plugins/VirtualProduction/DynamicMaterial/Source/DynamicMaterialEditor/Private/Components/MaterialStageExpressions/DMMSEAdd.cpp
+++ b/Engine/Plugins/VirtualProduction/DynamicMaterial/Source/DynamicMaterialEditor/Private/Components/MaterialStageExpressions/DMMSEAdd.cpp
@@ -3,6 +3,8 @@
 #include "Components/MaterialStageExpressions/DMMSEAdd.h"
 #include "Materials/MaterialExpressionAdd.h"
 
+#include UE_INLINE_GENERATED_CPP_BY_NAME(DMMSEAdd)
+
 #define LOCTEXT_NAMESPACE "DMMaterialStageExpressionAdd"
 
 UDMMaterialStageExpressionAdd::UDMMaterialStageExpressionAdd()
@@ -14,4 +16,4 @@ UDMMaterialStageExpressionAdd::UDMMaterialStageExpressionAdd()
 	SetupInputs(2);
 }
 
-#undef LOCTEXT_NAMESPACE
\ No newline at end of file
+#undef LOCTEXT_NAMESPACE
diff --git a/Engine/Plugins/VirtualProduction/DynamicMaterial/Source/DynamicMaterialEditor/Private/Components/MaterialStageExpressions/DMMSEArcCosine.cpp b/Engine/Plugins/VirtualProduction/DynamicMaterial/Source/DynamicMaterialEditor/Private/Components/MaterialStageExpressions/DMMSEArcCosine.cpp
--- a/Engine/Plugins/VirtualProduction/DynamicMaterial/Source/DynamicMaterialEditor/Private/Components/MaterialStageExpressions/DMMSEArcCosine.cpp
+++ b/Engine/Plugins/VirtualProduction/DynamicMaterial/Source/DynamicMaterialEditor/Private/Components/MaterialStageExpressions/DMMSEArcCosine.cpp
@@ -4,6 +4,8 @@
 
 #include "Materials/MaterialExpressionArccosine.h"
 
+#include UE_INLINE_GENERATED_CPP_BY_NAME(DMMSEArcCosine)
+
 #define LOCTEXT_NAMESPACE "DMMaterialStageExpressionArccosine"
 
 UDMMaterialStageExpressionArccosine::UDMMaterialStageExpressionArccosine()
diff --git a/Engine/Plugins/VirtualProduction/DynamicMaterial/Source/DynamicMaterialEditor/Private/Components/MaterialStageExpressions/DMMSEArcCosineFast.cpp b/Engine/Plugins/VirtualProduction/DynamicMaterial/Source/DynamicMaterialEditor/Private/Components/MaterialStageExpressions/DMMSEArcCosineFast.cpp
--- a/Engine/Plugins/VirtualProduction/DynamicMaterial/Source/DynamicMaterialEditor/Private/Components/MaterialStageExpressions/DMMSEArcCosineFast.cpp
+++ b/Engine/Plugins/VirtualProduction/DynamicMaterial/Source/DynamicMaterialEditor/Private/Components/MaterialStageExpressions/DMMSEArcCosineFast.cpp
@@ -4,6 +4,8 @@
 
 #include "Materials/MaterialExpressionArccosineFast.h"
 
+#include UE_INLINE_GENERATED_CPP_BY_NAME(DMMSEArcCosineFast)
+
 #define LOCTEXT_NAMESPACE "DMMaterialStageExpressionArccosineFast"
 
 UDMMaterialStageExpressionArccosineFast::UDMMaterialStageExpressionArccosineFast()
diff --git a/Engine/Plugins/VirtualProduction/DynamicMaterial/Source/DynamicMaterialEditor/Private/Components/MaterialStageExpressions/DMMSEArcSine.cpp b/Engine/Plugins/VirtualProduction/DynamicMaterial/Source/DynamicMaterialEditor/Private/Components/MaterialStageExpressions/DMMSEArcSine.cpp
--- a/Engine/Plugins/VirtualProduction/DynamicMaterial/Source/DynamicMaterialEditor/Private/Components/MaterialStageExpressions/DMMSEArcSine.cpp
+++ b/Engine/Plugins/VirtualProduction/DynamicMaterial/Source/DynamicMaterialEditor/Private/Components/MaterialStageExpressions/DMMSEArcSine.cpp
@@ -4,6 +4,8 @@
 
 #include "Materials/MaterialExpressionArcsine.h"
 
+#include UE_INLINE_GENERATED_CPP_BY_NAME(DMMSEArcSine)
+
 #define LOCTEXT_NAMESPACE "DMMaterialStageExpressionArcsine"
 
 UDMMaterialStageExpressionArcsine::UDMMaterialStageExpressionArcsine()
diff --git a/Engine/Plugins/VirtualProduction/DynamicMaterial/Source/DynamicMaterialEditor/Private/Components/MaterialStageExpressions/DMMSEArcSineFast.cpp b/Engine/Plugins/VirtualProduction/DynamicMaterial/Source/DynamicMaterialEditor/Private/Components/MaterialStageExpressions/DMMSEArcSineFast.cpp
--- a/Engine/Plugins/VirtualProduction/DynamicMaterial/Source/DynamicMaterialEditor/Private/Components/MaterialStageExpressions/DMMSEArcSineFast.cpp
+++ b/Engine/Plugins/VirtualProduction/DynamicMaterial/Source/DynamicMaterialEditor/Private/Components/MaterialStageExpressions/DMMSEArcSineFast.cpp
@@ -4,6 +4,8 @@
 
 #include "Materials/MaterialExpressionArcsineFast.h"
 
+#include UE_INLINE_GENERATED_CPP_BY_NAME(DMMSEArcSineFast)
+
 #define LOCTEXT_NAMESPACE "DMMaterialStageExpressionArcsineFast"
 
 UDMMaterialStageExpressionArcsineFast::UDMMaterialStageExpressionArcsineFast()
diff --git a/Engine/Plugins/VirtualProduction/DynamicMaterial/Source/DynamicMaterialEditor/Private/Components/MaterialStageExpressions/DMMSEArctangent.cpp b/Engine/Plugins/VirtualProduction/DynamicMaterial/Source/DynamicMaterialEditor/Private/Components/MaterialStageExpressions/DMMSEArctangent.cpp
--- a/Engine/Plugins/VirtualProduction/DynamicMaterial/Source/DynamicMaterialEditor/Private/Components/MaterialStageExpressions/DMMSEArctangent.cpp
+++ b/Engine/Plugins/VirtualProduction/DynamicMaterial/Source/DynamicMaterialEditor/Private/Components/MaterialStageExpressions/DMMSEArctangent.cpp
@@ -3,6 +3,8 @@
 #include "Components/MaterialStageExpressions/DMMSEArctangent.h"
 #include "Materials/MaterialExpressionArctangent.h"
 
+#include UE_INLINE_GENERATED_CPP_BY_NAME(DMMSEArctangent)
+
 #define LOCTEXT_NAMESPACE "DMMaterialStageExpressionArctangent"
 
 UDMMaterialStageExpressionArctangent::UDMMaterialStageExpressionArctangent()
diff --git a/Engine/Plugins/VirtualProduction/DynamicMaterial/Source/DynamicMaterialEditor/Private/Components/MaterialStageExpressions/DMMSEArctangent2.cpp b/Engine/Plugins/VirtualProduction/DynamicMaterial/Source/DynamicMaterialEditor/Private/Components/MaterialStageExpressions/DMMSEArctangent2.cpp
--- a/Engine/Plugins/VirtualProduction/DynamicMaterial/Source/DynamicMaterialEditor/Private/Components/MaterialStageExpressions/DMMSEArctangent2.cpp
+++ b/Engine/Plugins/VirtualProduction/DynamicMaterial/Source/DynamicMaterialEditor/Private/Components/MaterialStageExpressions/DMMSEArctangent2.cpp
@@ -3,6 +3,8 @@
 #include "Components/MaterialStageExpressions/DMMSEArctangent2.h"
 #include "Materials/MaterialExpressionArctangent2.h"
 
+#include UE_INLINE_GENERATED_CPP_BY_NAME(DMMSEArctangent2)
+
 #define LOCTEXT_NAMESPACE "DMMaterialStageExpressionArctangent2"
 
 UDMMaterialStageExpressionArctangent2::UDMMaterialStageExpressionArctangent2()
diff --git a/Engine/Plugins/VirtualProduction/DynamicMaterial/Source/DynamicMaterialEditor/Private/Components/MaterialStageExpressions/DMMSEArctangent2Fast.cpp b/Engine/Plugins/VirtualProduction/DynamicMaterial/Source/DynamicMaterialEditor/Private/Components/MaterialStageExpressions/DMMSEArctangent2Fast.cpp
--- a/Engine/Plugins/VirtualProduction/DynamicMaterial/Source/DynamicMaterialEditor/Private/Components/MaterialStageExpressions/DMMSEArctangent2Fast.cpp
+++ b/Engine/Plugins/VirtualProduction/DynamicMaterial/Source/DynamicMaterialEditor/Private/Components/MaterialStageExpressions/DMMSEArctangent2Fast.cpp
@@ -3,6 +3,8 @@
 #include "Components/MaterialStageExpressions/DMMSEArctangent2Fast.h"
 #include "Materials/MaterialExpressionArctangent2Fast.h"
 
+#include UE_INLINE_GENERATED_CPP_BY_NAME(DMMSEArctangent2Fast)
+
 #define LOCTEXT_NAMESPACE "DMMaterialStageExpressionArctangent2Fast"
 
 UDMMaterialStageExpressionArctangent2Fast::UDMMaterialStageExpressionArctangent2Fast()
diff --git a/Engine/Plugins/VirtualProduction/DynamicMaterial/Source/DynamicMaterialEditor/Private/Components/MaterialStageExpressions/DMMSEArctangentFast.cpp b/Engine/Plugins/VirtualProduction/DynamicMaterial/Source/DynamicMaterialEditor/Private/Components/MaterialStageExpressions/DMMSEArctangentFast.cpp
--- a/Engine/Plugins/VirtualProduction/DynamicMaterial/Source/DynamicMaterialEditor/Private/Components/MaterialStageExpressions/DMMSEArctangentFast.cpp
+++ b/Engine/Plugins/VirtualProduction/DynamicMaterial/Source/DynamicMaterialEditor/Private/Components/MaterialStageExpressions/DMMSEArctangentFast.cpp
@@ -3,6 +3,8 @@
 #include "Components/MaterialStageExpressions/DMMSEArctangentFast.h"
 #include "Materials/MaterialExpressionArctangentFast.h"
 
+#include UE_INLINE_GENERATED_CPP_BY_NAME(DMMSEArctangentFast)
+
 #define LOCTEXT_NAMESPACE "DMMaterialStageExpressionArctangentFast"
 
 UDMMaterialStageExpressionArctangentFast::UDMMaterialStageExpressionArctangentFast()
diff --git a/Engine/Plugins/VirtualProduction/DynamicMaterial/Source/DynamicMaterialEditor/Private/Components/MaterialStageExpressions/DMMSECameraPositionWS.cpp b/Engine/Plugins/VirtualProduction/DynamicMaterial/Source/DynamicMaterialEditor/Private/Components/MaterialStageExpressions/DMMSECameraPositionWS.cpp
--- a/Engine/Plugins/VirtualProduction/DynamicMaterial/Source/DynamicMaterialEditor/Private/Components/MaterialStageExpressions/DMMSECameraPositionWS.cpp
+++ b/Engine/Plugins/VirtualProduction/DynamicMaterial/Source/DynamicMaterialEditor/Private/Components/MaterialStageExpressions/DMMSECameraPositionWS.cpp
@@ -3,6 +3,8 @@
 #include "Components/MaterialStageExpressions/DMMSECameraPositionWS.h"
 #include "Materials/MaterialExpressionCameraPositionWS.h"
 
+#include UE_INLINE_GENERATED_CPP_BY_NAME(DMMSECameraPositionWS)
+
 #define LOCTEXT_NAMESPACE "DMMaterialStageExpressionCameraPositionWS"
 
 UDMMaterialStageExpressionCameraPositionWS::UDMMaterialStageExpressionCameraPositionWS()
diff --git a/Engine/Plugins/VirtualProduction/DynamicMaterial/Source/DynamicMaterialEditor/Private/Components/MaterialStageExpressions/DMMSECameraVectorWS.cpp b/Engine/Plugins/VirtualProduction/DynamicMaterial/Source/DynamicMaterialEditor/Private/Components/MaterialStageExpressions/DMMSECameraVectorWS.cpp
--- a/Engine/Plugins/VirtualProduction/DynamicMaterial/Source/DynamicMaterialEditor/Private/Components/MaterialStageExpressions/DMMSECameraVectorWS.cpp
+++ b/Engine/Plugins/VirtualProduction/DynamicMaterial/Source/DynamicMaterialEditor/Private/Components/MaterialStageExpressions/DMMSECameraVectorWS.cpp
@@ -3,6 +3,8 @@
 #include "Components/MaterialStageExpressions/DMMSECameraVectorWS.h"
 #include "Materials/MaterialExpressionCameraVectorWS.h"
 
+#include UE_INLINE_GENERATED_CPP_BY_NAME(DMMSECameraVectorWS)
+
 #define LOCTEXT_NAMESPACE "DMMaterialStageExpressionCameraVectorWS"
 
 UDMMaterialStageExpressionCameraVectorWS::UDMMaterialStageExpressionCameraVectorWS()
diff --git a/Engine/Plugins/VirtualProduction/DynamicMaterial/Source/DynamicMaterialEditor/Private/Components/MaterialStageExpressions/DMMSECeil.cpp b/Engine/Plugins/VirtualProduction/DynamicMaterial/Source/DynamicMaterialEditor/Private/Components/MaterialStageExpressions/DMMSECeil.cpp
--- a/Engine/Plugins/VirtualProduction/DynamicMaterial/Source/DynamicMaterialEditor/Private/Components/MaterialStageExpressions/DMMSECeil.cpp
+++ b/Engine/Plugins/VirtualProduction/DynamicMaterial/Source/DynamicMaterialEditor/Private/Components/MaterialStageExpressions/DMMSECeil.cpp
@@ -3,6 +3,8 @@
 #include "Components/MaterialStageExpressions/DMMSECeil.h"
 #include "Materials/MaterialExpressionCeil.h"
 
+#include UE_INLINE_GENERATED_CPP_BY_NAME(DMMSECeil)
+
 #define LOCTEXT_NAMESPACE "DMMaterialStageExpressionCeil"
 
 UDMMaterialStageExpressionCeil::UDMMaterialStageExpressionCeil()
diff --git a/Engine/Plugins/VirtualProduction/DynamicMaterial/Source/DynamicMaterialEditor/Private/Components/MaterialStageExpressions/DMMSEClamp.cpp b/Engine/Plugins/VirtualProduction/DynamicMaterial/Source/DynamicMaterialEditor/Private/Components/MaterialStageExpressions/DMMSEClamp.cpp
--- a/Engine/Plugins/VirtualProduction/DynamicMaterial/Source/DynamicMaterialEditor/Private/Components/MaterialStageExpressions/DMMSEClamp.cpp
+++ b/Engine/Plugins/VirtualProduction/DynamicMaterial/Source/DynamicMaterialEditor/Private/Components/MaterialStageExpressions/DMMSEClamp.cpp
@@ -5,6 +5,8 @@
 #include "Components/MaterialStageInputs/DMMSIValue.h"
 #include "Components/MaterialValues/DMMaterialValueFloat1.h"
 #include "Materials/MaterialExpressionClamp.h"
+
+#include UE_INLINE_GENERATED_CPP_BY_NAME(DMMSEClamp)
  
 #define LOCTEXT_NAMESPACE "DMMaterialStageExpressionClamp"
  
diff --git a/Engine/Plugins/VirtualProduction/DynamicMaterial/Source/DynamicMaterialEditor/Private/Components/MaterialStageExpressions/DMMSECosine.cpp b/Engine/Plugins/VirtualProduction/DynamicMaterial/Source/DynamicMaterialEditor/Private/Components/MaterialStageExpressions/DMMSECosine.cpp
--- a/Engine/Plugins/VirtualProduction/DynamicMaterial/Source/DynamicMaterialEditor/Private/Components/MaterialStageExpressions/DMMSECosine.cpp
+++ b/Engine/Plugins/VirtualProduction/DynamicMaterial/Source/DynamicMaterialEditor/Private/Components/MaterialStageExpressions/DMMSECosine.cpp
@@ -3,6 +3,8 @@
 #include "Components/MaterialStageExpressions/DMMSECosine.h"
 #include "Materials/MaterialExpressionCosine.h"
 
+#include UE_INLINE_GENERATED_CPP_BY_NAME(DMMSECosine)
+
 #define LOCTEXT_NAMESPACE "DMMaterialStageExpressionCosine"
 
 UDMMaterialStageExpressionCosine::UDMMaterialStageExpressionCosine()
diff --git a/Engine/Plugins/VirtualProduction/DynamicMaterial/Source/DynamicMaterialEditor/Private/Components/MaterialStageExpressions/DMMSECrossProduct.cpp b/Engine/Plugins/VirtualProduction/DynamicMaterial/Source/DynamicMaterialEditor/Private/Components/MaterialStageExpressions/DMMSECrossProduct.cpp
--- a/Engine/Plugins/VirtualProduction/DynamicMaterial/Source/DynamicMaterialEditor/Private/Components/MaterialStageExpressions/DMMSECrossProduct.cpp
+++ b/Engine/Plugins/VirtualProduction/DynamicMaterial/Source/DynamicMaterialEditor/Private/Components/MaterialStageExpressions/DMMSECrossProduct.cpp
@@ -3,6 +3,8 @@
 #include "Components/MaterialStageExpressions/DMMSECrossProduct.h"
 #include "Materials/MaterialExpressionCrossProduct.h"
 
+#include UE_INLINE_GENERATED_CPP_BY_NAME(DMMSECrossProduct)
+
 #define LOCTEXT_NAMESPACE "DMMaterialStageExpressionCrossProduct"
 
 UDMMaterialStageExpressionCrossProduct::UDMMaterialStageExpressionCrossProduct()
diff --git a/Engine/Plugins/VirtualProduction/DynamicMaterial/Source/DynamicMaterialEditor/Private/Components/MaterialStageExpressions/DMMSEDeltaTime.cpp b/Engine/Plugins/VirtualProduction/DynamicMaterial/Source/DynamicMaterialEditor/Private/Components/MaterialStageExpressions/DMMSEDeltaTime.cpp
--- a/Engine/Plugins/VirtualProduction/DynamicMaterial/Source/DynamicMaterialEditor/Private/Components/MaterialStageExpressions/DMMSEDeltaTime.cpp
+++ b/Engine/Plugins/VirtualProduction/DynamicMaterial/Source/DynamicMaterialEditor/Private/Components/MaterialStageExpressions/DMMSEDeltaTime.cpp
@@ -3,6 +3,8 @@
 #include "Components/MaterialStageExpressions/DMMSEDeltaTime.h"
 #include "Materials/MaterialExpressionDeltaTime.h"
 
+#include UE_INLINE_GENERATED_CPP_BY_NAME(DMMSEDeltaTime)
+
 #define LOCTEXT_NAMESPACE "DMMaterialStageExpressionDeltaTime"
 
 UDMMaterialStageExpressionDeltaTime::UDMMaterialStageExpressionDeltaTime()
diff --git a/Engine/Plugins/VirtualProduction/DynamicMaterial/Source/DynamicMaterialEditor/Private/Components/MaterialStageExpressions/DMMSEDivide.cpp b/Engine/Plugins/VirtualProduction/DynamicMaterial/Source/DynamicMaterialEditor/Private/Components/MaterialStageExpressions/DMMSEDivide.cpp
--- a/Engine/Plugins/VirtualProduction/DynamicMaterial/Source/DynamicMaterialEditor/Private/Components/MaterialStageExpressions/DMMSEDivide.cpp
+++ b/Engine/Plugins/VirtualProduction/DynamicMaterial/Source/DynamicMaterialEditor/Private/Components/MaterialStageExpressions/DMMSEDivide.cpp
@@ -5,6 +5,8 @@
 #include "Components/MaterialStageInputs/DMMSIValue.h"
 #include "Components/MaterialValues/DMMaterialValueFloat1.h"
 #include "Materials/MaterialExpressionDivide.h"
+
+#include UE_INLINE_GENERATED_CPP_BY_NAME(DMMSEDivide)
  
 #define LOCTEXT_NAMESPACE "DMMaterialStageExpressionDivide"
  
diff --git a/Engine/Plugins/VirtualProduction/DynamicMaterial/Source/DynamicMaterialEditor/Private/Components/MaterialStageExpressions/DMMSEDotProduct.cpp b/Engine/Plugins/VirtualProduction/DynamicMaterial/Source/DynamicMaterialEditor/Private/Components/MaterialStageExpressions/DMMSEDotProduct.cpp
--- a/Engine/Plugins/VirtualProduction/DynamicMaterial/Source/DynamicMaterialEditor/Private/Components/MaterialStageExpressions/DMMSEDotProduct.cpp
+++ b/Engine/Plugins/VirtualProduction/DynamicMaterial/Source/DynamicMaterialEditor/Private/Components/MaterialStageExpressions/DMMSEDotProduct.cpp
@@ -3,6 +3,8 @@
 #include "Components/MaterialStageExpressions/DMMSEDotProduct.h"
 #include "Materials/MaterialExpressionDotProduct.h"
 
+#include UE_INLINE_GENERATED_CPP_BY_NAME(DMMSEDotProduct)
+
 #define LOCTEXT_NAMESPACE "DMMaterialStageExpressionDotProduct"
 
 UDMMaterialStageExpressionDotProduct::UDMMaterialStageExpressionDotProduct()
diff --git a/Engine/Plugins/VirtualProduction/DynamicMaterial/Source/DynamicMaterialEditor/Private/Components/MaterialStageExpressions/DMMSEFloatModulo.cpp b/Engine/Plugins/VirtualProduction/DynamicMaterial/Source/DynamicMaterialEditor/Private/Components/MaterialStageExpressions/DMMSEFloatModulo.cpp
--- a/Engine/Plugins/VirtualProduction/DynamicMaterial/Source/DynamicMaterialEditor/Private/Components/MaterialStageExpressions/DMMSEFloatModulo.cpp
+++ b/Engine/Plugins/VirtualProduction/DynamicMaterial/Source/DynamicMaterialEditor/Private/Components/MaterialStageExpressions/DMMSEFloatModulo.cpp
@@ -5,6 +5,8 @@
 #include "Components/MaterialStageInputs/DMMSIValue.h"
 #include "Components/MaterialValues/DMMaterialValueFloat1.h"
 #include "Materials/MaterialExpressionFmod.h"
+
+#include UE_INLINE_GENERATED_CPP_BY_NAME(DMMSEFloatModulo)
  
 #define LOCTEXT_NAMESPACE "DMMaterialStageExpressionFloatModulo"
  
diff --git a/Engine/Plugins/VirtualProduction/DynamicMaterial/Source/DynamicMaterialEditor/Private/Components/MaterialStageExpressions/DMMSEFloor.cpp b/Engine/Plugins/VirtualProduction/DynamicMaterial/Source/DynamicMaterialEditor/Private/Components/MaterialStageExpressions/DMMSEFloor.cpp
--- a/Engine/Plugins/VirtualProduction/DynamicMaterial/Source/DynamicMaterialEditor/Private/Components/MaterialStageExpressions/DMMSEFloor.cpp
+++ b/Engine/Plugins/VirtualProduction/DynamicMaterial/Source/DynamicMaterialEditor/Private/Components/MaterialStageExpressions/DMMSEFloor.cpp
@@ -3,6 +3,8 @@
 #include "Components/MaterialStageExpressions/DMMSEFloor.h"
 #include "Materials/MaterialExpressionFloor.h"
 
+#include UE_INLINE_GENERATED_CPP_BY_NAME(DMMSEFloor)
+
 #define LOCTEXT_NAMESPACE "DMMaterialStageExpressionFloor"
 
 UDMMaterialStageExpressionFloor::UDMMaterialStageExpressionFloor()
diff --git a/Engine/Plugins/VirtualProduction/DynamicMaterial/Source/DynamicMaterialEditor/Private/Components/MaterialStageExpressions/DMMSEInverseLinearInterpolate.cpp b/Engine/Plugins/VirtualProduction/DynamicMaterial/Source/DynamicMaterialEditor/Private/Components/MaterialStageExpressions/DMMSEInverseLinearInterpolate.cpp
--- a/Engine/Plugins/VirtualProduction/DynamicMaterial/Source/DynamicMaterialEditor/Private/Components/MaterialStageExpressions/DMMSEInverseLinearInterpolate.cpp
+++ b/Engine/Plugins/VirtualProduction/DynamicMaterial/Source/DynamicMaterialEditor/Private/Components/MaterialStageExpressions/DMMSEInverseLinearInterpolate.cpp
@@ -3,6 +3,8 @@
 #include "Components/MaterialStageExpressions/DMMSEInverseLinearInterpolate.h"
 #include "Materials/MaterialExpressionInverseLinearInterpolate.h"
 
+#include UE_INLINE_GENERATED_CPP_BY_NAME(DMMSEInverseLinearInterpolate)
+
 #define LOCTEXT_NAMESPACE "DMMaterialStageExpressionInverseLinearInterpolate"
 
 UDMMaterialStageExpressionInverseLinearInterpolate::UDMMaterialStageExpressionInverseLinearInterpolate()
diff --git a/Engine/Plugins/VirtualProduction/DynamicMaterial/Source/DynamicMaterialEditor/Private/Components/MaterialStageExpressions/DMMSELightVector.cpp b/Engine/Plugins/VirtualProduction/DynamicMaterial/Source/DynamicMaterialEditor/Private/Components/MaterialStageExpressions/DMMSELightVector.cpp
--- a/Engine/Plugins/VirtualProduction/DynamicMaterial/Source/DynamicMaterialEditor/Private/Components/MaterialStageExpressions/DMMSELightVector.cpp
+++ b/Engine/Plugins/VirtualProduction/DynamicMaterial/Source/DynamicMaterialEditor/Private/Components/MaterialStageExpressions/DMMSELightVector.cpp
@@ -3,6 +3,8 @@
 #include "Components/MaterialStageExpressions/DMMSELightVector.h"
 #include "Materials/MaterialExpressionLightVector.h"
 
+#include UE_INLINE_GENERATED_CPP_BY_NAME(DMMSELightVector)
+
 #define LOCTEXT_NAMESPACE "DMMaterialStageExpressionLightVector"
 
 UDMMaterialStageExpressionLightVector::UDMMaterialStageExpressionLightVector()
diff --git a/Engine/Plugins/VirtualProduction/DynamicMaterial/Source/DynamicMaterialEditor/Private/Components/MaterialStageExpressions/DMMSELightmapUVs.cpp b/Engine/Plugins/VirtualProduction/DynamicMaterial/Source/DynamicMaterialEditor/Private/Components/MaterialStageExpressions/DMMSELightmapUVs.cpp
--- a/Engine/Plugins/VirtualProduction/DynamicMaterial/Source/DynamicMaterialEditor/Private/Components/MaterialStageExpressions/DMMSELightmapUVs.cpp
+++ b/Engine/Plugins/VirtualProduction/DynamicMaterial/Source/DynamicMaterialEditor/Private/Components/MaterialStageExpressions/DMMSELightmapUVs.cpp
@@ -3,6 +3,8 @@
 #include "Components/MaterialStageExpressions/DMMSELightmapUVs.h"
 #include "Materials/MaterialExpressionLightmapUVs.h"
 
+#include UE_INLINE_GENERATED_CPP_BY_NAME(DMMSELightmapUVs)
+
 #define LOCTEXT_NAMESPACE "DMMaterialStageExpressionLightmapUVs"
 
 UDMMaterialStageExpressionLightmapUVs::UDMMaterialStageExpressionLightmapUVs()
diff --git a/Engine/Plugins/VirtualProduction/DynamicMaterial/Source/DynamicMaterialEditor/Private/Components/MaterialStageExpressions/DMMSELinearInterpolate.cpp b/Engine/Plugins/VirtualProduction/DynamicMaterial/Source/DynamicMaterialEditor/Private/Components/MaterialStageExpressions/DMMSELinearInterpolate.cpp
--- a/Engine/Plugins/VirtualProduction/DynamicMaterial/Source/DynamicMaterialEditor/Private/Components/MaterialStageExpressions/DMMSELinearInterpolate.cpp
+++ b/Engine/Plugins/VirtualProduction/DynamicMaterial/Source/DynamicMaterialEditor/Private/Components/MaterialStageExpressions/DMMSELinearInterpolate.cpp
@@ -5,6 +5,8 @@
 #include "Components/MaterialStageInputs/DMMSIValue.h"
 #include "Components/MaterialValues/DMMaterialValueFloat1.h"
 #include "Materials/MaterialExpressionLinearInterpolate.h"
+
+#include UE_INLINE_GENERATED_CPP_BY_NAME(DMMSELinearInterpolate)
  
 #define LOCTEXT_NAMESPACE "DMMaterialStageExpressionLinearInterpolate"
  
diff --git a/Engine/Plugins/VirtualProduction/DynamicMaterial/Source/DynamicMaterialEditor/Private/Components/MaterialStageExpressions/DMMSELogarithm10.cpp b/Engine/Plugins/VirtualProduction/DynamicMaterial/Source/DynamicMaterialEditor/Private/Components/MaterialStageExpressions/DMMSELogarithm10.cpp
--- a/Engine/Plugins/VirtualProduction/DynamicMaterial/Source/DynamicMaterialEditor/Private/Components/MaterialStageExpressions/DMMSELogarithm10.cpp
+++ b/Engine/Plugins/VirtualProduction/DynamicMaterial/Source/DynamicMaterialEditor/Private/Components/MaterialStageExpressions/DMMSELogarithm10.cpp
@@ -3,6 +3,8 @@
 #include "Components/MaterialStageExpressions/DMMSELogarithm10.h"
 #include "Materials/MaterialExpressionLogarithm10.h"
 
+#include UE_INLINE_GENERATED_CPP_BY_NAME(DMMSELogarithm10)
+
 #define LOCTEXT_NAMESPACE "DMMaterialStageExpressionLogarithm10"
 
 UDMMaterialStageExpressionLogarithm10::UDMMaterialStageExpressionLogarithm10()
diff --git a/Engine/Plugins/VirtualProduction/DynamicMaterial/Source/DynamicMaterialEditor/Private/Components/MaterialStageExpressions/DMMSELogarithm2.cpp b/Engine/Plugins/VirtualProduction/DynamicMaterial/Source/DynamicMaterialEditor/Private/Components/MaterialStageExpressions/DMMSELogarithm2.cpp
--- a/Engine/Plugins/VirtualProduction/DynamicMaterial/Source/DynamicMaterialEditor/Private/Components/MaterialStageExpressions/DMMSELogarithm2.cpp
+++ b/Engine/Plugins/VirtualProduction/DynamicMaterial/Source/DynamicMaterialEditor/Private/Components/MaterialStageExpressions/DMMSELogarithm2.cpp
@@ -3,6 +3,8 @@
 #include "Components/MaterialStageExpressions/DMMSELogarithm2.h"
 #include "Materials/MaterialExpressionLogarithm2.h"
 
+#include UE_INLINE_GENERATED_CPP_BY_NAME(DMMSELogarithm2)
+
 #define LOCTEXT_NAMESPACE "DMMaterialStageExpressionLogarithm2"
 
 UDMMaterialStageExpressionLogarithm2::UDMMaterialStageExpressionLogarithm2()
diff --git a/Engine/Plugins/VirtualProduction/DynamicMaterial/Source/DynamicMaterialEditor/Private/Components/MaterialStageExpressions/DMMSEMathBase.cpp b/Engine/Plugins/VirtualProduction/DynamicMaterial/Source/DynamicMaterialEditor/Private/Components/MaterialStageExpressions/DMMSEMathBase.cpp
--- a/Engine/Plugins/VirtualProduction/DynamicMaterial/Source/DynamicMaterialEditor/Private/Components/MaterialStageExpressions/DMMSEMathBase.cpp
+++ b/Engine/Plugins/VirtualProduction/DynamicMaterial/Source/DynamicMaterialEditor/Private/Components/MaterialStageExpressions/DMMSEMathBase.cpp
@@ -13,6 +13,8 @@
 #include "Utils/DMInputNodeBuilder.h"
 #include "Utils/DMPrivate.h"
 
+#include UE_INLINE_GENERATED_CPP_BY_NAME(DMMSEMathBase)
+
 #define LOCTEXT_NAMESPACE "DMMaterialStageExpressionMathBase"
 
 UDMMaterialStageExpressionMathBase::UDMMaterialStageExpressionMathBase()
diff --git a/Engine/Plugins/VirtualProduction/DynamicMaterial/Source/DynamicMaterialEditor/Private/Components/MaterialStageExpressions/DMMSEMax.cpp b/Engine/Plugins/VirtualProduction/DynamicMaterial/Source/DynamicMaterialEditor/Private/Components/MaterialStageExpressions/DMMSEMax.cpp
--- a/Engine/Plugins/VirtualProduction/DynamicMaterial/Source/DynamicMaterialEditor/Private/Components/MaterialStageExpressions/DMMSEMax.cpp
+++ b/Engine/Plugins/VirtualProduction/DynamicMaterial/Source/DynamicMaterialEditor/Private/Components/MaterialStageExpressions/DMMSEMax.cpp
@@ -3,6 +3,8 @@
 #include "Components/MaterialStageExpressions/DMMSEMax.h"
 #include "Materials/MaterialExpressionMax.h"
 
+#include UE_INLINE_GENERATED_CPP_BY_NAME(DMMSEMax)
+
 #define LOCTEXT_NAMESPACE "DMMaterialStageExpressionMax"
 
 UDMMaterialStageExpressionMax::UDMMaterialStageExpressionMax()
diff --git a/Engine/Plugins/VirtualProduction/DynamicMaterial/Source/DynamicMaterialEditor/Private/Components/MaterialStageExpressions/DMMSEMin.cpp b/Engine/Plugins/VirtualProduction/DynamicMaterial/Source/DynamicMaterialEditor/Private/Components/MaterialStageExpressions/DMMSEMin.cpp
--- a/Engine/Plugins/VirtualProduction/DynamicMaterial/Source/DynamicMaterialEditor/Private/Components/MaterialStageExpressions/DMMSEMin.cpp
+++ b/Engine/Plugins/VirtualProduction/DynamicMaterial/Source/DynamicMaterialEditor/Private/Components/MaterialStageExpressions/DMMSEMin.cpp
@@ -3,6 +3,8 @@
 #include "Components/MaterialStageExpressions/DMMSEMin.h"
 #include "Materials/MaterialExpressionMin.h"
 
+#include UE_INLINE_GENERATED_CPP_BY_NAME(DMMSEMin)
+
 #define LOCTEXT_NAMESPACE "DMMaterialStageExpressionMin"
 
 UDMMaterialStageExpressionMin::UDMMaterialStageExpressionMin()
diff --git a/Engine/Plugins/VirtualProduction/DynamicMaterial/Source/DynamicMaterialEditor/Private/Components/MaterialStageExpressions/DMMSEMultiply.cpp b/Engine/Plugins/VirtualProduction/DynamicMaterial/Source/DynamicMaterialEditor/Private/Components/MaterialStageExpressions/DMMSEMultiply.cpp
--- a/Engine/Plugins/VirtualProduction/DynamicMaterial/Source/DynamicMaterialEditor/Private/Components/MaterialStageExpressions/DMMSEMultiply.cpp
+++ b/Engine/Plugins/VirtualProduction/DynamicMaterial/Source/DynamicMaterialEditor/Private/Components/MaterialStageExpressions/DMMSEMultiply.cpp
@@ -3,6 +3,8 @@
 #include "Components/MaterialStageExpressions/DMMSEMultiply.h"
 #include "Materials/MaterialExpressionMultiply.h"
 
+#include UE_INLINE_GENERATED_CPP_BY_NAME(DMMSEMultiply)
+
 #define LOCTEXT_NAMESPACE "DMMaterialStageExpressionMultiply"
 
 UDMMaterialStageExpressionMultiply::UDMMaterialStageExpressionMultiply()
diff --git a/Engine/Plugins/VirtualProduction/DynamicMaterial/Source/DynamicMaterialEditor/Private/Components/MaterialStageExpressions/DMMSEObjectBounds.cpp b/Engine/Plugins/VirtualProduction/DynamicMaterial/Source/DynamicMaterialEditor/Private/Components/MaterialStageExpressions/DMMSEObjectBounds.cpp
--- a/Engine/Plugins/VirtualProduction/DynamicMaterial/Source/DynamicMaterialEditor/Private/Components/MaterialStageExpressions/DMMSEObjectBounds.cpp
+++ b/Engine/Plugins/VirtualProduction/DynamicMaterial/Source/DynamicMaterialEditor/Private/Components/MaterialStageExpressions/DMMSEObjectBounds.cpp
@@ -3,6 +3,8 @@
 #include "Components/MaterialStageExpressions/DMMSEObjectBounds.h"
 #include "Materials/MaterialExpressionObjectBounds.h"
 
+#include UE_INLINE_GENERATED_CPP_BY_NAME(DMMSEObjectBounds)
+
 #define LOCTEXT_NAMESPACE "DMMaterialStageExpressionObjectBounds"
 
 UDMMaterialStageExpressionObjectBounds::UDMMaterialStageExpressionObjectBounds()
diff --git a/Engine/Plugins/VirtualProduction/DynamicMaterial/Source/DynamicMaterialEditor/Private/Components/MaterialStageExpressions/DMMSEObjectOrientation.cpp b/Engine/Plugins/VirtualProduction/DynamicMaterial/Source/DynamicMaterialEditor/Private/Components/MaterialStageExpressions/DMMSEObjectOrientation.cpp
--- a/Engine/Plugins/VirtualProduction/DynamicMaterial/Source/DynamicMaterialEditor/Private/Components/MaterialStageExpressions/DMMSEObjectOrientation.cpp
+++ b/Engine/Plugins/VirtualProduction/DynamicMaterial/Source/DynamicMaterialEditor/Private/Components/MaterialStageExpressions/DMMSEObjectOrientation.cpp
@@ -3,6 +3,8 @@
 #include "Components/MaterialStageExpressions/DMMSEObjectOrientation.h"
 #include "Materials/MaterialExpressionObjectOrientation.h"
 
+#include UE_INLINE_GENERATED_CPP_BY_NAME(DMMSEObjectOrientation)
+
 #define LOCTEXT_NAMESPACE "DMMaterialStageExpressionObjectOrientation"
 
 UDMMaterialStageExpressionObjectOrientation::UDMMaterialStageExpressionObjectOrientation()
diff --git a/Engine/Plugins/VirtualProduction/DynamicMaterial/Source/DynamicMaterialEditor/Private/Components/MaterialStageExpressions/DMMSEObjectPositionWS.cpp b/Engine/Plugins/VirtualProduction/DynamicMaterial/Source/DynamicMaterialEditor/Private/Components/MaterialStageExpressions/DMMSEObjectPositionWS.cpp
--- a/Engine/Plugins/VirtualProduction/DynamicMaterial/Source/DynamicMaterialEditor/Private/Components/MaterialStageExpressions/DMMSEObjectPositionWS.cpp
+++ b/Engine/Plugins/VirtualProduction/DynamicMaterial/Source/DynamicMaterialEditor/Private/Components/MaterialStageExpressions/DMMSEObjectPositionWS.cpp
@@ -3,6 +3,8 @@
 #include "Components/MaterialStageExpressions/DMMSEObjectPositionWS.h"
 #include "Materials/MaterialExpressionObjectPositionWS.h"
 
+#include UE_INLINE_GENERATED_CPP_BY_NAME(DMMSEObjectPositionWS)
+
 #define LOCTEXT_NAMESPACE "DMMaterialStageExpressionObjectPositionWS"
 
 UDMMaterialStageExpressionObjectPositionWS::UDMMaterialStageExpressionObjectPositionWS()
diff --git a/Engine/Plugins/VirtualProduction/DynamicMaterial/Source/DynamicMaterialEditor/Private/Components/MaterialStageExpressions/DMMSEObjectRadius.cpp b/Engine/Plugins/VirtualProduction/DynamicMaterial/Source/DynamicMaterialEditor/Private/Components/MaterialStageExpressions/DMMSEObjectRadius.cpp
--- a/Engine/Plugins/VirtualProduction/DynamicMaterial/Source/DynamicMaterialEditor/Private/Components/MaterialStageExpressions/DMMSEObjectRadius.cpp
+++ b/Engine/Plugins/VirtualProduction/DynamicMaterial/Source/DynamicMaterialEditor/Private/Components/MaterialStageExpressions/DMMSEObjectRadius.cpp
@@ -3,6 +3,8 @@
 #include "Components/MaterialStageExpressions/DMMSEObjectRadius.h"
 #include "Materials/MaterialExpressionObjectRadius.h"
 
+#include UE_INLINE_GENERATED_CPP_BY_NAME(DMMSEObjectRadius)
+
 #define LOCTEXT_NAMESPACE "DMMaterialStageExpressionObjectRadius"
 
 UDMMaterialStageExpressionObjectRadius::UDMMaterialStageExpressionObjectRadius()
diff --git a/Engine/Plugins/VirtualProduction/DynamicMaterial/Source/DynamicMaterialEditor/Private/Components/MaterialStageExpressions/DMMSEOneMinus.cpp b/Engine/Plugins/VirtualProduction/DynamicMaterial/Source/DynamicMaterialEditor/Private/Components/MaterialStageExpressions/DMMSEOneMinus.cpp
--- a/Engine/Plugins/VirtualProduction/DynamicMaterial/Source/DynamicMaterialEditor/Private/Components/MaterialStageExpressions/DMMSEOneMinus.cpp
+++ b/Engine/Plugins/VirtualProduction/DynamicMaterial/Source/DynamicMaterialEditor/Private/Components/MaterialStageExpressions/DMMSEOneMinus.cpp
@@ -3,6 +3,8 @@
 #include "Components/MaterialStageExpressions/DMMSEOneMinus.h"
 #include "Materials/MaterialExpressionOneMinus.h"
 
+#include UE_INLINE_GENERATED_CPP_BY_NAME(DMMSEOneMinus)
+
 #define LOCTEXT_NAMESPACE "DMMaterialStageExpressionOneMinus"
 
 UDMMaterialStageExpressionOneMinus::UDMMaterialStageExpressionOneMinus()
diff --git a/Engine/Plugins/VirtualProduction/DynamicMaterial/Source/DynamicMaterialEditor/Private/Components/MaterialStageExpressions/DMMSEPanner.cpp b/Engine/Plugins/VirtualProduction/DynamicMaterial/Source/DynamicMaterialEditor/Private/Components/MaterialStageExpressions/DMMSEPanner.cpp
--- a/Engine/Plugins/VirtualProduction/DynamicMaterial/Source/DynamicMaterialEditor/Private/Components/MaterialStageExpressions/DMMSEPanner.cpp
+++ b/Engine/Plugins/VirtualProduction/DynamicMaterial/Source/DynamicMaterialEditor/Private/Components/MaterialStageExpressions/DMMSEPanner.cpp
@@ -3,6 +3,8 @@
 #include "Components/MaterialStageExpressions/DMMSEPanner.h"
 #include "Materials/MaterialExpressionPanner.h"
 
+#include UE_INLINE_GENERATED_CPP_BY_NAME(DMMSEPanner)
+
 #define LOCTEXT_NAMESPACE "DMMaterialStageExpressionPanner"
 
 UDMMaterialStageExpressionPanner::UDMMaterialStageExpressionPanner()
diff --git a/Engine/Plugins/VirtualProduction/DynamicMaterial/Source/DynamicMaterialEditor/Private/Components/MaterialStageExpressions/DMMSEParticleColor.cpp b/Engine/Plugins/VirtualProduction/DynamicMaterial/Source/DynamicMaterialEditor/Private/Components/MaterialStageExpressions/DMMSEParticleColor.cpp
--- a/Engine/Plugins/VirtualProduction/DynamicMaterial/Source/DynamicMaterialEditor/Private/Components/MaterialStageExpressions/DMMSEParticleColor.cpp
+++ b/Engine/Plugins/VirtualProduction/DynamicMaterial/Source/DynamicMaterialEditor/Private/Components/MaterialStageExpressions/DMMSEParticleColor.cpp
@@ -3,6 +3,8 @@
 #include "Components/MaterialStageExpressions/DMMSEParticleColor.h"
 #include "Materials/MaterialExpressionParticleColor.h"
 
+#include UE_INLINE_GENERATED_CPP_BY_NAME(DMMSEParticleColor)
+
 #define LOCTEXT_NAMESPACE "DMMaterialStageExpressionParticleColor"
 
 UDMMaterialStageExpressionParticleColor::UDMMaterialStageExpressionParticleColor()
diff --git a/Engine/Plugins/VirtualProduction/DynamicMaterial/Source/DynamicMaterialEditor/Private/Components/MaterialStageExpressions/DMMSEParticleDirection.cpp b/Engine/Plugins/VirtualProduction/DynamicMaterial/Source/DynamicMaterialEditor/Private/Components/MaterialStageExpressions/DMMSEParticleDirection.cpp
--- a/Engine/Plugins/VirtualProduction/DynamicMaterial/Source/DynamicMaterialEditor/Private/Components/MaterialStageExpressions/DMMSEParticleDirection.cpp
+++ b/Engine/Plugins/VirtualProduction/DynamicMaterial/Source/DynamicMaterialEditor/Private/Components/MaterialStageExpressions/DMMSEParticleDirection.cpp
@@ -3,6 +3,8 @@
 #include "Components/MaterialStageExpressions/DMMSEParticleDirection.h"
 #include "Materials/MaterialExpressionParticleDirection.h"
 
+#include UE_INLINE_GENERATED_CPP_BY_NAME(DMMSEParticleDirection)
+
 #define LOCTEXT_NAMESPACE "DMMaterialStageExpressionParticleDirection"
 
 UDMMaterialStageExpressionParticleDirection::UDMMaterialStageExpressionParticleDirection()
diff --git a/Engine/Plugins/VirtualProduction/DynamicMaterial/Source/DynamicMaterialEditor/Private/Components/MaterialStageExpressions/DMMSEParticleMacroUV.cpp b/Engine/Plugins/VirtualProduction/DynamicMaterial/Source/DynamicMaterialEditor/Private/Components/MaterialStageExpressions/DMMSEParticleMacroUV.cpp
--- a/Engine/Plugins/VirtualProduction/DynamicMaterial/Source/DynamicMaterialEditor/Private/Components/MaterialStageExpressions/DMMSEParticleMacroUV.cpp
+++ b/Engine/Plugins/VirtualProduction/DynamicMaterial/Source/DynamicMaterialEditor/Private/Components/MaterialStageExpressions/DMMSEParticleMacroUV.cpp
@@ -3,6 +3,8 @@
 #include "Components/MaterialStageExpressions/DMMSEParticleMacroUV.h"
 #include "Materials/MaterialExpressionParticleMacroUV.h"
 
+#include UE_INLINE_GENERATED_CPP_BY_NAME(DMMSEParticleMacroUV)
+
 #define LOCTEXT_NAMESPACE "DMMaterialStageExpressionParticleMacroUV"
 
 UDMMaterialStageExpressionParticleMacroUV::UDMMaterialStageExpressionParticleMacroUV()
diff --git a/Engine/Plugins/VirtualProduction/DynamicMaterial/Source/DynamicMaterialEditor/Private/Components/MaterialStageExpressions/DMMSEParticleMotionBlurFade.cpp b/Engine/Plugins/VirtualProduction/DynamicMaterial/Source/DynamicMaterialEditor/Private/Components/MaterialStageExpressions/DMMSEParticleMotionBlurFade.cpp
--- a/Engine/Plugins/VirtualProduction/DynamicMaterial/Source/DynamicMaterialEditor/Private/Components/MaterialStageExpressions/DMMSEParticleMotionBlurFade.cpp
+++ b/Engine/Plugins/VirtualProduction/DynamicMaterial/Source/DynamicMaterialEditor/Private/Components/MaterialStageExpressions/DMMSEParticleMotionBlurFade.cpp
@@ -3,6 +3,8 @@
 #include "Components/MaterialStageExpressions/DMMSEParticleMotionBlurFade.h"
 #include "Materials/MaterialExpressionParticleMotionBlurFade.h"
 
+#include UE_INLINE_GENERATED_CPP_BY_NAME(DMMSEParticleMotionBlurFade)
+
 #define LOCTEXT_NAMESPACE "DMMaterialStageExpressionParticleMotionBlurFade"
 
 UDMMaterialStageExpressionParticleMotionBlurFade::UDMMaterialStageExpressionParticleMotionBlurFade()
diff --git a/Engine/Plugins/VirtualProduction/DynamicMaterial/Source/DynamicMaterialEditor/Private/Components/MaterialStageExpressions/DMMSEParticlePositionWS.cpp b/Engine/Plugins/VirtualProduction/DynamicMaterial/Source/DynamicMaterialEditor/Private/Components/MaterialStageExpressions/DMMSEParticlePositionWS.cpp
--- a/Engine/Plugins/VirtualProduction/DynamicMaterial/Source/DynamicMaterialEditor/Private/Components/MaterialStageExpressions/DMMSEParticlePositionWS.cpp
+++ b/Engine/Plugins/VirtualProduction/DynamicMaterial/Source/DynamicMaterialEditor/Private/Components/MaterialStageExpressions/DMMSEParticlePositionWS.cpp
@@ -3,6 +3,8 @@
 #include "Components/MaterialStageExpressions/DMMSEParticlePositionWS.h"
 #include "Materials/MaterialExpressionParticlePositionWS.h"
 
+#include UE_INLINE_GENERATED_CPP_BY_NAME(DMMSEParticlePositionWS)
+
 #define LOCTEXT_NAMESPACE "DMMaterialStageExpressionParticlePositionWS"
 
 UDMMaterialStageExpressionParticlePositionWS::UDMMaterialStageExpressionParticlePositionWS()
diff --git a/Engine/Plugins/VirtualProduction/DynamicMaterial/Source/DynamicMaterialEditor/Private/Components/MaterialStageExpressions/DMMSEParticleRadius.cpp b/Engine/Plugins/VirtualProduction/DynamicMaterial/Source/DynamicMaterialEditor/Private/Components/MaterialStageExpressions/DMMSEParticleRadius.cpp
--- a/Engine/Plugins/VirtualProduction/DynamicMaterial/Source/DynamicMaterialEditor/Private/Components/MaterialStageExpressions/DMMSEParticleRadius.cpp
+++ b/Engine/Plugins/VirtualProduction/DynamicMaterial/Source/DynamicMaterialEditor/Private/Components/MaterialStageExpressions/DMMSEParticleRadius.cpp
@@ -3,6 +3,8 @@
 #include "Components/MaterialStageExpressions/DMMSEParticleRadius.h"
 #include "Materials/MaterialExpressionParticleRadius.h"
 
+#include UE_INLINE_GENERATED_CPP_BY_NAME(DMMSEParticleRadius)
+
 #define LOCTEXT_NAMESPACE "DMMaterialStageExpressionParticleRadius"
 
 UDMMaterialStageExpressionParticleRadius::UDMMaterialStageExpressionParticleRadius()
diff --git a/Engine/Plugins/VirtualProduction/DynamicMaterial/Source/DynamicMaterialEditor/Private/Components/MaterialStageExpressions/DMMSEParticleRandom.cpp b/Engine/Plugins/VirtualProduction/DynamicMaterial/Source/DynamicMaterialEditor/Private/Components/MaterialStageExpressions/DMMSEParticleRandom.cpp
--- a/Engine/Plugins/VirtualProduction/DynamicMaterial/Source/DynamicMaterialEditor/Private/Components/MaterialStageExpressions/DMMSEParticleRandom.cpp
+++ b/Engine/Plugins/VirtualProduction/DynamicMaterial/Source/DynamicMaterialEditor/Private/Components/MaterialStageExpressions/DMMSEParticleRandom.cpp
@@ -3,6 +3,8 @@
 #include "Components/MaterialStageExpressions/DMMSEParticleRandom.h"
 #include "Materials/MaterialExpressionParticleRandom.h"
 
+#include UE_INLINE_GENERATED_CPP_BY_NAME(DMMSEParticleRandom)
+
 #define LOCTEXT_NAMESPACE "DMMaterialStageExpressionParticleRandom"
 
 UDMMaterialStageExpressionParticleRandom::UDMMaterialStageExpressionParticleRandom()
diff --git a/Engine/Plugins/VirtualProduction/DynamicMaterial/Source/DynamicMaterialEditor/Private/Components/MaterialStageExpressions/DMMSEParticleRelativeTime.cpp b/Engine/Plugins/VirtualProduction/DynamicMaterial/Source/DynamicMaterialEditor/Private/Components/MaterialStageExpressions/DMMSEParticleRelativeTime.cpp
--- a/Engine/Plugins/VirtualProduction/DynamicMaterial/Source/DynamicMaterialEditor/Private/Components/MaterialStageExpressions/DMMSEParticleRelativeTime.cpp
+++ b/Engine/Plugins/VirtualProduction/DynamicMaterial/Source/DynamicMaterialEditor/Private/Components/MaterialStageExpressions/DMMSEParticleRelativeTime.cpp
@@ -3,6 +3,8 @@
 #include "Components/MaterialStageExpressions/DMMSEParticleRelativeTime.h"
 #include "Materials/MaterialExpressionParticleRelativeTime.h"
 
+#include UE_INLINE_GENERATED_CPP_BY_NAME(DMMSEParticleRelativeTime)
+
 #define LOCTEXT_NAMESPACE "DMMaterialStageExpressionParticleRelativeTime"
 
 UDMMaterialStageExpressionParticleRelativeTime::UDMMaterialStageExpressionParticleRelativeTime()
diff --git a/Engine/Plugins/VirtualProduction/DynamicMaterial/Source/DynamicMaterialEditor/Private/Components/MaterialStageExpressions/DMMSEParticleSize.cpp b/Engine/Plugins/VirtualProduction/DynamicMaterial/Source/DynamicMaterialEditor/Private/Components/MaterialStageExpressions/DMMSEParticleSize.cpp
--- a/Engine/Plugins/VirtualProduction/DynamicMaterial/Source/DynamicMaterialEditor/Private/Components/MaterialStageExpressions/DMMSEParticleSize.cpp
+++ b/Engine/Plugins/VirtualProduction/DynamicMaterial/Source/DynamicMaterialEditor/Private/Components/MaterialStageExpressions/DMMSEParticleSize.cpp
@@ -3,6 +3,8 @@
 #include "Components/MaterialStageExpressions/DMMSEParticleSize.h"
 #include "Materials/MaterialExpressionParticleSize.h"
 
+#include UE_INLINE_GENERATED_CPP_BY_NAME(DMMSEParticleSize)
+
 #define LOCTEXT_NAMESPACE "DMMaterialStageExpressionParticleSize"
 
 UDMMaterialStageExpressionParticleSize::UDMMaterialStageExpressionParticleSize()
diff --git a/Engine/Plugins/VirtualProduction/DynamicMaterial/Source/DynamicMaterialEditor/Private/Components/MaterialStageExpressions/DMMSEParticleSpeed.cpp b/Engine/Plugins/VirtualProduction/DynamicMaterial/Source/DynamicMaterialEditor/Private/Components/MaterialStageExpressions/DMMSEParticleSpeed.cpp
--- a/Engine/Plugins/VirtualProduction/DynamicMaterial/Source/DynamicMaterialEditor/Private/Components/MaterialStageExpressions/DMMSEParticleSpeed.cpp
+++ b/Engine/Plugins/VirtualProduction/DynamicMaterial/Source/DynamicMaterialEditor/Private/Components/MaterialStageExpressions/DMMSEParticleSpeed.cpp
@@ -3,6 +3,8 @@
 #include "Components/MaterialStageExpressions/DMMSEParticleSpeed.h"
 #include "Materials/MaterialExpressionParticleSpeed.h"
 
+#include UE_INLINE_GENERATED_CPP_BY_NAME(DMMSEParticleSpeed)
+
 #define LOCTEXT_NAMESPACE "DMMaterialStageExpressionParticleSpeed"
 
 UDMMaterialStageExpressionParticleSpeed::UDMMaterialStageExpressionParticleSpeed()
diff --git a/Engine/Plugins/VirtualProduction/DynamicMaterial/Source/DynamicMaterialEditor/Private/Components/MaterialStageExpressions/DMMSEParticleSubUV.cpp b/Engine/Plugins/VirtualProduction/DynamicMaterial/Source/DynamicMaterialEditor/Private/Components/MaterialStageExpressions/DMMSEParticleSubUV.cpp
--- a/Engine/Plugins/VirtualProduction/DynamicMaterial/Source/DynamicMaterialEditor/Private/Components/MaterialStageExpressions/DMMSEParticleSubUV.cpp
+++ b/Engine/Plugins/VirtualProduction/DynamicMaterial/Source/DynamicMaterialEditor/Private/Components/MaterialStageExpressions/DMMSEParticleSubUV.cpp
@@ -3,6 +3,8 @@
 #include "Components/MaterialStageExpressions/DMMSEParticleSubUV.h"
 #include "Materials/MaterialExpressionParticleSubUV.h"
 
+#include UE_INLINE_GENERATED_CPP_BY_NAME(DMMSEParticleSubUV)
+
 #define LOCTEXT_NAMESPACE "DMMaterialStageExpressionParticleSubUV"
 
 UDMMaterialStageExpressionParticleSubUV::UDMMaterialStageExpressionParticleSubUV()
diff --git a/Engine/Plugins/VirtualProduction/DynamicMaterial/Source/DynamicMaterialEditor/Private/Components/MaterialStageExpressions/DMMSEParticleSubUVProperties.cpp b/Engine/Plugins/VirtualProduction/DynamicMaterial/Source/DynamicMaterialEditor/Private/Components/MaterialStageExpressions/DMMSEParticleSubUVProperties.cpp
--- a/Engine/Plugins/VirtualProduction/DynamicMaterial/Source/DynamicMaterialEditor/Private/Components/MaterialStageExpressions/DMMSEParticleSubUVProperties.cpp
+++ b/Engine/Plugins/VirtualProduction/DynamicMaterial/Source/DynamicMaterialEditor/Private/Components/MaterialStageExpressions/DMMSEParticleSubUVProperties.cpp
@@ -3,6 +3,8 @@
 #include "Components/MaterialStageExpressions/DMMSEParticleSubUVProperties.h"
 #include "Materials/MaterialExpressionParticleSubUVProperties.h"
 
+#include UE_INLINE_GENERATED_CPP_BY_NAME(DMMSEParticleSubUVProperties)
+
 #define LOCTEXT_NAMESPACE "DMMaterialStageExpressionParticleSubUVProperties"
 
 UDMMaterialStageExpressionParticleSubUVProperties::UDMMaterialStageExpressionParticleSubUVProperties()
diff --git a/Engine/Plugins/VirtualProduction/DynamicMaterial/Source/DynamicMaterialEditor/Private/Components/MaterialStageExpressions/DMMSEPixelDepth.cpp b/Engine/Plugins/VirtualProduction/DynamicMaterial/Source/DynamicMaterialEditor/Private/Components/MaterialStageExpressions/DMMSEPixelDepth.cpp
--- a/Engine/Plugins/VirtualProduction/DynamicMaterial/Source/DynamicMaterialEditor/Private/Components/MaterialStageExpressions/DMMSEPixelDepth.cpp
+++ b/Engine/Plugins/VirtualProduction/DynamicMaterial/Source/DynamicMaterialEditor/Private/Components/MaterialStageExpressions/DMMSEPixelDepth.cpp
@@ -3,6 +3,8 @@
 #include "Components/MaterialStageExpressions/DMMSEPixelDepth.h"
 #include "Materials/MaterialExpressionPixelDepth.h"
 
+#include UE_INLINE_GENERATED_CPP_BY_NAME(DMMSEPixelDepth)
+
 #define LOCTEXT_NAMESPACE "DMMaterialStageExpressionPixelDepth"
 
 UDMMaterialStageExpressionPixelDepth::UDMMaterialStageExpressionPixelDepth()
diff --git a/Engine/Plugins/VirtualProduction/DynamicMaterial/Source/DynamicMaterialEditor/Private/Components/MaterialStageExpressions/DMMSEPixelNormalWS.cpp b/Engine/Plugins/VirtualProduction/DynamicMaterial/Source/DynamicMaterialEditor/Private/Components/MaterialStageExpressions/DMMSEPixelNormalWS.cpp
--- a/Engine/Plugins/VirtualProduction/DynamicMaterial/Source/DynamicMaterialEditor/Private/Components/MaterialStageExpressions/DMMSEPixelNormalWS.cpp
+++ b/Engine/Plugins/VirtualProduction/DynamicMaterial/Source/DynamicMaterialEditor/Private/Components/MaterialStageExpressions/DMMSEPixelNormalWS.cpp
@@ -3,6 +3,8 @@
 #include "Components/MaterialStageExpressions/DMMSEPixelNormalWS.h"
 #include "Materials/MaterialExpressionPixelNormalWS.h"
 
+#include UE_INLINE_GENERATED_CPP_BY_NAME(DMMSEPixelNormalWS)
+
 #define LOCTEXT_NAMESPACE "DMMaterialStageExpressionPixelNormalWS"
 
 UDMMaterialStageExpressionPixelNormalWS::UDMMaterialStageExpressionPixelNormalWS()
diff --git a/Engine/Plugins/VirtualProduction/DynamicMaterial/Source/DynamicMaterialEditor/Private/Components/MaterialStageExpressions/DMMSEPower.cpp b/Engine/Plugins/VirtualProduction/DynamicMaterial/Source/DynamicMaterialEditor/Private/Components/MaterialStageExpressions/DMMSEPower.cpp
--- a/Engine/Plugins/VirtualProduction/DynamicMaterial/Source/DynamicMaterialEditor/Private/Components/MaterialStageExpressions/DMMSEPower.cpp
+++ b/Engine/Plugins/VirtualProduction/DynamicMaterial/Source/DynamicMaterialEditor/Private/Components/MaterialStageExpressions/DMMSEPower.cpp
@@ -3,6 +3,8 @@
 #include "Components/MaterialStageExpressions/DMMSEPower.h"
 #include "Materials/MaterialExpressionPower.h"
 
+#include UE_INLINE_GENERATED_CPP_BY_NAME(DMMSEPower)
+
 #define LOCTEXT_NAMESPACE "DMMaterialStageExpressionPower"
 
 UDMMaterialStageExpressionPower::UDMMaterialStageExpressionPower()
diff --git a/Engine/Plugins/VirtualProduction/DynamicMaterial/Source/DynamicMaterialEditor/Private/Components/MaterialStageExpressions/DMMSEReflectionVectorWS.cpp b/Engine/Plugins/VirtualProduction/DynamicMaterial/Source/DynamicMaterialEditor/Private/Components/MaterialStageExpressions/DMMSEReflectionVectorWS.cpp
--- a/Engine/Plugins/VirtualProduction/DynamicMaterial/Source/DynamicMaterialEditor/Private/Components/MaterialStageExpressions/DMMSEReflectionVectorWS.cpp
+++ b/Engine/Plugins/VirtualProduction/DynamicMaterial/Source/DynamicMaterialEditor/Private/Components/MaterialStageExpressions/DMMSEReflectionVectorWS.cpp
@@ -3,6 +3,8 @@
 #include "Components/MaterialStageExpressions/DMMSEReflectionVectorWS.h"
 #include "Materials/MaterialExpressionReflectionVectorWS.h"
 
+#include UE_INLINE_GENERATED_CPP_BY_NAME(DMMSEReflectionVectorWS)
+
 #define LOCTEXT_NAMESPACE "DMMaterialStageExpressionReflectionVectorWS"
 
 UDMMaterialStageExpressionReflectionVectorWS::UDMMaterialStageExpressionReflectionVectorWS()
diff --git a/Engine/Plugins/VirtualProduction/DynamicMaterial/Source/DynamicMaterialEditor/Private/Components/MaterialStageExpressions/DMMSERotateAboutAxis.cpp b/Engine/Plugins/VirtualProduction/DynamicMaterial/Source/DynamicMaterialEditor/Private/Components/MaterialStageExpressions/DMMSERotateAboutAxis.cpp
--- a/Engine/Plugins/VirtualProduction/DynamicMaterial/Source/DynamicMaterialEditor/Private/Components/MaterialStageExpressions/DMMSERotateAboutAxis.cpp
+++ b/Engine/Plugins/VirtualProduction/DynamicMaterial/Source/DynamicMaterialEditor/Private/Components/MaterialStageExpressions/DMMSERotateAboutAxis.cpp
@@ -3,6 +3,8 @@
 #include "Components/MaterialStageExpressions/DMMSERotateAboutAxis.h"
 #include "Materials/MaterialExpressionRotateAboutAxis.h"
 
+#include UE_INLINE_GENERATED_CPP_BY_NAME(DMMSERotateAboutAxis)
+
 #define LOCTEXT_NAMESPACE "DMMaterialStageExpressionRotateAboutAxis"
 
 UDMMaterialStageExpressionRotateAboutAxis::UDMMaterialStageExpressionRotateAboutAxis()
diff --git a/Engine/Plugins/VirtualProduction/DynamicMaterial/Source/DynamicMaterialEditor/Private/Components/MaterialStageExpressions/DMMSERound.cpp b/Engine/Plugins/VirtualProduction/DynamicMaterial/Source/DynamicMaterialEditor/Private/Components/MaterialStageExpressions/DMMSERound.cpp
--- a/Engine/Plugins/VirtualProduction/DynamicMaterial/Source/DynamicMaterialEditor/Private/Components/MaterialStageExpressions/DMMSERound.cpp
+++ b/Engine/Plugins/VirtualProduction/DynamicMaterial/Source/DynamicMaterialEditor/Private/Components/MaterialStageExpressions/DMMSERound.cpp
@@ -3,6 +3,8 @@
 #include "Components/MaterialStageExpressions/DMMSERound.h"
 #include "Materials/MaterialExpressionRound.h"
 
+#include UE_INLINE_GENERATED_CPP_BY_NAME(DMMSERound)
+
 #define LOCTEXT_NAMESPACE "DMMaterialStageExpressionRound"
 
 UDMMaterialStageExpressionRound::UDMMaterialStageExpressionRound()
diff --git a/Engine/Plugins/VirtualProduction/DynamicMaterial/Source/DynamicMaterialEditor/Private/Components/MaterialStageExpressions/DMMSESceneColor.cpp b/Engine/Plugins/VirtualProduction/DynamicMaterial/Source/DynamicMaterialEditor/Private/Components/MaterialStageExpressions/DMMSESceneColor.cpp
--- a/Engine/Plugins/VirtualProduction/DynamicMaterial/Source/DynamicMaterialEditor/Private/Components/MaterialStageExpressions/DMMSESceneColor.cpp
+++ b/Engine/Plugins/VirtualProduction/DynamicMaterial/Source/DynamicMaterialEditor/Private/Components/MaterialStageExpressions/DMMSESceneColor.cpp
@@ -3,6 +3,8 @@
 #include "Components/MaterialStageExpressions/DMMSESceneColor.h"
 #include "Materials/MaterialExpressionSceneColor.h"
 
+#include UE_INLINE_GENERATED_CPP_BY_NAME(DMMSESceneColor)
+
 #define LOCTEXT_NAMESPACE "DMMaterialStageExpressionSceneColor"
 
 UDMMaterialStageExpressionSceneColor::UDMMaterialStageExpressionSceneColor()
diff --git a/Engine/Plugins/VirtualProduction/DynamicMaterial/Source/DynamicMaterialEditor/Private/Components/MaterialStageExpressions/DMMSESceneDepth.cpp b/Engine/Plugins/VirtualProduction/DynamicMaterial/Source/DynamicMaterialEditor/Private/Components/MaterialStageExpressions/DMMSESceneDepth.cpp
--- a/Engine/Plugins/VirtualProduction/DynamicMaterial/Source/DynamicMaterialEditor/Private/Components/MaterialStageExpressions/DMMSESceneDepth.cpp
+++ b/Engine/Plugins/VirtualProduction/DynamicMaterial/Source/DynamicMaterialEditor/Private/Components/MaterialStageExpressions/DMMSESceneDepth.cpp
@@ -3,6 +3,8 @@
 #include "Components/MaterialStageExpressions/DMMSESceneDepth.h"
 #include "Materials/MaterialExpressionSceneDepth.h"
 
+#include UE_INLINE_GENERATED_CPP_BY_NAME(DMMSESceneDepth)
+
 #define LOCTEXT_NAMESPACE "DMMaterialStageExpressionSceneDepth"
 
 UDMMaterialStageExpressionSceneDepth::UDMMaterialStageExpressionSceneDepth()
diff --git a/Engine/Plugins/VirtualProduction/DynamicMaterial/Source/DynamicMaterialEditor/Private/Components/MaterialStageExpressions/DMMSESceneTexture.cpp b/Engine/Plugins/VirtualProduction/DynamicMaterial/Source/DynamicMaterialEditor/Private/Components/MaterialStageExpressions/DMMSESceneTexture.cpp
--- a/Engine/Plugins/VirtualProduction/DynamicMaterial/Source/DynamicMaterialEditor/Private/Components/MaterialStageExpressions/DMMSESceneTexture.cpp
+++ b/Engine/Plugins/VirtualProduction/DynamicMaterial/Source/DynamicMaterialEditor/Private/Components/MaterialStageExpressions/DMMSESceneTexture.cpp
@@ -3,6 +3,8 @@
 #include "Components/MaterialStageExpressions/DMMSESceneTexture.h"
 #include "Materials/MaterialExpressionSceneTexture.h"
 
+#include UE_INLINE_GENERATED_CPP_BY_NAME(DMMSESceneTexture)
+
 #define LOCTEXT_NAMESPACE "DMMaterialStageExpressionSceneTexture"
 
 UDMMaterialStageExpressionSceneTexture::UDMMaterialStageExpressionSceneTexture()
diff --git a/Engine/Plugins/VirtualProduction/DynamicMaterial/Source/DynamicMaterialEditor/Private/Components/MaterialStageExpressions/DMMSEScreenPosition.cpp b/Engine/Plugins/VirtualProduction/DynamicMaterial/Source/DynamicMaterialEditor/Private/Components/MaterialStageExpressions/DMMSEScreenPosition.cpp
--- a/Engine/Plugins/VirtualProduction/DynamicMaterial/Source/DynamicMaterialEditor/Private/Components/MaterialStageExpressions/DMMSEScreenPosition.cpp
+++ b/Engine/Plugins/VirtualProduction/DynamicMaterial/Source/DynamicMaterialEditor/Private/Components/MaterialStageExpressions/DMMSEScreenPosition.cpp
@@ -3,6 +3,8 @@
 #include "Components/MaterialStageExpressions/DMMSEScreenPosition.h"
 #include "Materials/MaterialExpressionScreenPosition.h"
 
+#include UE_INLINE_GENERATED_CPP_BY_NAME(DMMSEScreenPosition)
+
 #define LOCTEXT_NAMESPACE "DMMaterialStageExpressionScreenPosition"
 
 UDMMaterialStageExpressionScreenPosition::UDMMaterialStageExpressionScreenPosition()
diff --git a/Engine/Plugins/VirtualProduction/DynamicMaterial/Source/DynamicMaterialEditor/Private/Components/MaterialStageExpressions/DMMSESign.cpp b/Engine/Plugins/VirtualProduction/DynamicMaterial/Source/DynamicMaterialEditor/Private/Components/MaterialStageExpressions/DMMSESign.cpp
--- a/Engine/Plugins/VirtualProduction/DynamicMaterial/Source/DynamicMaterialEditor/Private/Components/MaterialStageExpressions/DMMSESign.cpp
+++ b/Engine/Plugins/VirtualProduction/DynamicMaterial/Source/DynamicMaterialEditor/Private/Components/MaterialStageExpressions/DMMSESign.cpp
@@ -3,6 +3,8 @@
 #include "Components/MaterialStageExpressions/DMMSESign.h"
 #include "Materials/MaterialExpressionSign.h"
 
+#include UE_INLINE_GENERATED_CPP_BY_NAME(DMMSESign)
+
 #define LOCTEXT_NAMESPACE "DMMaterialStageExpressionSign"
 
 UDMMaterialStageExpressionSign::UDMMaterialStageExpressionSign()
diff --git a/Engine/Plugins/VirtualProduction/DynamicMaterial/Source/DynamicMaterialEditor/Private/Components/MaterialStageExpressions/DMMSESine.cpp b/Engine/Plugins/VirtualProduction/DynamicMaterial/Source/DynamicMaterialEditor/Private/Components/MaterialStageExpressions/DMMSESine.cpp
--- a/Engine/Plugins/VirtualProduction/DynamicMaterial/Source/DynamicMaterialEditor/Private/Components/MaterialStageExpressions/DMMSESine.cpp
+++ b/Engine/Plugins/VirtualProduction/DynamicMaterial/Source/DynamicMaterialEditor/Private/Components/MaterialStageExpressions/DMMSESine.cpp
@@ -3,6 +3,8 @@
 #include "Components/MaterialStageExpressions/DMMSESine.h"
 #include "Materials/MaterialExpressionSine.h"
 
+#include UE_INLINE_GENERATED_CPP_BY_NAME(DMMSESine)
+
 #define LOCTEXT_NAMESPACE "DMMaterialStageExpressionSine"
 
 UDMMaterialStageExpressionSine::UDMMaterialStageExpressionSine()
diff --git a/Engine/Plugins/VirtualProduction/DynamicMaterial/Source/DynamicMaterialEditor/Private/Components/MaterialStageExpressions/DMMSESmoothStep.cpp b/Engine/Plugins/VirtualProduction/DynamicMaterial/Source/DynamicMaterialEditor/Private/Components/MaterialStageExpressions/DMMSESmoothStep.cpp
--- a/Engine/Plugins/VirtualProduction/DynamicMaterial/Source/DynamicMaterialEditor/Private/Components/MaterialStageExpressions/DMMSESmoothStep.cpp
+++ b/Engine/Plugins/VirtualProduction/DynamicMaterial/Source/DynamicMaterialEditor/Private/Components/MaterialStageExpressions/DMMSESmoothStep.cpp
@@ -3,6 +3,8 @@
 #include "Components/MaterialStageExpressions/DMMSESmoothStep.h"
 #include "Materials/MaterialExpressionSmoothStep.h"
 
+#include UE_INLINE_GENERATED_CPP_BY_NAME(DMMSESmoothStep)
+
 #define LOCTEXT_NAMESPACE "DMMaterialStageExpressionSmoothStep"
 
 UDMMaterialStageExpressionSmoothStep::UDMMaterialStageExpressionSmoothStep()
diff --git a/Engine/Plugins/VirtualProduction/DynamicMaterial/Source/DynamicMaterialEditor/Private/Components/MaterialStageExpressions/DMMSESquareRoot.cpp b/Engine/Plugins/VirtualProduction/DynamicMaterial/Source/DynamicMaterialEditor/Private/Components/MaterialStageExpressions/DMMSESquareRoot.cpp
--- a/Engine/Plugins/VirtualProduction/DynamicMaterial/Source/DynamicMaterialEditor/Private/Components/MaterialStageExpressions/DMMSESquareRoot.cpp
+++ b/Engine/Plugins/VirtualProduction/DynamicMaterial/Source/DynamicMaterialEditor/Private/Components/MaterialStageExpressions/DMMSESquareRoot.cpp
@@ -3,6 +3,8 @@
 #include "Components/MaterialStageExpressions/DMMSESquareRoot.h"
 #include "Materials/MaterialExpressionSquareRoot.h"
 
+#include UE_INLINE_GENERATED_CPP_BY_NAME(DMMSESquareRoot)
+
 #define LOCTEXT_NAMESPACE "DMMaterialStageExpressionSquareRoot"
 
 UDMMaterialStageExpressionSquareRoot::UDMMaterialStageExpressionSquareRoot()
diff --git a/Engine/Plugins/VirtualProduction/DynamicMaterial/Source/DynamicMaterialEditor/Private/Components/MaterialStageExpressions/DMMSEStep.cpp b/Engine/Plugins/VirtualProduction/DynamicMaterial/Source/DynamicMaterialEditor/Private/Components/MaterialStageExpressions/DMMSEStep.cpp
--- a/Engine/Plugins/VirtualProduction/DynamicMaterial/Source/DynamicMaterialEditor/Private/Components/MaterialStageExpressions/DMMSEStep.cpp
+++ b/Engine/Plugins/VirtualProduction/DynamicMaterial/Source/DynamicMaterialEditor/Private/Components/MaterialStageExpressions/DMMSEStep.cpp
@@ -3,6 +3,8 @@
 #include "Components/MaterialStageExpressions/DMMSEStep.h"
 #include "Materials/MaterialExpressionStep.h"
 
+#include UE_INLINE_GENERATED_CPP_BY_NAME(DMMSEStep)
+
 #define LOCTEXT_NAMESPACE "DMMaterialStageExpressionStep"
 
 UDMMaterialStageExpressionStep::UDMMaterialStageExpressionStep()
diff --git a/Engine/Plugins/VirtualProduction/DynamicMaterial/Source/DynamicMaterialEditor/Private/Components/MaterialStageExpressions/DMMSESubtract.cpp b/Engine/Plugins/VirtualProduction/DynamicMaterial/Source/DynamicMaterialEditor/Private/Components/MaterialStageExpressions/DMMSESubtract.cpp
--- a/Engine/Plugins/VirtualProduction/DynamicMaterial/Source/DynamicMaterialEditor/Private/Components/MaterialStageExpressions/DMMSESubtract.cpp
+++ b/Engine/Plugins/VirtualProduction/DynamicMaterial/Source/DynamicMaterialEditor/Private/Components/MaterialStageExpressions/DMMSESubtract.cpp
@@ -3,6 +3,8 @@
 #include "Components/MaterialStageExpressions/DMMSESubtract.h"
 #include "Materials/MaterialExpressionSubtract.h"
 
+#include UE_INLINE_GENERATED_CPP_BY_NAME(DMMSESubtract)
+
 #define LOCTEXT_NAMESPACE "DMMaterialStageExpressionSubtract"
 
 UDMMaterialStageExpressionSubtract::UDMMaterialStageExpressionSubtract()
diff --git a/Engine/Plugins/VirtualProduction/DynamicMaterial/Source/DynamicMaterialEditor/Private/Components/MaterialStageExpressions/DMMSETangent.cpp b/Engine/Plugins/VirtualProduction/DynamicMaterial/Source/DynamicMaterialEditor/Private/Components/MaterialStageExpressions/DMMSETangent.cpp
--- a/Engine/Plugins/VirtualProduction/DynamicMaterial/Source/DynamicMaterialEditor/Private/Components/MaterialStageExpressions/DMMSETangent.cpp
+++ b/Engine/Plugins/VirtualProduction/DynamicMaterial/Source/DynamicMaterialEditor/Private/Components/MaterialStageExpressions/DMMSETangent.cpp
@@ -3,6 +3,8 @@
 #include "Components/MaterialStageExpressions/DMMSETangent.h"
 #include "Materials/MaterialExpressionTangent.h"
 
+#include UE_INLINE_GENERATED_CPP_BY_NAME(DMMSETangent)
+
 #define LOCTEXT_NAMESPACE "DMMaterialStageExpressionTangent"
 
 UDMMaterialStageExpressionTangent::UDMMaterialStageExpressionTangent()
diff --git a/Engine/Plugins/VirtualProduction/DynamicMaterial/Source/DynamicMaterialEditor/Private/Components/MaterialStageExpressions/DMMSETextureCoordinate.cpp b/Engine/Plugins/VirtualProduction/DynamicMaterial/Source/DynamicMaterialEditor/Private/Components/MaterialStageExpressions/DMMSETextureCoordinate.cpp
--- a/Engine/Plugins/VirtualProduction/DynamicMaterial/Source/DynamicMaterialEditor/Private/Components/MaterialStageExpressions/DMMSETextureCoordinate.cpp
+++ b/Engine/Plugins/VirtualProduction/DynamicMaterial/Source/DynamicMaterialEditor/Private/Components/MaterialStageExpressions/DMMSETextureCoordinate.cpp
@@ -3,6 +3,8 @@
 #include "Components/MaterialStageExpressions/DMMSETextureCoordinate.h"
 #include "Materials/MaterialExpressionTextureCoordinate.h"
 
+#include UE_INLINE_GENERATED_CPP_BY_NAME(DMMSETextureCoordinate)
+
 #define LOCTEXT_NAMESPACE "DMMaterialStageExpressionTextureCoordinate"
 
 UDMMaterialStageExpressionTextureCoordinate::UDMMaterialStageExpressionTextureCoordinate()
diff --git a/Engine/Plugins/VirtualProduction/DynamicMaterial/Source/DynamicMaterialEditor/Private/Components/MaterialStageExpressions/DMMSETextureSample.cpp b/Engine/Plugins/VirtualProduction/DynamicMaterial/Source/DynamicMaterialEditor/Private/Components/MaterialStageExpressions/DMMSETextureSample.cpp
--- a/Engine/Plugins/VirtualProduction/DynamicMaterial/Source/DynamicMaterialEditor/Private/Components/MaterialStageExpressions/DMMSETextureSample.cpp
+++ b/Engine/Plugins/VirtualProduction/DynamicMaterial/Source/DynamicMaterialEditor/Private/Components/MaterialStageExpressions/DMMSETextureSample.cpp
@@ -10,6 +10,8 @@
 #include "Materials/MaterialExpressionTextureSample.h"
 #include "Model/DMMaterialBuildState.h"
 
+#include UE_INLINE_GENERATED_CPP_BY_NAME(DMMSETextureSample)
+
 #define LOCTEXT_NAMESPACE "DMMaterialStageExpressionTextureSample"
 
 UDMMaterialStageExpressionTextureSample::UDMMaterialStageExpressionTextureSample()
diff --git a/Engine/Plugins/VirtualProduction/DynamicMaterial/Source/DynamicMaterialEditor/Private/Components/MaterialStageExpressions/DMMSETextureSampleBase.cpp b/Engine/Plugins/VirtualProduction/DynamicMaterial/Source/DynamicMaterialEditor/Private/Components/MaterialStageExpressions/DMMSETextureSampleBase.cpp
--- a/Engine/Plugins/VirtualProduction/DynamicMaterial/Source/DynamicMaterialEditor/Private/Components/MaterialStageExpressions/DMMSETextureSampleBase.cpp
+++ b/Engine/Plugins/VirtualProduction/DynamicMaterial/Source/DynamicMaterialEditor/Private/Components/MaterialStageExpressions/DMMSETextureSampleBase.cpp
@@ -19,6 +19,8 @@
 #include "Model/DynamicMaterialModelEditorOnlyData.h"
 #include "Utils/DMUtils.h"
 
+#include UE_INLINE_GENERATED_CPP_BY_NAME(DMMSETextureSampleBase)
+
 #define LOCTEXT_NAMESPACE "DMMaterialStageExpressionTextureSample"
 
 UDMMaterialStageExpressionTextureSampleBase::UDMMaterialStageExpressionTextureSampleBase()
diff --git a/Engine/Plugins/VirtualProduction/DynamicMaterial/Source/DynamicMaterialEditor/Private/Components/MaterialStageExpressions/DMMSETextureSampleEdgeColor.cpp b/Engine/Plugins/VirtualProduction/DynamicMaterial/Source/DynamicMaterialEditor/Private/Components/MaterialStageExpressions/DMMSETextureSampleEdgeColor.cpp
--- a/Engine/Plugins/VirtualProduction/DynamicMaterial/Source/DynamicMaterialEditor/Private/Components/MaterialStageExpressions/DMMSETextureSampleEdgeColor.cpp
+++ b/Engine/Plugins/VirtualProduction/DynamicMaterial/Source/DynamicMaterialEditor/Private/Components/MaterialStageExpressions/DMMSETextureSampleEdgeColor.cpp
@@ -10,6 +10,8 @@
 #include "Materials/MaterialExpressionTextureSample.h"
 #include "Utils/DMPrivate.h"
 
+#include UE_INLINE_GENERATED_CPP_BY_NAME(DMMSETextureSampleEdgeColor)
+
 #define LOCTEXT_NAMESPACE "DMMaterialStageExpressionTextureSample"
 
 UDMMaterialStageExpressionTextureSampleEdgeColor::UDMMaterialStageExpressionTextureSampleEdgeColor()
diff --git a/Engine/Plugins/VirtualProduction/DynamicMaterial/Source/DynamicMaterialEditor/Private/Components/MaterialStageExpressions/DMMSETime.cpp b/Engine/Plugins/VirtualProduction/DynamicMaterial/Source/DynamicMaterialEditor/Private/Components/MaterialStageExpressions/DMMSETime.cpp
--- a/Engine/Plugins/VirtualProduction/DynamicMaterial/Source/DynamicMaterialEditor/Private/Components/MaterialStageExpressions/DMMSETime.cpp
+++ b/Engine/Plugins/VirtualProduction/DynamicMaterial/Source/DynamicMaterialEditor/Private/Components/MaterialStageExpressions/DMMSETime.cpp
@@ -3,6 +3,8 @@
 #include "Components/MaterialStageExpressions/DMMSETime.h"
 #include "Materials/MaterialExpressionTime.h"
 
+#include UE_INLINE_GENERATED_CPP_BY_NAME(DMMSETime)
+
 #define LOCTEXT_NAMESPACE "DMMaterialStageExpressionTime"
 
 UDMMaterialStageExpressionTime::UDMMaterialStageExpressionTime()
diff --git a/Engine/Plugins/VirtualProduction/DynamicMaterial/Source/DynamicMaterialEditor/Private/Components/MaterialStageExpressions/DMMSETransform.cpp b/Engine/Plugins/VirtualProduction/DynamicMaterial/Source/DynamicMaterialEditor/Private/Components/MaterialStageExpressions/DMMSETransform.cpp
--- a/Engine/Plugins/VirtualProduction/DynamicMaterial/Source/DynamicMaterialEditor/Private/Components/MaterialStageExpressions/DMMSETransform.cpp
+++ b/Engine/Plugins/VirtualProduction/DynamicMaterial/Source/DynamicMaterialEditor/Private/Components/MaterialStageExpressions/DMMSETransform.cpp
@@ -3,6 +3,8 @@
 #include "Components/MaterialStageExpressions/DMMSETransform.h"
 #include "Materials/MaterialExpressionTransform.h"
 
+#include UE_INLINE_GENERATED_CPP_BY_NAME(DMMSETransform)
+
 #define LOCTEXT_NAMESPACE "DMMaterialStageExpressionTransform"
 
 UDMMaterialStageExpressionTransform::UDMMaterialStageExpressionTransform()
diff --git a/Engine/Plugins/VirtualProduction/DynamicMaterial/Source/DynamicMaterialEditor/Private/Components/MaterialStageExpressions/DMMSETruncate.cpp b/Engine/Plugins/VirtualProduction/DynamicMaterial/Source/DynamicMaterialEditor/Private/Components/MaterialStageExpressions/DMMSETruncate.cpp
--- a/Engine/Plugins/VirtualProduction/DynamicMaterial/Source/DynamicMaterialEditor/Private/Components/MaterialStageExpressions/DMMSETruncate.cpp
+++ b/Engine/Plugins/VirtualProduction/DynamicMaterial/Source/DynamicMaterialEditor/Private/Components/MaterialStageExpressions/DMMSETruncate.cpp
@@ -3,6 +3,8 @@
 #include "Components/MaterialStageExpressions/DMMSETruncate.h"
 #include "Materials/MaterialExpressionTruncate.h"
 
+#include UE_INLINE_GENERATED_CPP_BY_NAME(DMMSETruncate)
+
 #define LOCTEXT_NAMESPACE "DMMaterialStageExpressionTruncate"
 
 UDMMaterialStageExpressionTruncate::UDMMaterialStageExpressionTruncate()
diff --git a/Engine/Plugins/VirtualProduction/DynamicMaterial/Source/DynamicMaterialEditor/Private/Components/MaterialStageExpressions/DMMSEUVRotator.cpp b/Engine/Plugins/VirtualProduction/DynamicMaterial/Source/DynamicMaterialEditor/Private/Components/MaterialStageExpressions/DMMSEUVRotator.cpp
--- a/Engine/Plugins/VirtualProduction/DynamicMaterial/Source/DynamicMaterialEditor/Private/Components/MaterialStageExpressions/DMMSEUVRotator.cpp
+++ b/Engine/Plugins/VirtualProduction/DynamicMaterial/Source/DynamicMaterialEditor/Private/Components/MaterialStageExpressions/DMMSEUVRotator.cpp
@@ -3,6 +3,8 @@
 #include "Components/MaterialStageExpressions/DMMSEUVRotator.h"
 #include "Materials/MaterialExpressionRotator.h"
 
+#include UE_INLINE_GENERATED_CPP_BY_NAME(DMMSEUVRotator)
+
 #define LOCTEXT_NAMESPACE "DMMaterialStageExpressionUVRotator"
 
 UDMMaterialStageExpressionUVRotator::UDMMaterialStageExpressionUVRotator()
diff --git a/Engine/Plugins/VirtualProduction/DynamicMaterial/Source/DynamicMaterialEditor/Private/Components/MaterialStageExpressions/DMMSEVertexColor.cpp b/Engine/Plugins/VirtualProduction/DynamicMaterial/Source/DynamicMaterialEditor/Private/Components/MaterialStageExpressions/DMMSEVertexColor.cpp
--- a/Engine/Plugins/VirtualProduction/DynamicMaterial/Source/DynamicMaterialEditor/Private/Components/MaterialStageExpressions/DMMSEVertexColor.cpp
+++ b/Engine/Plugins/VirtualProduction/DynamicMaterial/Source/DynamicMaterialEditor/Private/Components/MaterialStageExpressions/DMMSEVertexColor.cpp
@@ -3,6 +3,8 @@
 #include "Components/MaterialStageExpressions/DMMSEVertexColor.h"
 #include "Materials/MaterialExpressionVertexColor.h"
 
+#include UE_INLINE_GENERATED_CPP_BY_NAME(DMMSEVertexColor)
+
 #define LOCTEXT_NAMESPACE "DMMaterialStageExpressionVertexColor"
 
 UDMMaterialStageExpressionVertexColor::UDMMaterialStageExpressionVertexColor()
diff --git a/Engine/Plugins/VirtualProduction/DynamicMaterial/Source/DynamicMaterialEditor/Private/Components/MaterialStageExpressions/DMMSEVertexNormalWS.cpp b/Engine/Plugins/VirtualProduction/DynamicMaterial/Source/DynamicMaterialEditor/Private/Components/MaterialStageExpressions/DMMSEVertexNormalWS.cpp
--- a/Engine/Plugins/VirtualProduction/DynamicMaterial/Source/DynamicMaterialEditor/Private/Components/MaterialStageExpressions/DMMSEVertexNormalWS.cpp
+++ b/Engine/Plugins/VirtualProduction/DynamicMaterial/Source/DynamicMaterialEditor/Private/Components/MaterialStageExpressions/DMMSEVertexNormalWS.cpp
@@ -3,6 +3,8 @@
 #include "Components/MaterialStageExpressions/DMMSEVertexNormalWS.h"
 #include "Materials/MaterialExpressionVertexNormalWS.h"
 
+#include UE_INLINE_GENERATED_CPP_BY_NAME(DMMSEVertexNormalWS)
+
 #define LOCTEXT_NAMESPACE "DMMaterialStageExpressionVertexNormalWS"
 
 UDMMaterialStageExpressionVertexNormalWS::UDMMaterialStageExpressionVertexNormalWS()
diff --git a/Engine/Plugins/VirtualProduction/DynamicMaterial/Source/DynamicMaterialEditor/Private/Components/MaterialStageExpressions/DMMSEVertexTangentWS.cpp b/Engine/Plugins/VirtualProduction/DynamicMaterial/Source/DynamicMaterialEditor/Private/Components/MaterialStageExpressions/DMMSEVertexTangentWS.cpp
--- a/Engine/Plugins/VirtualProduction/DynamicMaterial/Source/DynamicMaterialEditor/Private/Components/MaterialStageExpressions/DMMSEVertexTangentWS.cpp
+++ b/Engine/Plugins/VirtualProduction/DynamicMaterial/Source/DynamicMaterialEditor/Private/Components/MaterialStageExpressions/DMMSEVertexTangentWS.cpp
@@ -3,6 +3,8 @@
 #include "Components/MaterialStageExpressions/DMMSEVertexTangentWS.h"
 #include "Materials/MaterialExpressionVertexTangentWS.h"
 
+#include UE_INLINE_GENERATED_CPP_BY_NAME(DMMSEVertexTangentWS)
+
 #define LOCTEXT_NAMESPACE "DMMaterialStageExpressionVertexTangentWS"
 
 UDMMaterialStageExpressionVertexTangentWS::UDMMaterialStageExpressionVertexTangentWS()
diff --git a/Engine/Plugins/VirtualProduction/DynamicMaterial/Source/DynamicMaterialEditor/Private/Components/MaterialStageExpressions/DMMSEViewSize.cpp b/Engine/Plugins/VirtualProduction/DynamicMaterial/Source/DynamicMaterialEditor/Private/Components/MaterialStageExpressions/DMMSEViewSize.cpp
--- a/Engine/Plugins/VirtualProduction/DynamicMaterial/Source/DynamicMaterialEditor/Private/Components/MaterialStageExpressions/DMMSEViewSize.cpp
+++ b/Engine/Plugins/VirtualProduction/DynamicMaterial/Source/DynamicMaterialEditor/Private/Components/MaterialStageExpressions/DMMSEViewSize.cpp
@@ -3,6 +3,8 @@
 #include "Components/MaterialStageExpressions/DMMSEViewSize.h"
 #include "Materials/MaterialExpressionViewSize.h"
 
+#include UE_INLINE_GENERATED_CPP_BY_NAME(DMMSEViewSize)
+
 #define LOCTEXT_NAMESPACE "DMMaterialStageExpressionViewSize"
 
 UDMMaterialStageExpressionViewSize::UDMMaterialStageExpressionViewSize()
diff --git a/Engine/Plugins/VirtualProduction/DynamicMaterial/Source/DynamicMaterialEditor/Private/Components/MaterialStageExpressions/DMMSEWorldPosition.cpp b/Engine/Plugins/VirtualProduction/DynamicMaterial/Source/DynamicMaterialEditor/Private/Components/MaterialStageExpressions/DMMSEWorldPosition.cpp
--- a/Engine/Plugins/VirtualProduction/DynamicMaterial/Source/DynamicMaterialEditor/Private/Components/MaterialStageExpressions/DMMSEWorldPosition.cpp
+++ b/Engine/Plugins/VirtualProduction/DynamicMaterial/Source/DynamicMaterialEditor/Private/Components/MaterialStageExpressions/DMMSEWorldPosition.cpp
@@ -3,6 +3,8 @@
 #include "Components/MaterialStageExpressions/DMMSEWorldPosition.h"
 #include "Materials/MaterialExpressionWorldPosition.h"
 
+#include UE_INLINE_GENERATED_CPP_BY_NAME(DMMSEWorldPosition)
+
 #define LOCTEXT_NAMESPACE "DMMaterialStageExpressionWorldPosition"
 
 UDMMaterialStageExpressionWorldPosition::UDMMaterialStageExpressionWorldPosition()
diff --git a/Engine/Plugins/VirtualProduction/DynamicMaterial/Source/DynamicMaterialEditor/Private/Components/MaterialStageExpressions/DMMSEWorldPositionNoise.cpp b/Engine/Plugins/VirtualProduction/DynamicMaterial/Source/DynamicMaterialEditor/Private/Components/MaterialStageExpressions/DMMSEWorldPositionNoise.cpp
--- a/Engine/Plugins/VirtualProduction/DynamicMaterial/Source/DynamicMaterialEditor/Private/Components/MaterialStageExpressions/DMMSEWorldPositionNoise.cpp
+++ b/Engine/Plugins/VirtualProduction/DynamicMaterial/Source/DynamicMaterialEditor/Private/Components/MaterialStageExpressions/DMMSEWorldPositionNoise.cpp
@@ -18,6 +18,8 @@
 #include "Model/DMMaterialBuildUtils.h"
 #include "Utils/DMUtils.h"
 
+#include UE_INLINE_GENERATED_CPP_BY_NAME(DMMSEWorldPositionNoise)
+
 #define LOCTEXT_NAMESPACE "DMMaterialStageExpressionTextureSample"
 
 UDMMaterialStageExpressionWorldPositionNoise::UDMMaterialStageExpressionWorldPositionNoise()
diff --git a/Engine/Plugins/VirtualProduction/DynamicMaterial/Source/DynamicMaterialEditor/Private/Components/MaterialStageGradients/DMMSGLinear.cpp b/Engine/Plugins/VirtualProduction/DynamicMaterial/Source/DynamicMaterialEditor/Private/Components/MaterialStageGradients/DMMSGLinear.cpp
--- a/Engine/Plugins/VirtualProduction/DynamicMaterial/Source/DynamicMaterialEditor/Private/Components/MaterialStageGradients/DMMSGLinear.cpp
+++ b/Engine/Plugins/VirtualProduction/DynamicMaterial/Source/DynamicMaterialEditor/Private/Components/MaterialStageGradients/DMMSGLinear.cpp
@@ -7,6 +7,8 @@
 #include "Materials/MaterialExpressionMaterialFunctionCall.h"
 #include "Materials/MaterialFunctionInterface.h"
 
+#include UE_INLINE_GENERATED_CPP_BY_NAME(DMMSGLinear)
+
 #define LOCTEXT_NAMESPACE "DMMaterialStageGradientLinear"
 
 TSoftObjectPtr<UMaterialFunctionInterface> UDMMaterialStageGradientLinear::LinearGradientNoTileFunction = TSoftObjectPtr<UMaterialFunctionInterface>(FSoftObjectPath(TEXT(
diff --git a/Engine/Plugins/VirtualProduction/DynamicMaterial/Source/DynamicMaterialEditor/Private/Components/MaterialStageGradients/DMMSGRadial.cpp b/Engine/Plugins/VirtualProduction/DynamicMaterial/Source/DynamicMaterialEditor/Private/Components/MaterialStageGradients/DMMSGRadial.cpp
--- a/Engine/Plugins/VirtualProduction/DynamicMaterial/Source/DynamicMaterialEditor/Private/Components/MaterialStageGradients/DMMSGRadial.cpp
+++ b/Engine/Plugins/VirtualProduction/DynamicMaterial/Source/DynamicMaterialEditor/Private/Components/MaterialStageGradients/DMMSGRadial.cpp
@@ -5,6 +5,8 @@
 #include "Materials/MaterialExpressionMaterialFunctionCall.h"
 #include "Materials/MaterialFunctionInterface.h"
 
+#include UE_INLINE_GENERATED_CPP_BY_NAME(DMMSGRadial)
+
 #define LOCTEXT_NAMESPACE "DMMaterialStageGradientRadial"
 
 TSoftObjectPtr<UMaterialFunctionInterface> UDMMaterialStageGradientRadial::RadialGradientFunction = TSoftObjectPtr<UMaterialFunctionInterface>(FSoftObjectPath(TEXT(
diff --git a/Engine/Plugins/VirtualProduction/DynamicMaterial/Source/DynamicMaterialEditor/Private/Components/MaterialStageInputs/DMMSIExpression.cpp b/Engine/Plugins/VirtualProduction/DynamicMaterial/Source/DynamicMaterialEditor/Private/Components/MaterialStageInputs/DMMSIExpression.cpp
--- a/Engine/Plugins/VirtualProduction/DynamicMaterial/Source/DynamicMaterialEditor/Private/Components/MaterialStageInputs/DMMSIExpression.cpp
+++ b/Engine/Plugins/VirtualProduction/DynamicMaterial/Source/DynamicMaterialEditor/Private/Components/MaterialStageInputs/DMMSIExpression.cpp
@@ -8,6 +8,8 @@
 #include "Components/DMMaterialStageExpression.h"
 #include "Model/DynamicMaterialModelEditorOnlyData.h"
 
+#include UE_INLINE_GENERATED_CPP_BY_NAME(DMMSIExpression)
+
 TArray<TStrongObjectPtr<UClass>> UDMMaterialStageInputExpression::InputExpressions = TArray<TStrongObjectPtr<UClass>>();
 
 UDMMaterialStage* UDMMaterialStageInputExpression::CreateStage(TSubclassOf<UDMMaterialStageExpression> InMaterialStageExpressionClass, UDMMaterialLayerObject* InLayer)
diff --git a/Engine/Plugins/VirtualProduction/DynamicMaterial/Source/DynamicMaterialEditor/Private/Components/MaterialStageInputs/DMMSIFunction.cpp b/Engine/Plugins/VirtualProduction/DynamicMaterial/Source/DynamicMaterialEditor/Private/Components/MaterialStageInputs/DMMSIFunction.cpp
--- a/Engine/Plugins/VirtualProduction/DynamicMaterial/Source/DynamicMaterialEditor/Private/Components/MaterialStageInputs/DMMSIFunction.cpp
+++ b/Engine/Plugins/VirtualProduction/DynamicMaterial/Source/DynamicMaterialEditor/Private/Components/MaterialStageInputs/DMMSIFunction.cpp
@@ -5,6 +5,8 @@
 #include "Components/DMMaterialStage.h"
 #include "Components/DMMaterialStageFunction.h"
 
+#include UE_INLINE_GENERATED_CPP_BY_NAME(DMMSIFunction)
+
 UDMMaterialStage* UDMMaterialStageInputFunction::CreateStage(UDMMaterialLayerObject* InLayer)
 {
 	const FDMUpdateGuard Guard;
