Run UnrealCodeFixup to add #include UE_INLINE_GENERATED_CPP_BY_NAME to files where possible

diff --git a/Engine/Plugins/VirtualProduction/DynamicMaterial/Source/DynamicMaterialEditor/Private/Components/MaterialStageInputs/DMMSIGradient.cpp b/Engine/Plugins/VirtualProduction/DynamicMaterial/Source/DynamicMaterialEditor/Private/Components/MaterialStageInputs/DMMSIGradient.cpp
--- a/Engine/Plugins/VirtualProduction/DynamicMaterial/Source/DynamicMaterialEditor/Private/Components/MaterialStageInputs/DMMSIGradient.cpp
+++ b/Engine/Plugins/VirtualProduction/DynamicMaterial/Source/DynamicMaterialEditor/Private/Components/MaterialStageInputs/DMMSIGradient.cpp
@@ -6,6 +6,8 @@
 #include "Components/DMMaterialStageExpression.h"
 #include "Components/DMMaterialStageGradient.h"
 
+#include UE_INLINE_GENERATED_CPP_BY_NAME(DMMSIGradient)
+
 TArray<TStrongObjectPtr<UClass>> UDMMaterialStageInputGradient::Gradients = TArray<TStrongObjectPtr<UClass>>();
 
 UDMMaterialStage* UDMMaterialStageInputGradient::CreateStage(TSubclassOf<UDMMaterialStageGradient> InMaterialStageGradientClass, UDMMaterialLayerObject* InLayer)
diff --git a/Engine/Plugins/VirtualProduction/DynamicMaterial/Source/DynamicMaterialEditor/Private/Components/MaterialStageInputs/DMMSISlot.cpp b/Engine/Plugins/VirtualProduction/DynamicMaterial/Source/DynamicMaterialEditor/Private/Components/MaterialStageInputs/DMMSISlot.cpp
--- a/Engine/Plugins/VirtualProduction/DynamicMaterial/Source/DynamicMaterialEditor/Private/Components/MaterialStageInputs/DMMSISlot.cpp
+++ b/Engine/Plugins/VirtualProduction/DynamicMaterial/Source/DynamicMaterialEditor/Private/Components/MaterialStageInputs/DMMSISlot.cpp
@@ -13,6 +13,8 @@
 #include "Model/DynamicMaterialModelEditorOnlyData.h"
 #include "Utils/DMPrivate.h"
 
+#include UE_INLINE_GENERATED_CPP_BY_NAME(DMMSISlot)
+
 #define LOCTEXT_NAMESPACE "DMMaterialStageInputSlot"
 
 const FString UDMMaterialStageInputSlot::SlotPathToken = FString(TEXT("Slot"));
diff --git a/Engine/Plugins/VirtualProduction/DynamicMaterial/Source/DynamicMaterialEditor/Private/Components/MaterialStageInputs/DMMSITextureUV.cpp b/Engine/Plugins/VirtualProduction/DynamicMaterial/Source/DynamicMaterialEditor/Private/Components/MaterialStageInputs/DMMSITextureUV.cpp
--- a/Engine/Plugins/VirtualProduction/DynamicMaterial/Source/DynamicMaterialEditor/Private/Components/MaterialStageInputs/DMMSITextureUV.cpp
+++ b/Engine/Plugins/VirtualProduction/DynamicMaterial/Source/DynamicMaterialEditor/Private/Components/MaterialStageInputs/DMMSITextureUV.cpp
@@ -25,6 +25,8 @@
 #include "Utils/DMMaterialFunctionLibrary.h"
 #include "Utils/DMUtils.h"
 
+#include UE_INLINE_GENERATED_CPP_BY_NAME(DMMSITextureUV)
+
 #define LOCTEXT_NAMESPACE "DMMaterialStageInputTextureUV"
 
 const FString UDMMaterialStageInputTextureUV::TextureUVPathToken = FString(TEXT("TextureUV"));
diff --git a/Engine/Plugins/VirtualProduction/DynamicMaterial/Source/DynamicMaterialEditor/Private/Components/MaterialStageInputs/DMMSIThroughput.cpp b/Engine/Plugins/VirtualProduction/DynamicMaterial/Source/DynamicMaterialEditor/Private/Components/MaterialStageInputs/DMMSIThroughput.cpp
--- a/Engine/Plugins/VirtualProduction/DynamicMaterial/Source/DynamicMaterialEditor/Private/Components/MaterialStageInputs/DMMSIThroughput.cpp
+++ b/Engine/Plugins/VirtualProduction/DynamicMaterial/Source/DynamicMaterialEditor/Private/Components/MaterialStageInputs/DMMSIThroughput.cpp
@@ -13,6 +13,8 @@
 #include "Model/DynamicMaterialModelEditorOnlyData.h"
 #include "Utils/DMPrivate.h"
 
+#include UE_INLINE_GENERATED_CPP_BY_NAME(DMMSIThroughput)
+
 #define LOCTEXT_NAMESPACE "DMMaterialStageInputThroughput"
 
 const FString UDMMaterialStageInputThroughput::SubStagePathToken = FString(TEXT("SubStage"));
diff --git a/Engine/Plugins/VirtualProduction/DynamicMaterial/Source/DynamicMaterialEditor/Private/Components/MaterialStageInputs/DMMSIValue.cpp b/Engine/Plugins/VirtualProduction/DynamicMaterial/Source/DynamicMaterialEditor/Private/Components/MaterialStageInputs/DMMSIValue.cpp
--- a/Engine/Plugins/VirtualProduction/DynamicMaterial/Source/DynamicMaterialEditor/Private/Components/MaterialStageInputs/DMMSIValue.cpp
+++ b/Engine/Plugins/VirtualProduction/DynamicMaterial/Source/DynamicMaterialEditor/Private/Components/MaterialStageInputs/DMMSIValue.cpp
@@ -20,6 +20,8 @@
 #include "Model/DynamicMaterialModel.h"
 #include "Model/DynamicMaterialModelEditorOnlyData.h"
 
+#include UE_INLINE_GENERATED_CPP_BY_NAME(DMMSIValue)
+
 #define LOCTEXT_NAMESPACE "DMMaterialStageInputValue"
 
 const FString UDMMaterialStageInputValue::ValuePathToken = FString(TEXT("Value"));
diff --git a/Engine/Plugins/VirtualProduction/DynamicMaterial/Source/DynamicMaterialEditor/Private/DMObjectMaterialProperty.cpp b/Engine/Plugins/VirtualProduction/DynamicMaterial/Source/DynamicMaterialEditor/Private/DMObjectMaterialProperty.cpp
--- a/Engine/Plugins/VirtualProduction/DynamicMaterial/Source/DynamicMaterialEditor/Private/DMObjectMaterialProperty.cpp
+++ b/Engine/Plugins/VirtualProduction/DynamicMaterial/Source/DynamicMaterialEditor/Private/DMObjectMaterialProperty.cpp
@@ -12,6 +12,8 @@
 #include "UObject/Object.h"
 #include "UObject/UnrealType.h"
 
+#include UE_INLINE_GENERATED_CPP_BY_NAME(DMObjectMaterialProperty)
+
 #define LOCTEXT_NAMESPACE "DMObjectMaterialProperty"
 
 FDMObjectMaterialProperty::FDMObjectMaterialProperty()
diff --git a/Engine/Plugins/VirtualProduction/DynamicMaterial/Source/DynamicMaterialEditor/Private/DMWorldSubsystem.cpp b/Engine/Plugins/VirtualProduction/DynamicMaterial/Source/DynamicMaterialEditor/Private/DMWorldSubsystem.cpp
--- a/Engine/Plugins/VirtualProduction/DynamicMaterial/Source/DynamicMaterialEditor/Private/DMWorldSubsystem.cpp
+++ b/Engine/Plugins/VirtualProduction/DynamicMaterial/Source/DynamicMaterialEditor/Private/DMWorldSubsystem.cpp
@@ -6,6 +6,8 @@
 #include "LevelEditor/DMLevelEditorIntegration.h"
 #include "Model/DynamicMaterialModel.h"
 
+#include UE_INLINE_GENERATED_CPP_BY_NAME(DMWorldSubsystem)
+
 UDMWorldSubsystem::UDMWorldSubsystem()
 	: KeyframeHandler(nullptr)
 {
diff --git a/Engine/Plugins/VirtualProduction/DynamicMaterial/Source/DynamicMaterialEditor/Private/DynamicMaterialEditorSettings.cpp b/Engine/Plugins/VirtualProduction/DynamicMaterial/Source/DynamicMaterialEditor/Private/DynamicMaterialEditorSettings.cpp
--- a/Engine/Plugins/VirtualProduction/DynamicMaterial/Source/DynamicMaterialEditor/Private/DynamicMaterialEditorSettings.cpp
+++ b/Engine/Plugins/VirtualProduction/DynamicMaterial/Source/DynamicMaterialEditor/Private/DynamicMaterialEditorSettings.cpp
@@ -18,6 +18,8 @@
 #include "Modules/ModuleManager.h"
 #include "UObject/UObjectIterator.h"
 
+#include UE_INLINE_GENERATED_CPP_BY_NAME(DynamicMaterialEditorSettings)
+
 #define LOCTEXT_NAMESPACE "MaterialDesignerSettings"
 
 namespace UE::DynamicMaterialEditor::Private
diff --git a/Engine/Plugins/VirtualProduction/DynamicMaterial/Source/DynamicMaterialEditor/Private/Material/AssetDefinition_DynamicMaterialInstance.cpp b/Engine/Plugins/VirtualProduction/DynamicMaterial/Source/DynamicMaterialEditor/Private/Material/AssetDefinition_DynamicMaterialInstance.cpp
--- a/Engine/Plugins/VirtualProduction/DynamicMaterial/Source/DynamicMaterialEditor/Private/Material/AssetDefinition_DynamicMaterialInstance.cpp
+++ b/Engine/Plugins/VirtualProduction/DynamicMaterial/Source/DynamicMaterialEditor/Private/Material/AssetDefinition_DynamicMaterialInstance.cpp
@@ -9,6 +9,8 @@
 #include "Model/DynamicMaterialModelDynamic.h"
 #include "ThumbnailRendering/SceneThumbnailInfoWithPrimitive.h"
 
+#include UE_INLINE_GENERATED_CPP_BY_NAME(AssetDefinition_DynamicMaterialInstance)
+
 #define LOCTEXT_NAMESPACE "AssetDefinition_DynamicMaterialInstance"
 
 FText UAssetDefinition_DynamicMaterialInstance::GetAssetDisplayName() const
diff --git a/Engine/Plugins/VirtualProduction/DynamicMaterial/Source/DynamicMaterialEditor/Private/Material/DynamicMaterialInstanceFactory.cpp b/Engine/Plugins/VirtualProduction/DynamicMaterial/Source/DynamicMaterialEditor/Private/Material/DynamicMaterialInstanceFactory.cpp
--- a/Engine/Plugins/VirtualProduction/DynamicMaterial/Source/DynamicMaterialEditor/Private/Material/DynamicMaterialInstanceFactory.cpp
+++ b/Engine/Plugins/VirtualProduction/DynamicMaterial/Source/DynamicMaterialEditor/Private/Material/DynamicMaterialInstanceFactory.cpp
@@ -11,6 +11,8 @@
 #include "Model/IDynamicMaterialModelEditorOnlyDataInterface.h"
 #include "Modules/ModuleManager.h"
 
+#include UE_INLINE_GENERATED_CPP_BY_NAME(DynamicMaterialInstanceFactory)
+
 #define LOCTEXT_NAMESPACE "MaterialDesignerInstanceFactory"
 
 UDynamicMaterialInstanceFactory::UDynamicMaterialInstanceFactory()
diff --git a/Engine/Plugins/VirtualProduction/DynamicMaterial/Source/DynamicMaterialEditor/Private/Model/AssetDefinition_DynamicMaterialModel.cpp b/Engine/Plugins/VirtualProduction/DynamicMaterial/Source/DynamicMaterialEditor/Private/Model/AssetDefinition_DynamicMaterialModel.cpp
--- a/Engine/Plugins/VirtualProduction/DynamicMaterial/Source/DynamicMaterialEditor/Private/Model/AssetDefinition_DynamicMaterialModel.cpp
+++ b/Engine/Plugins/VirtualProduction/DynamicMaterial/Source/DynamicMaterialEditor/Private/Model/AssetDefinition_DynamicMaterialModel.cpp
@@ -6,6 +6,8 @@
 #include "Model/DynamicMaterialModel.h"
 #include "Toolkits/IToolkitHost.h"
 
+#include UE_INLINE_GENERATED_CPP_BY_NAME(AssetDefinition_DynamicMaterialModel)
+
 #define LOCTEXT_NAMESPACE "AssetDefinition_DynamicMaterialModel"
 
 FText UAssetDefinition_DynamicMaterialModel::GetAssetDisplayName() const
diff --git a/Engine/Plugins/VirtualProduction/DynamicMaterial/Source/DynamicMaterialEditor/Private/Model/AssetDefinition_DynamicMaterialModelDynamic.cpp b/Engine/Plugins/VirtualProduction/DynamicMaterial/Source/DynamicMaterialEditor/Private/Model/AssetDefinition_DynamicMaterialModelDynamic.cpp
--- a/Engine/Plugins/VirtualProduction/DynamicMaterial/Source/DynamicMaterialEditor/Private/Model/AssetDefinition_DynamicMaterialModelDynamic.cpp
+++ b/Engine/Plugins/VirtualProduction/DynamicMaterial/Source/DynamicMaterialEditor/Private/Model/AssetDefinition_DynamicMaterialModelDynamic.cpp
@@ -6,6 +6,8 @@
 #include "Model/DynamicMaterialModelDynamic.h"
 #include "Toolkits/IToolkitHost.h"
 
+#include UE_INLINE_GENERATED_CPP_BY_NAME(AssetDefinition_DynamicMaterialModelDynamic)
+
 #define LOCTEXT_NAMESPACE "AssetDefinition_DynamicMaterialModelDynamic"
 
 FText UAssetDefinition_DynamicMaterialModelDynamic::GetAssetDisplayName() const
diff --git a/Engine/Plugins/VirtualProduction/DynamicMaterial/Source/DynamicMaterialEditor/Private/Model/DynamicMaterialModelEditorOnlyData.cpp b/Engine/Plugins/VirtualProduction/DynamicMaterial/Source/DynamicMaterialEditor/Private/Model/DynamicMaterialModelEditorOnlyData.cpp
--- a/Engine/Plugins/VirtualProduction/DynamicMaterial/Source/DynamicMaterialEditor/Private/Model/DynamicMaterialModelEditorOnlyData.cpp
+++ b/Engine/Plugins/VirtualProduction/DynamicMaterial/Source/DynamicMaterialEditor/Private/Model/DynamicMaterialModelEditorOnlyData.cpp
@@ -45,6 +45,8 @@
 #include "UObject/Package.h"
 #include "Utils/DMBuildRequestSubsystem.h"
 
+#include UE_INLINE_GENERATED_CPP_BY_NAME(DynamicMaterialModelEditorOnlyData)
+
 #define LOCTEXT_NAMESPACE "MaterialDesignerModel"
 
 const FString UDynamicMaterialModelEditorOnlyData::SlotsPathToken               = FString(TEXT("Slots"));
diff --git a/Engine/Plugins/VirtualProduction/DynamicMaterial/Source/DynamicMaterialEditor/Private/UI/Menus/DMMenuContext.cpp b/Engine/Plugins/VirtualProduction/DynamicMaterial/Source/DynamicMaterialEditor/Private/UI/Menus/DMMenuContext.cpp
--- a/Engine/Plugins/VirtualProduction/DynamicMaterial/Source/DynamicMaterialEditor/Private/UI/Menus/DMMenuContext.cpp
+++ b/Engine/Plugins/VirtualProduction/DynamicMaterial/Source/DynamicMaterialEditor/Private/UI/Menus/DMMenuContext.cpp
@@ -14,6 +14,8 @@
 #include "UI/Widgets/Editor/SlotEditor/SDMMaterialStage.h"
 #include "UI/Widgets/SDMMaterialEditor.h"
 
+#include UE_INLINE_GENERATED_CPP_BY_NAME(DMMenuContext)
+
 UDMMenuContext* UDMMenuContext::Create(const TWeakPtr<SDMMaterialEditor>& InEditorWidget, const TWeakPtr<SDMMaterialStage>& InStageWidget, UDMMaterialLayerObject* InLayerObject)
 {
 	UDMMenuContext* Context = NewObject<UDMMenuContext>();
diff --git a/Engine/Plugins/VirtualProduction/DynamicMaterial/Source/DynamicMaterialEditor/Private/UI/Utils/DynamicMaterialInstanceThumbnailRenderer.cpp b/Engine/Plugins/VirtualProduction/DynamicMaterial/Source/DynamicMaterialEditor/Private/UI/Utils/DynamicMaterialInstanceThumbnailRenderer.cpp
--- a/Engine/Plugins/VirtualProduction/DynamicMaterial/Source/DynamicMaterialEditor/Private/UI/Utils/DynamicMaterialInstanceThumbnailRenderer.cpp
+++ b/Engine/Plugins/VirtualProduction/DynamicMaterial/Source/DynamicMaterialEditor/Private/UI/Utils/DynamicMaterialInstanceThumbnailRenderer.cpp
@@ -15,6 +15,8 @@
 #include "UI/Utils/DynamicMaterialInstanceThumbnailScene.h"
 #include "Widgets/Images/SImage.h"
 
+#include UE_INLINE_GENERATED_CPP_BY_NAME(DynamicMaterialInstanceThumbnailRenderer)
+
 void UDynamicMaterialInstanceThumbnailRenderer::Draw(UObject* InObject, int32 InX, int32 InY, uint32 InWidth, uint32 InHeight, 
 	FRenderTarget* InRenderTarget, FCanvas* InCanvas, bool bInAdditionalViewFamily)
 {
diff --git a/Engine/Plugins/VirtualProduction/DynamicMaterial/Source/DynamicMaterialEditor/Private/UI/Widgets/Editor/SDMMaterialPreview.cpp b/Engine/Plugins/VirtualProduction/DynamicMaterial/Source/DynamicMaterialEditor/Private/UI/Widgets/Editor/SDMMaterialPreview.cpp
--- a/Engine/Plugins/VirtualProduction/DynamicMaterial/Source/DynamicMaterialEditor/Private/UI/Widgets/Editor/SDMMaterialPreview.cpp
+++ b/Engine/Plugins/VirtualProduction/DynamicMaterial/Source/DynamicMaterialEditor/Private/UI/Widgets/Editor/SDMMaterialPreview.cpp
@@ -30,6 +30,8 @@
 #include "UnrealEdGlobals.h"
 #include "Widgets/SNullWidget.h"
 
+#include UE_INLINE_GENERATED_CPP_BY_NAME(SDMMaterialPreview)
+
 #define LOCTEXT_NAMESPACE "SDMMaterialPreview"
 
 void SDMMaterialPreview::PrivateRegisterAttributes(FSlateAttributeDescriptor::FInitializer&)
diff --git a/Engine/Plugins/VirtualProduction/DynamicMaterial/Source/DynamicMaterialEditor/Private/UI/Widgets/Editor/SlotEditor/SDMMaterialLayerBlendMode.cpp b/Engine/Plugins/VirtualProduction/DynamicMaterial/Source/DynamicMaterialEditor/Private/UI/Widgets/Editor/SlotEditor/SDMMaterialLayerBlendMode.cpp
--- a/Engine/Plugins/VirtualProduction/DynamicMaterial/Source/DynamicMaterialEditor/Private/UI/Widgets/Editor/SlotEditor/SDMMaterialLayerBlendMode.cpp
+++ b/Engine/Plugins/VirtualProduction/DynamicMaterial/Source/DynamicMaterialEditor/Private/UI/Widgets/Editor/SlotEditor/SDMMaterialLayerBlendMode.cpp
@@ -44,6 +44,8 @@
 #include "Widgets/SNullWidget.h"
 #include "Widgets/Text/STextBlock.h"
 
+#include UE_INLINE_GENERATED_CPP_BY_NAME(SDMMaterialLayerBlendMode)
+
 #define LOCTEXT_NAMESPACE "SDMMaterialLayerBlendMode"
 
 namespace UE::DynamicMaterialEditor::Private
diff --git a/Engine/Plugins/VirtualProduction/DynamicMaterial/Source/DynamicMaterialEditor/Private/Utils/DMBuildRequestSubsystem.cpp b/Engine/Plugins/VirtualProduction/DynamicMaterial/Source/DynamicMaterialEditor/Private/Utils/DMBuildRequestSubsystem.cpp
--- a/Engine/Plugins/VirtualProduction/DynamicMaterial/Source/DynamicMaterialEditor/Private/Utils/DMBuildRequestSubsystem.cpp
+++ b/Engine/Plugins/VirtualProduction/DynamicMaterial/Source/DynamicMaterialEditor/Private/Utils/DMBuildRequestSubsystem.cpp
@@ -6,6 +6,8 @@
 #include "DMEDefs.h"
 #include "Editor.h"
 
+#include UE_INLINE_GENERATED_CPP_BY_NAME(DMBuildRequestSubsystem)
+
 UDMBuildRequestSubsystem* UDMBuildRequestSubsystem::Get()
 {
 	if (!GEditor)
diff --git a/Engine/Plugins/VirtualProduction/DynamicMaterial/Source/DynamicMaterialEditor/Private/Utils/DMMaterialEffectStackPresetSubsystem.cpp b/Engine/Plugins/VirtualProduction/DynamicMaterial/Source/DynamicMaterialEditor/Private/Utils/DMMaterialEffectStackPresetSubsystem.cpp
--- a/Engine/Plugins/VirtualProduction/DynamicMaterial/Source/DynamicMaterialEditor/Private/Utils/DMMaterialEffectStackPresetSubsystem.cpp
+++ b/Engine/Plugins/VirtualProduction/DynamicMaterial/Source/DynamicMaterialEditor/Private/Utils/DMMaterialEffectStackPresetSubsystem.cpp
@@ -15,6 +15,8 @@
 #include "Serialization/JsonWriter.h"
 #include "Utils/DMJsonUtils.h"
 
+#include UE_INLINE_GENERATED_CPP_BY_NAME(DMMaterialEffectStackPresetSubsystem)
+
 #define LOCTEXT_NAMESPACE "DMMaterialEffectStackPresetSubsystem"
 
 namespace UE::DynamicMaterialEditor::Private
diff --git a/Engine/Plugins/VirtualProduction/DynamicMaterial/Source/DynamicMaterialEditor/Private/Utils/DMMaterialFunctionFunctionLibrary.cpp b/Engine/Plugins/VirtualProduction/DynamicMaterial/Source/DynamicMaterialEditor/Private/Utils/DMMaterialFunctionFunctionLibrary.cpp
--- a/Engine/Plugins/VirtualProduction/DynamicMaterial/Source/DynamicMaterialEditor/Private/Utils/DMMaterialFunctionFunctionLibrary.cpp
+++ b/Engine/Plugins/VirtualProduction/DynamicMaterial/Source/DynamicMaterialEditor/Private/Utils/DMMaterialFunctionFunctionLibrary.cpp
@@ -18,6 +18,8 @@
 #include "PropertyHandle.h"
 #include "Templates/SharedPointer.h"
 
+#include UE_INLINE_GENERATED_CPP_BY_NAME(DMMaterialFunctionFunctionLibrary)
+
 #define LOCTEXT_NAMESPACE "DMMaterialFunctionFunctionLibrary"
 
 void UDMMaterialFunctionFunctionLibrary::ApplyMetaData(const FFunctionExpressionInput& InFunctionInput,
diff --git a/Engine/Plugins/VirtualProduction/DynamicMaterial/Source/DynamicMaterialEditor/Private/Utils/DMMaterialInstanceFunctionLibrary.cpp b/Engine/Plugins/VirtualProduction/DynamicMaterial/Source/DynamicMaterialEditor/Private/Utils/DMMaterialInstanceFunctionLibrary.cpp
--- a/Engine/Plugins/VirtualProduction/DynamicMaterial/Source/DynamicMaterialEditor/Private/Utils/DMMaterialInstanceFunctionLibrary.cpp
+++ b/Engine/Plugins/VirtualProduction/DynamicMaterial/Source/DynamicMaterialEditor/Private/Utils/DMMaterialInstanceFunctionLibrary.cpp
@@ -23,6 +23,8 @@
 #include "UObject/UObjectGlobals.h"
 #include "Utils/DMPrivate.h"
 
+#include UE_INLINE_GENERATED_CPP_BY_NAME(DMMaterialInstanceFunctionLibrary)
+
 #define LOCTEXT_NAMESPACE "DMMaterialInstanceFunctionLibrary"
 
 TArray<FDMObjectMaterialProperty> UDMMaterialInstanceFunctionLibrary::GetActorMaterialProperties(AActor* InActor)
diff --git a/Engine/Plugins/VirtualProduction/DynamicMaterial/Source/DynamicMaterialEditor/Private/Utils/DMMaterialModelFunctionLibrary.cpp b/Engine/Plugins/VirtualProduction/DynamicMaterial/Source/DynamicMaterialEditor/Private/Utils/DMMaterialModelFunctionLibrary.cpp
--- a/Engine/Plugins/VirtualProduction/DynamicMaterial/Source/DynamicMaterialEditor/Private/Utils/DMMaterialModelFunctionLibrary.cpp
+++ b/Engine/Plugins/VirtualProduction/DynamicMaterial/Source/DynamicMaterialEditor/Private/Utils/DMMaterialModelFunctionLibrary.cpp
@@ -27,6 +27,8 @@
 #include "UObject/UObjectGlobals.h"
 #include "Utils/DMPrivate.h"
 
+#include UE_INLINE_GENERATED_CPP_BY_NAME(DMMaterialModelFunctionLibrary)
+
 #define LOCTEXT_NAMESPACE "DMMaterialModelFunctionLibrary"
 
 UDynamicMaterialInstance* UDMMaterialModelFunctionLibrary::ExportMaterial(UDynamicMaterialModelBase* InMaterialModelBase)
diff --git a/Engine/Plugins/VirtualProduction/DynamicMaterial/Source/DynamicMaterialEditor/Private/Utils/DMMaterialSlotFunctionLibrary.cpp b/Engine/Plugins/VirtualProduction/DynamicMaterial/Source/DynamicMaterialEditor/Private/Utils/DMMaterialSlotFunctionLibrary.cpp
--- a/Engine/Plugins/VirtualProduction/DynamicMaterial/Source/DynamicMaterialEditor/Private/Utils/DMMaterialSlotFunctionLibrary.cpp
+++ b/Engine/Plugins/VirtualProduction/DynamicMaterial/Source/DynamicMaterialEditor/Private/Utils/DMMaterialSlotFunctionLibrary.cpp
@@ -26,6 +26,8 @@
 #include "Model/DynamicMaterialModelEditorOnlyData.h"
 #include "Utils/DMMaterialStageFunctionLibrary.h"
 
+#include UE_INLINE_GENERATED_CPP_BY_NAME(DMMaterialSlotFunctionLibrary)
+
 #define LOCTEXT_NAMESPACE "UDMMaterialSlotFunctionLibrary"
 
 UDMMaterialLayerObject* UDMMaterialSlotFunctionLibrary::AddNewLayer(UDMMaterialSlot* InSlot, UDMMaterialStage* InNewBaseStage, 
diff --git a/Engine/Plugins/VirtualProduction/DynamicMaterial/Source/DynamicMaterialEditor/Private/Utils/DMMaterialStageFunctionLibrary.cpp b/Engine/Plugins/VirtualProduction/DynamicMaterial/Source/DynamicMaterialEditor/Private/Utils/DMMaterialStageFunctionLibrary.cpp
--- a/Engine/Plugins/VirtualProduction/DynamicMaterial/Source/DynamicMaterialEditor/Private/Utils/DMMaterialStageFunctionLibrary.cpp
+++ b/Engine/Plugins/VirtualProduction/DynamicMaterial/Source/DynamicMaterialEditor/Private/Utils/DMMaterialStageFunctionLibrary.cpp
@@ -12,6 +12,8 @@
 #include "Components/MaterialValues/DMMaterialValueRenderTarget.h"
 #include "Components/MaterialValues/DMMaterialValueTexture.h"
 
+#include UE_INLINE_GENERATED_CPP_BY_NAME(DMMaterialStageFunctionLibrary)
+
 UDMMaterialStageInputValue* UDMMaterialStageFunctionLibrary::FindDefaultStageOpacityInputValue(UDMMaterialStage* InStage)
 {
 	if (!IsValid(InStage))
diff --git a/Engine/Plugins/VirtualProduction/DynamicMaterial/Source/DynamicMaterialEditor/Private/Utils/DMTextureSetFunctionLibrary.cpp b/Engine/Plugins/VirtualProduction/DynamicMaterial/Source/DynamicMaterialEditor/Private/Utils/DMTextureSetFunctionLibrary.cpp
--- a/Engine/Plugins/VirtualProduction/DynamicMaterial/Source/DynamicMaterialEditor/Private/Utils/DMTextureSetFunctionLibrary.cpp
+++ b/Engine/Plugins/VirtualProduction/DynamicMaterial/Source/DynamicMaterialEditor/Private/Utils/DMTextureSetFunctionLibrary.cpp
@@ -21,6 +21,8 @@
 #include "Utils/DMUtils.h"
 #include "Widgets/Notifications/SNotificationList.h"
 
+#include UE_INLINE_GENERATED_CPP_BY_NAME(DMTextureSetFunctionLibrary)
+
 #define LOCTEXT_NAMESPACE "DMTextureSetFunctionLibrary"
 
 bool UDMTextureSetFunctionLibrary::AddTextureSetToModel(UDynamicMaterialModelEditorOnlyData* InEditorOnlyData, UDMTextureSet* InTextureSet,
diff --git a/Engine/Plugins/VirtualProduction/DynamicMaterial/Source/DynamicMaterialTextureSet/Private/DMTextureSet.cpp b/Engine/Plugins/VirtualProduction/DynamicMaterial/Source/DynamicMaterialTextureSet/Private/DMTextureSet.cpp
--- a/Engine/Plugins/VirtualProduction/DynamicMaterial/Source/DynamicMaterialTextureSet/Private/DMTextureSet.cpp
+++ b/Engine/Plugins/VirtualProduction/DynamicMaterial/Source/DynamicMaterialTextureSet/Private/DMTextureSet.cpp
@@ -5,6 +5,8 @@
 #include "DMTextureSetMaterialProperty.h"
 #include "Engine/Texture.h"
 
+#include UE_INLINE_GENERATED_CPP_BY_NAME(DMTextureSet)
+
 UDMTextureSet::UDMTextureSet()
 {
 	UEnum* MaterialPropertyEnum = StaticEnum<EDMTextureSetMaterialProperty>();
diff --git a/Engine/Plugins/VirtualProduction/DynamicMaterial/Source/DynamicMaterialTextureSetEditor/Private/AssetDefinition_DMTextureSet.cpp b/Engine/Plugins/VirtualProduction/DynamicMaterial/Source/DynamicMaterialTextureSetEditor/Private/AssetDefinition_DMTextureSet.cpp
--- a/Engine/Plugins/VirtualProduction/DynamicMaterial/Source/DynamicMaterialTextureSetEditor/Private/AssetDefinition_DMTextureSet.cpp
+++ b/Engine/Plugins/VirtualProduction/DynamicMaterial/Source/DynamicMaterialTextureSetEditor/Private/AssetDefinition_DMTextureSet.cpp
@@ -4,6 +4,8 @@
 
 #include "DMTextureSet.h"
 
+#include UE_INLINE_GENERATED_CPP_BY_NAME(AssetDefinition_DMTextureSet)
+
 #define LOCTEXT_NAMESPACE "AssetDefinition_DMTextureSet"
 
 FText UAssetDefinition_DMTextureSet::GetAssetDisplayName() const
diff --git a/Engine/Plugins/VirtualProduction/DynamicMaterial/Source/DynamicMaterialTextureSetEditor/Private/DMTextureSetBlueprintFunctionLibrary.cpp b/Engine/Plugins/VirtualProduction/DynamicMaterial/Source/DynamicMaterialTextureSetEditor/Private/DMTextureSetBlueprintFunctionLibrary.cpp
--- a/Engine/Plugins/VirtualProduction/DynamicMaterial/Source/DynamicMaterialTextureSetEditor/Private/DMTextureSetBlueprintFunctionLibrary.cpp
+++ b/Engine/Plugins/VirtualProduction/DynamicMaterial/Source/DynamicMaterialTextureSetEditor/Private/DMTextureSetBlueprintFunctionLibrary.cpp
@@ -11,6 +11,8 @@
 #include "Widgets/SDMTextureSetBuilder.h"
 #include "Widgets/SWindow.h"
 
+#include UE_INLINE_GENERATED_CPP_BY_NAME(DMTextureSetBlueprintFunctionLibrary)
+
 #define LOCTEXT_NAMESPACE "UDMTextureSetBlueprintFunctionLibrary"
 
 UDMTextureSet* UDMTextureSetBlueprintFunctionLibrary::CreateTextureSetFromAssets(const TArray<FAssetData>& InAssets)
diff --git a/Engine/Plugins/VirtualProduction/DynamicMaterial/Source/DynamicMaterialTextureSetEditor/Private/DMTextureSetFactory.cpp b/Engine/Plugins/VirtualProduction/DynamicMaterial/Source/DynamicMaterialTextureSetEditor/Private/DMTextureSetFactory.cpp
--- a/Engine/Plugins/VirtualProduction/DynamicMaterial/Source/DynamicMaterialTextureSetEditor/Private/DMTextureSetFactory.cpp
+++ b/Engine/Plugins/VirtualProduction/DynamicMaterial/Source/DynamicMaterialTextureSetEditor/Private/DMTextureSetFactory.cpp
@@ -5,6 +5,8 @@
 #include "DMTextureSet.h"
 #include "Modules/ModuleManager.h"
 
+#include UE_INLINE_GENERATED_CPP_BY_NAME(DMTextureSetFactory)
+
 #define LOCTEXT_NAMESPACE "MaterialDesignerInstanceFactory"
 
 UDMTextureSetFactory::UDMTextureSetFactory()
diff --git a/Engine/Plugins/VirtualProduction/DynamicMaterial/Source/DynamicMaterialTextureSetEditor/Private/DMTextureSetFilter.cpp b/Engine/Plugins/VirtualProduction/DynamicMaterial/Source/DynamicMaterialTextureSetEditor/Private/DMTextureSetFilter.cpp
--- a/Engine/Plugins/VirtualProduction/DynamicMaterial/Source/DynamicMaterialTextureSetEditor/Private/DMTextureSetFilter.cpp
+++ b/Engine/Plugins/VirtualProduction/DynamicMaterial/Source/DynamicMaterialTextureSetEditor/Private/DMTextureSetFilter.cpp
@@ -4,6 +4,8 @@
 
 #include "DMTextureSetSettings.h"
 
+#include UE_INLINE_GENERATED_CPP_BY_NAME(DMTextureSetFilter)
+
 FDMTextureSetFilter::FDMTextureSetFilter()
 	: FilterStrings({TEXT("_")})
 	, MaterialProperties({{EDMTextureSetMaterialProperty::BaseColor, EDMTextureChannelMask::RGBA}})
diff --git a/Engine/Plugins/VirtualProduction/DynamicMaterial/Source/DynamicMaterialTextureSetEditor/Private/DMTextureSetSettings.cpp b/Engine/Plugins/VirtualProduction/DynamicMaterial/Source/DynamicMaterialTextureSetEditor/Private/DMTextureSetSettings.cpp
--- a/Engine/Plugins/VirtualProduction/DynamicMaterial/Source/DynamicMaterialTextureSetEditor/Private/DMTextureSetSettings.cpp
+++ b/Engine/Plugins/VirtualProduction/DynamicMaterial/Source/DynamicMaterialTextureSetEditor/Private/DMTextureSetSettings.cpp
@@ -4,6 +4,8 @@
 
 #include "DMTextureSetMaterialProperty.h"
 
+#include UE_INLINE_GENERATED_CPP_BY_NAME(DMTextureSetSettings)
+
 UDMTextureSetSettings::UDMTextureSetSettings()
 {
 	CategoryName = TEXT("Plugins");
diff --git a/Engine/Plugins/VirtualProduction/DynamicMaterial/Source/DynamicMaterialTextureSetEditor/Private/Widgets/DMTextureSetBuilderEntry.cpp b/Engine/Plugins/VirtualProduction/DynamicMaterial/Source/DynamicMaterialTextureSetEditor/Private/Widgets/DMTextureSetBuilderEntry.cpp
--- a/Engine/Plugins/VirtualProduction/DynamicMaterial/Source/DynamicMaterialTextureSetEditor/Private/Widgets/DMTextureSetBuilderEntry.cpp
+++ b/Engine/Plugins/VirtualProduction/DynamicMaterial/Source/DynamicMaterialTextureSetEditor/Private/Widgets/DMTextureSetBuilderEntry.cpp
@@ -4,6 +4,8 @@
 
 #include "DMTextureSetMaterialProperty.h"
 
+#include UE_INLINE_GENERATED_CPP_BY_NAME(DMTextureSetBuilderEntry)
+
 FDMTextureSetBuilderEntry::FDMTextureSetBuilderEntry()
 	: FDMTextureSetBuilderEntry(EDMTextureSetMaterialProperty::BaseColor, nullptr, EDMTextureChannelMask::RGBA)
 {
