Ran UnrealCodeFixup to remove nullptr initialization of TSubclassOf

diff --git a/Engine/Plugins/VirtualProduction/DynamicMaterial/Source/DynamicMaterialEditor/Public/Components/DMMaterialEffectStack.h b/Engine/Plugins/VirtualProduction/DynamicMaterial/Source/DynamicMaterialEditor/Public/Components/DMMaterialEffectStack.h
--- a/Engine/Plugins/VirtualProduction/DynamicMaterial/Source/DynamicMaterialEditor/Public/Components/DMMaterialEffectStack.h
+++ b/Engine/Plugins/VirtualProduction/DynamicMaterial/Source/DynamicMaterialEditor/Public/Components/DMMaterialEffectStack.h
@@ -18,7 +18,7 @@ struct FDMMaterialBuildState;
 
 struct FDMMaterialEffectJson
 {
-	TSubclassOf<UDMMaterialEffect> Class = nullptr;
+	TSubclassOf<UDMMaterialEffect> Class;
 	TSharedPtr<FJsonValue> Data = nullptr;
 };
 
