Material Designer: Fixed packaging issue related to class export.

diff --git a/Engine/Plugins/VirtualProduction/DynamicMaterial/Source/DynamicMaterial/Private/Utils/DMJsonUtils.cpp b/Engine/Plugins/VirtualProduction/DynamicMaterial/Source/DynamicMaterial/Private/Utils/DMJsonUtils.cpp
--- a/Engine/Plugins/VirtualProduction/DynamicMaterial/Source/DynamicMaterial/Private/Utils/DMJsonUtils.cpp
+++ b/Engine/Plugins/VirtualProduction/DynamicMaterial/Source/DynamicMaterial/Private/Utils/DMJsonUtils.cpp
@@ -1,16 +1,20 @@
 // Copyright Epic Games, Inc. All Rights Reserved.
 
-#if WITH_EDITOR
-
 #include "Utils/DMJsonUtils.h"
 
+#if WITH_EDITOR
+
 #include "DMDefs.h"
 #include "Dom/JsonValue.h"
 #include "DynamicMaterialModule.h"
 #include "JsonObjectConverter.h"
 
+#endif
+
 #include UE_INLINE_GENERATED_CPP_BY_NAME(DMJsonUtils)
 
+#if WITH_EDITOR
+
 namespace UE::DynamicMaterial::Private
 {
 	const FString JsonKey_Class = TEXT("Class");
