[DynamicMaterialEditor] Add inline editable text style for layer names widget

diff --git a/Engine/Plugins/VirtualProduction/DynamicMaterial/Source/DynamicMaterialEditor/Private/DynamicMaterialEditorStyle.cpp b/Engine/Plugins/VirtualProduction/DynamicMaterial/Source/DynamicMaterialEditor/Private/DynamicMaterialEditorStyle.cpp
--- a/Engine/Plugins/VirtualProduction/DynamicMaterial/Source/DynamicMaterialEditor/Private/DynamicMaterialEditorStyle.cpp
+++ b/Engine/Plugins/VirtualProduction/DynamicMaterial/Source/DynamicMaterialEditor/Private/DynamicMaterialEditorStyle.cpp
@@ -7,6 +7,7 @@
 #include "Interfaces/IPluginManager.h"
 #include "Styling/SlateStyleMacros.h"
 #include "Styling/SlateStyleRegistry.h"
+#include "Styling/SlateTypes.h"
 #include "Styling/StarshipCoreStyle.h"
 #include "Styling/StyleColors.h"
 
@@ -446,6 +447,12 @@ void FDynamicMaterialEditorStyle::SetupTextStyles()
 
 	Set("Font.Stage.Details.Small.Bold", FTextBlockStyle(NormalTextStyle)
 		.SetFont(IDetailLayoutBuilder::GetDetailFontBold()));
+
+	const FEditableTextBoxStyle DefaultInlineEditableTextBoxStyle(FAppStyle::GetWidgetStyle<FEditableTextBoxStyle>(TEXT("InlineEditableTextBoxStyle")));
+
+	FEditableTextBoxStyle InlineEditableTextBoxStyle = DefaultInlineEditableTextBoxStyle;
+	InlineEditableTextBoxStyle.SetBackgroundColor(FStyleColors::Transparent);
+	Set("InlineEditableTextBoxStyle", InlineEditableTextBoxStyle);
 }
 
 void FDynamicMaterialEditorStyle::SetupComponentIcons()
