[DynamicMaterialEditor] Use foreground color instead of style color to remove warning log and be consistent with the rest of the engine UI

diff --git a/Engine/Plugins/VirtualProduction/DynamicMaterial/Source/DynamicMaterialEditor/Private/UI/Widgets/Editor/SDMMaterialSlotEditor.cpp b/Engine/Plugins/VirtualProduction/DynamicMaterial/Source/DynamicMaterialEditor/Private/UI/Widgets/Editor/SDMMaterialSlotEditor.cpp
--- a/Engine/Plugins/VirtualProduction/DynamicMaterial/Source/DynamicMaterialEditor/Private/UI/Widgets/Editor/SDMMaterialSlotEditor.cpp
+++ b/Engine/Plugins/VirtualProduction/DynamicMaterial/Source/DynamicMaterialEditor/Private/UI/Widgets/Editor/SDMMaterialSlotEditor.cpp
@@ -819,6 +819,7 @@ TSharedRef<SWidget> SDMMaterialSlotEditor::CreateSlot_LayerSettings()
 		[
 			SNew(STextBlock)
 			.TextStyle(FDynamicMaterialEditorStyle::Get(), "SlotLayerInfo")
+			.ColorAndOpacity(FSlateColor::UseForeground())
 			.Text(GetLayerButtonsDescription())
 		]
 
@@ -839,6 +840,7 @@ TSharedRef<SWidget> SDMMaterialSlotEditor::CreateSlot_LayerSettings()
 			[
 				SNew(SImage)
 				.Image(FDynamicMaterialEditorStyle::Get().GetBrush("EffectsView.Row.Fx"))
+				.ColorAndOpacity(FSlateColor::UseForeground())
 				.DesiredSizeOverride(FVector2D(16.0f))
 			]
 		]
@@ -859,7 +861,7 @@ TSharedRef<SWidget> SDMMaterialSlotEditor::CreateSlot_LayerSettings()
 			[
 				SNew(SImage)
 				.Image(FAppStyle::Get().GetBrush("Icons.Plus"))
-				.ColorAndOpacity(FDynamicMaterialEditorStyle::Get().GetColor("Color.Stage.Enabled"))
+				.ColorAndOpacity(FSlateColor::UseForeground())
 				.DesiredSizeOverride(FVector2D(16.0f))
 			]
 		]
@@ -884,6 +886,7 @@ TSharedRef<SWidget> SDMMaterialSlotEditor::CreateSlot_LayerSettings()
 				[
 					SNew(SImage)
 					.Image(FAppStyle::Get().GetBrush("Icons.Delete"))
+					.ColorAndOpacity(FSlateColor::UseForeground())
 					.DesiredSizeOverride(FVector2D(16.0f))
 				]
 			]
