Material Designer: Fixed texture edge color scrubbing.
- Changing predefined location while scrubbing now longer causes a details panel refresh.

#FailedValidations Preflight

diff --git a/Engine/Plugins/VirtualProduction/DynamicMaterial/Source/DynamicMaterialEditor/Private/Components/MaterialStageExpressions/DMMSETextureSampleEdgeColor.cpp b/Engine/Plugins/VirtualProduction/DynamicMaterial/Source/DynamicMaterialEditor/Private/Components/MaterialStageExpressions/DMMSETextureSampleEdgeColor.cpp
--- a/Engine/Plugins/VirtualProduction/DynamicMaterial/Source/DynamicMaterialEditor/Private/Components/MaterialStageExpressions/DMMSETextureSampleEdgeColor.cpp
+++ b/Engine/Plugins/VirtualProduction/DynamicMaterial/Source/DynamicMaterialEditor/Private/Components/MaterialStageExpressions/DMMSETextureSampleEdgeColor.cpp
@@ -109,8 +109,7 @@ void UDMMaterialStageExpressionTextureSampleEdgeColor::Update(UDMMaterialCompone
 	if (!bIsValidX || !bIsValidY)
 	{
 		EdgeLocation = EDMEdgeLocation::Custom;
-
-		Super::Update(InSource, InUpdateType | EDMUpdateType::Value | EDMUpdateType::RefreshDetailView);
+		Super::Update(InSource, InUpdateType | EDMUpdateType::Value);
 		return;
 	}
 
@@ -191,15 +190,7 @@ void UDMMaterialStageExpressionTextureSampleEdgeColor::Update(UDMMaterialCompone
 
 	if (bChangedEdgeLocation)
 	{
-		if (PreEditEdgeLocation == EDMEdgeLocation::Custom)
-		{
-			Super::Update(InSource, InUpdateType | EDMUpdateType::Value | EDMUpdateType::RefreshDetailView);
-		}
-		else
-		{
-			Super::Update(InSource, InUpdateType | EDMUpdateType::Value);
-		}
-
+		Super::Update(InSource, InUpdateType | EDMUpdateType::Value);
 		return;
 	}
 
