Material Designer: UV VIsualizer fixes
- Control-clicking when not over the gizmo will reset all "active" settings.
- Fixed the Pivot gizmo.

diff --git a/Engine/Plugins/VirtualProduction/DynamicMaterial/Source/DynamicMaterialEditor/Private/UI/Widgets/Visualizers/SDMTextureUVVisualizer.cpp b/Engine/Plugins/VirtualProduction/DynamicMaterial/Source/DynamicMaterialEditor/Private/UI/Widgets/Visualizers/SDMTextureUVVisualizer.cpp
--- a/Engine/Plugins/VirtualProduction/DynamicMaterial/Source/DynamicMaterialEditor/Private/UI/Widgets/Visualizers/SDMTextureUVVisualizer.cpp
+++ b/Engine/Plugins/VirtualProduction/DynamicMaterial/Source/DynamicMaterialEditor/Private/UI/Widgets/Visualizers/SDMTextureUVVisualizer.cpp
@@ -39,6 +39,9 @@ namespace UE::DynamicMaterialEditor::Private
 	/** Based distance the circle handle is from the center compared to the width of the image. */
 	constexpr float TextureUVVisualizerCircleHandleBaseRadiusMultiplier = 0.4f;
 
+	/** Scale for the clamped version of the base cirle base radius. */
+	constexpr float TextureUVVisualizerCircleHandleBaseRadiusClampedMultiplier = 4.f;
+
 	/** Size of the border around the center square */
 	constexpr float TextureUVVisualizerBorder = 8.f;
 }
@@ -556,7 +559,7 @@ float SDMTextureUVVisualizer::GetCircleHandleBaseRadius() const
 	const FVector2f ImageSize = StagePreview->GetBrush().GetImageSize();
 	const float CircleHandleRadius = FMath::Min(ImageSize.X, ImageSize.Y) * TextureUVVisualizerCircleHandleBaseRadiusMultiplier;
 
-	return FMath::Clamp(CircleHandleRadius, 10, 50.f);
+	return FMath::Clamp(CircleHandleRadius, 10, 50.f) * TextureUVVisualizerCircleHandleBaseRadiusClampedMultiplier;
 }
 
 FVector2f SDMTextureUVVisualizer::ApplyTextureUVTransform(const FVector2f& InUV) const
@@ -681,6 +684,13 @@ SDMTextureUVVisualizer::EHandleAxis SDMTextureUVVisualizer::GetCenterHandleAxis(
 	HandleOffset.X = FMath::Abs(HandleOffset.X);
 	HandleOffset.Y = FMath::Abs(HandleOffset.Y);
 
+	const float UIScale = GetTickSpaceGeometry().Scale;
+
+	if (UIScale > 0)
+	{
+		HandleOffset /= UIScale;
+	}
+
 	static constexpr float LargeMinusSmallRadius = TextureUVVisualizerLargeRadius - TextureUVVisualizerSmallRadius;
 
 	if (HandleOffset.X <= LargeMinusSmallRadius && HandleOffset.Y <= LargeMinusSmallRadius)
@@ -715,7 +725,15 @@ SDMTextureUVVisualizer::EHandleAxis SDMTextureUVVisualizer::GetCircleHandleAxis(
 
 	using namespace UE::DynamicMaterialEditor::Private;
 
-	const FVector2f HandleOffset = InAbsolutePosition - GetAbsolutePivotLocation();
+	float UIScale = GetTickSpaceGeometry().Scale;
+
+	if (FMath::IsNearlyZero(UIScale))
+	{
+		UIScale = 1.f;
+	}
+
+	const FVector2f HandleOffset = (InAbsolutePosition - GetAbsolutePivotLocation()) / UIScale;
+
 	const float DistanceFromHandle = HandleOffset.Size();
 	const FVector2D& Tiling2D = GetTiling();
 	// We manage angle clockwise from +Y axis. Atan2 handles it anti-clockwise from +X axis.
@@ -724,11 +742,11 @@ SDMTextureUVVisualizer::EHandleAxis SDMTextureUVVisualizer::GetCircleHandleAxis(
 
 	if (FMath::IsNearlyEqual(Tiling2D.X, Tiling2D.Y))
 	{
-		DistanceFromCircleHandle = FMath::Abs(DistanceFromHandle - GetCircleHandleBaseRadius());
+		DistanceFromCircleHandle = FMath::Abs(DistanceFromHandle - GetCircleHandleBaseRadius()) / UIScale;
 	}
 	else
 	{
-		DistanceFromCircleHandle = FMath::Abs(DistanceFromHandle - GetCircleHandleRadiusAtAngle(Angle));
+		DistanceFromCircleHandle = FMath::Abs(DistanceFromHandle - GetCircleHandleRadiusAtAngle(Angle)) / UIScale;
 	}
 
 	if (DistanceFromCircleHandle > TextureUVVisualizerCircleHandleRadius)
@@ -800,7 +818,23 @@ bool SDMTextureUVVisualizer::TryClickCenterHandle(const FVector2f& InMousePositi
 
 	if (Axis == EHandleAxis::None)
 	{
-		return false;
+		if (!bInResetToDefault)
+		{
+			return false;
+		}
+
+		if (!bPivotEditMode)
+		{
+			FScopedTransaction Transaction(LOCTEXT("ResetOffset", "Reset Offset to Default."));
+			SetOffset(FVector2D::ZeroVector);
+		}
+		else
+		{
+			FScopedTransaction Transaction(LOCTEXT("ResetPivot", "Reset Pivot to Default."));
+			SetPivot(FVector2D::ZeroVector);
+		}
+
+		return true;
 	}
 
 	// Offset handle
@@ -884,7 +918,19 @@ bool SDMTextureUVVisualizer::TryClickCircleHandle(const FVector2f& InMousePositi
 	switch (Axis)
 	{
 		default:
-			return false;
+			if (!bInResetToDefault)
+			{
+				return false;
+			}
+
+			{
+				FScopedTransaction Transaction(LOCTEXT("ResetRotationAndTiling", "Reset Rotation and Tiling to Default."));
+				ModifyTextureUVComponent();
+				SetRotation(0);
+				SetTiling(FVector2D::UnitVector);
+			}
+
+			return true;
 
 		// XY on the circle handle represents rotation
 		case EHandleAxis::XY:
