Material Designer: Fixed layer opacity disappearing in languages other than English

diff --git a/Engine/Plugins/VirtualProduction/DynamicMaterial/Source/DynamicMaterialEditor/Private/Components/DMMaterialStageBlend.cpp b/Engine/Plugins/VirtualProduction/DynamicMaterial/Source/DynamicMaterialEditor/Private/Components/DMMaterialStageBlend.cpp
--- a/Engine/Plugins/VirtualProduction/DynamicMaterial/Source/DynamicMaterialEditor/Private/Components/DMMaterialStageBlend.cpp
+++ b/Engine/Plugins/VirtualProduction/DynamicMaterial/Source/DynamicMaterialEditor/Private/Components/DMMaterialStageBlend.cpp
@@ -23,7 +23,7 @@
 #include UE_INLINE_GENERATED_CPP_BY_NAME(DMMaterialStageBlend)
  
 #define LOCTEXT_NAMESPACE "DMMaterialProperty"
- 
+
 TArray<TStrongObjectPtr<UClass>> UDMMaterialStageBlend::Blends = {};
  
 UDMMaterialStageBlend::UDMMaterialStageBlend()
@@ -97,6 +97,31 @@ void UDMMaterialStageBlend::BlendOpacityLayer(const TSharedRef<FDMMaterialBuildS
 		InMyLayerOpacityExpression, InMyOutputIndex, InMyOutputChannel, OutAddedExpressions, OutOutputIndex, OutOutputChannel);
 }
 
+UDMMaterialStageInputValue* UDMMaterialStageBlend::GetOpacityValue(UDMMaterialStage* InStage) const
+{
+	if (!InStage)
+	{
+		return nullptr;
+	}
+
+	const TArray<UDMMaterialStageInput*>& StageInputs = InStage->GetInputs();
+	const TArray<FDMMaterialStageConnection>& InputConnectionMap = InStage->GetInputConnectionMap();
+
+	constexpr int32 OpacityIndex = 0;
+
+	if (InputConnectionMap.IsValidIndex(OpacityIndex) && !InputConnectionMap[OpacityIndex].Channels.IsEmpty())
+	{
+		if (InputConnectionMap[OpacityIndex].Channels[0].SourceIndex >= FDMMaterialStageConnectorChannel::FIRST_STAGE_INPUT)
+		{
+			const int32 StageInputIndex = InputConnectionMap[OpacityIndex].Channels[0].SourceIndex - FDMMaterialStageConnectorChannel::FIRST_STAGE_INPUT;
+
+			return Cast<UDMMaterialStageInputValue>(StageInputs[StageInputIndex]);
+		}
+	}
+
+	return nullptr;
+}
+
 int32 UDMMaterialStageBlend::GetDefaultBaseChannelOverrideOutputIndex() const
 {
 	UDMMaterialStageInput* StageInputB = GetInputB();
diff --git a/Engine/Plugins/VirtualProduction/DynamicMaterial/Source/DynamicMaterialEditor/Private/Utils/DMMaterialStageFunctionLibrary.cpp b/Engine/Plugins/VirtualProduction/DynamicMaterial/Source/DynamicMaterialEditor/Private/Utils/DMMaterialStageFunctionLibrary.cpp
--- a/Engine/Plugins/VirtualProduction/DynamicMaterial/Source/DynamicMaterialEditor/Private/Utils/DMMaterialStageFunctionLibrary.cpp
+++ b/Engine/Plugins/VirtualProduction/DynamicMaterial/Source/DynamicMaterialEditor/Private/Utils/DMMaterialStageFunctionLibrary.cpp
@@ -3,6 +3,7 @@
 #include "Utils/DMMaterialStageFunctionLibrary.h"
 
 #include "Components/DMMaterialStage.h"
+#include "Components/DMMaterialStageBlend.h"
 #include "Components/DMMaterialStageThroughput.h"
 #include "Components/DMMaterialSubStage.h"
 #include "Components/DMRenderTargetRenderer.h"
@@ -16,39 +17,19 @@
 
 UDMMaterialStageInputValue* UDMMaterialStageFunctionLibrary::FindDefaultStageOpacityInputValue(UDMMaterialStage* InStage)
 {
-	if (!IsValid(InStage))
+	if (!InStage)
 	{
 		return nullptr;
 	}
 
-	int32 OpacityInputIndex = -1;
+	UDMMaterialStageBlend* Blend = Cast<UDMMaterialStageBlend>(InStage->GetSource());
 
-	if (UDMMaterialStageThroughput* const ThroughputSource = Cast<UDMMaterialStageThroughput>(InStage->GetSource()))
+	if (!Blend)
 	{
-		const TArray<FDMMaterialStageConnector>& Connectors = ThroughputSource->GetInputConnectors();
-		for (const FDMMaterialStageConnector& Connector : Connectors)
-		{
-			if (Connector.Name.ToString() == TEXT("Opacity"))
-			{
-				OpacityInputIndex = Connector.Index;
-			}
-		}
-	}
-
-	const TArray<UDMMaterialStageInput*>& StageInputs = InStage->GetInputs();
-	const TArray<FDMMaterialStageConnection>& InputConnectionMap = InStage->GetInputConnectionMap();
-
-	if (InputConnectionMap.IsValidIndex(OpacityInputIndex) && !InputConnectionMap[OpacityInputIndex].Channels.IsEmpty())
-	{
-		if (InputConnectionMap[OpacityInputIndex].Channels[0].SourceIndex >= FDMMaterialStageConnectorChannel::FIRST_STAGE_INPUT)
-		{
-			const int32 StageInputIndex = InputConnectionMap[OpacityInputIndex].Channels[0].SourceIndex - FDMMaterialStageConnectorChannel::FIRST_STAGE_INPUT;
-
-			return Cast<UDMMaterialStageInputValue>(StageInputs[StageInputIndex]);
-		}
+		return nullptr;
 	}
 
-	return nullptr;
+	return Blend->GetOpacityValue(InStage);
 }
 
 void UDMMaterialStageFunctionLibrary::SetStageInputToRenderer(UDMMaterialStage* InStage, TSubclassOf<UDMRenderTargetRenderer> InRendererClass, int32 InInputIndex)
diff --git a/Engine/Plugins/VirtualProduction/DynamicMaterial/Source/DynamicMaterialEditor/Public/Components/DMMaterialStageBlend.h b/Engine/Plugins/VirtualProduction/DynamicMaterial/Source/DynamicMaterialEditor/Public/Components/DMMaterialStageBlend.h
--- a/Engine/Plugins/VirtualProduction/DynamicMaterial/Source/DynamicMaterialEditor/Public/Components/DMMaterialStageBlend.h
+++ b/Engine/Plugins/VirtualProduction/DynamicMaterial/Source/DynamicMaterialEditor/Public/Components/DMMaterialStageBlend.h
@@ -11,7 +11,9 @@
 
 class FMenuBuilder;
 class UDMMaterialLayerObject;
+class UDMMaterialStage;
 class UDMMaterialStageInput;
+class UDMMaterialStageInputValue;
 class UDMMaterialValueFloat1;
 class UMaterial;
 enum class EAvaColorChannel : uint8;
@@ -62,6 +64,8 @@ public:
 		UMaterialExpression* InMyLayerOpacityExpression, int32 InMyOutputIndex, int32 InMyOutputChannel,
 		TArray<UMaterialExpression*>& OutAddedExpressions, int32& OutOutputIndex, int32& OutOutputChannel) const;
 
+	UDMMaterialStageInputValue* GetOpacityValue(UDMMaterialStage* InStage) const;
+
 	//~ Begin UDMMaterialStageThroughput
 	DYNAMICMATERIALEDITOR_API virtual bool CanInputAcceptType(int32 InputIndex, EDMValueType ValueType) const override;
 	DYNAMICMATERIALEDITOR_API virtual void AddDefaultInput(int32 InInputIndex) const override;
