Material Designer: The element selection dropdown now correctly identifies which material is open on multi-material meshes,

#FailedValidations Preflight

diff --git a/Engine/Plugins/VirtualProduction/DynamicMaterial/Source/DynamicMaterialEditor/Private/UI/Widgets/Editor/SDMToolBar.cpp b/Engine/Plugins/VirtualProduction/DynamicMaterial/Source/DynamicMaterialEditor/Private/UI/Widgets/Editor/SDMToolBar.cpp
--- a/Engine/Plugins/VirtualProduction/DynamicMaterial/Source/DynamicMaterialEditor/Private/UI/Widgets/Editor/SDMToolBar.cpp
+++ b/Engine/Plugins/VirtualProduction/DynamicMaterial/Source/DynamicMaterialEditor/Private/UI/Widgets/Editor/SDMToolBar.cpp
@@ -46,12 +46,24 @@ namespace UE::DynamicMaterialEditor::Private
 	static const FVector2D LargeIconToolBarButtonSize = FVector2D(16.f);
 }
 
-void SDMToolBar::Construct(const FArguments& InArgs, const TSharedRef<SDMMaterialEditor>& InEditorWidget, AActor* InActor)
+void SDMToolBar::Construct(const FArguments& InArgs, const TSharedRef<SDMMaterialEditor>& InEditorWidget, const FDMObjectMaterialProperty& InObjectProperty)
 {
 	EditorWidgetWeak = InEditorWidget;
-	MaterialActorWeak = InActor;
-	SelectedMaterialElementIndex = INDEX_NONE;
 
+	if (UObject* Outer = InObjectProperty.GetOuter())
+	{
+		if (AActor* OuterActor = Cast<AActor>(Outer))
+		{
+			MaterialActorWeak = OuterActor;
+		}
+		else if (AActor* OuterActorOuter = Outer->GetTypedOuter<AActor>())
+		{
+			MaterialActorWeak = OuterActorOuter;
+		}
+	}
+
+	SelectedMaterialElementIndex = InObjectProperty.GetIndex();
+	
 	SetCanTick(false);
 	
 	ChildSlot
@@ -69,7 +81,7 @@ void SDMToolBar::Construct(const FArguments& InArgs, const TSharedRef<SDMMateria
 			]
 		];
 
-	SetActorPropertySelected(InActor);
+	SetActorPropertySelected(InObjectProperty);
 	UpdateButtonVisibilities();
 }
 
@@ -375,14 +387,32 @@ TSharedRef<SWidget> SDMToolBar::CreateToolBarEntries()
 		];
 }
 
-void SDMToolBar::SetActorPropertySelected(AActor* InActor)
+void SDMToolBar::SetActorPropertySelected(const FDMObjectMaterialProperty& InObjectProperty)
 {
-	if (IsValid(InActor))
+	AActor* Actor = nullptr;
+
+	if (UObject* Outer = InObjectProperty.GetOuter())
 	{
+		if (AActor* OuterActor = Cast<AActor>(Outer))
+		{
+			Actor = OuterActor;
+		}
+		else if (AActor* OuterActorOuter = Outer->GetTypedOuter<AActor>())
+		{
+			Actor = OuterActorOuter;
+		}
+	}
+	
+	// We can only deal with non-property materials.
+	if (Actor && InObjectProperty.IsValid() && InObjectProperty.IsElement())
+	{
+		MaterialActorWeak = Actor;
+		SelectedMaterialElementIndex = InObjectProperty.GetIndex();
+
 		ActorNameWidget->SetText(GetActorName());
 		ActorRowWidget->SetVisibility(EVisibility::Visible);
 
-		TArray<FDMObjectMaterialProperty> ActorProperties = UDMMaterialInstanceFunctionLibrary::GetActorMaterialProperties(InActor);
+		TArray<FDMObjectMaterialProperty> ActorProperties = UDMMaterialInstanceFunctionLibrary::GetActorMaterialProperties(Actor);
 		UDynamicMaterialModelBase* MaterialModelBase = GetOriginalMaterialModelBase();
 
 		const int32 ActorPropertyCount = ActorProperties.Num();
@@ -391,17 +421,12 @@ void SDMToolBar::SetActorPropertySelected(AActor* InActor)
 		for (int32 MaterialPropertyIdx = 0; MaterialPropertyIdx < ActorPropertyCount; ++MaterialPropertyIdx)
 		{
 			const FDMObjectMaterialProperty& MaterialProperty = ActorProperties[MaterialPropertyIdx];
-
 			ActorMaterialProperties.Add(MakeShared<FDMObjectMaterialProperty>(MaterialProperty));
-
-			if (MaterialProperty.GetMaterialModelBase() == MaterialModelBase)
-			{
-				SelectedMaterialElementIndex = MaterialPropertyIdx;
-			}
 		}
 	}
 	else
 	{
+		MaterialActorWeak.Reset();
 		ActorMaterialProperties.Empty(0);
 		ActorNameWidget->SetText(FText::GetEmpty());
 		ActorRowWidget->SetVisibility(EVisibility::Collapsed);
diff --git a/Engine/Plugins/VirtualProduction/DynamicMaterial/Source/DynamicMaterialEditor/Private/UI/Widgets/Editor/SDMToolBar.h b/Engine/Plugins/VirtualProduction/DynamicMaterial/Source/DynamicMaterialEditor/Private/UI/Widgets/Editor/SDMToolBar.h
--- a/Engine/Plugins/VirtualProduction/DynamicMaterial/Source/DynamicMaterialEditor/Private/UI/Widgets/Editor/SDMToolBar.h
+++ b/Engine/Plugins/VirtualProduction/DynamicMaterial/Source/DynamicMaterialEditor/Private/UI/Widgets/Editor/SDMToolBar.h
@@ -38,7 +38,7 @@ public:
 
 	virtual ~SDMToolBar() override = default;
 
-	void Construct(const FArguments& InArgs, const TSharedRef<SDMMaterialEditor>& InEditorWidget, AActor* InActor);
+	void Construct(const FArguments& InArgs, const TSharedRef<SDMMaterialEditor>& InEditorWidget, const FDMObjectMaterialProperty& InObjectProperty);
 
 	AActor* GetMaterialActor() const { return MaterialActorWeak.Get(); }
 
@@ -72,7 +72,7 @@ protected:
 
 	TSharedRef<SWidget> CreateToolBarEntries();
 
-	void SetActorPropertySelected(AActor* InActor);
+	void SetActorPropertySelected(const FDMObjectMaterialProperty& InObjectProperty);
 
 	void UpdateButtonVisibilities();
 
diff --git a/Engine/Plugins/VirtualProduction/DynamicMaterial/Source/DynamicMaterialEditor/Private/UI/Widgets/SDMMaterialEditor.cpp b/Engine/Plugins/VirtualProduction/DynamicMaterial/Source/DynamicMaterialEditor/Private/UI/Widgets/SDMMaterialEditor.cpp
--- a/Engine/Plugins/VirtualProduction/DynamicMaterial/Source/DynamicMaterialEditor/Private/UI/Widgets/SDMMaterialEditor.cpp
+++ b/Engine/Plugins/VirtualProduction/DynamicMaterial/Source/DynamicMaterialEditor/Private/UI/Widgets/SDMMaterialEditor.cpp
@@ -339,7 +339,7 @@ void SDMMaterialEditor::SetMaterialActor(AActor* InActor)
 		return;
 	}
 
-	TSharedRef<SDMToolBar> NewToolBar = SNew(SDMToolBar, SharedThis(this), InActor);
+	TSharedRef<SDMToolBar> NewToolBar = SNew(SDMToolBar, SharedThis(this), FDMObjectMaterialProperty(InActor, 0));
 
 	ToolBarSlot << NewToolBar;
 }
@@ -1276,8 +1276,8 @@ TSharedRef<SDMToolBar> SDMMaterialEditor::CreateSlot_ToolBar()
 		SDMToolBar, 
 		SharedThis(this), 
 		ObjectMaterialPropertyOpt.IsSet()
-			? ObjectMaterialPropertyOpt->GetTypedOuter<AActor>()
-			: nullptr
+			? ObjectMaterialPropertyOpt.GetValue()
+			: FDMObjectMaterialProperty(GetMaterialActor(), 0)
 	);
 }
 
diff --git a/Engine/Plugins/VirtualProduction/DynamicMaterial/Source/DynamicMaterialEditor/Private/UI/Widgets/SDMMaterialWizard.cpp b/Engine/Plugins/VirtualProduction/DynamicMaterial/Source/DynamicMaterialEditor/Private/UI/Widgets/SDMMaterialWizard.cpp
--- a/Engine/Plugins/VirtualProduction/DynamicMaterial/Source/DynamicMaterialEditor/Private/UI/Widgets/SDMMaterialWizard.cpp
+++ b/Engine/Plugins/VirtualProduction/DynamicMaterial/Source/DynamicMaterialEditor/Private/UI/Widgets/SDMMaterialWizard.cpp
@@ -989,7 +989,14 @@ void SDMMaterialWizard::SetChannelListInModel(FName InChannelList, UDynamicMater
 	EditorOnlyData->SetChannelListPreset(InChannelList);
 	EditorOnlyData->OnWizardComplete();
 
-	DesignerWidget->OpenMaterialModelBase(MaterialModel);
+	if (MaterialObjectProperty.IsSet())
+	{
+		DesignerWidget->OpenObjectMaterialProperty(*MaterialObjectProperty);
+	}
+	else
+	{
+		DesignerWidget->OpenMaterialModelBase(MaterialModel);
+	}
 }
 
 void SDMMaterialWizard::CreateTemplateMaterialInActor(FName InChannelList, FDMObjectMaterialProperty& InMaterialObjectProperty)
