Material Designer: The element dropdown will now identify duplicate material slots (same material in both slots) and prevent the user from selecting the duplicated material.

diff --git a/Engine/Plugins/VirtualProduction/DynamicMaterial/Source/DynamicMaterialEditor/Private/UI/Widgets/Editor/SDMToolBar.cpp b/Engine/Plugins/VirtualProduction/DynamicMaterial/Source/DynamicMaterialEditor/Private/UI/Widgets/Editor/SDMToolBar.cpp
--- a/Engine/Plugins/VirtualProduction/DynamicMaterial/Source/DynamicMaterialEditor/Private/UI/Widgets/Editor/SDMToolBar.cpp
+++ b/Engine/Plugins/VirtualProduction/DynamicMaterial/Source/DynamicMaterialEditor/Private/UI/Widgets/Editor/SDMToolBar.cpp
@@ -3,6 +3,7 @@
 #include "UI/Widgets/Editor/SDMToolBar.h"
 
 #include "AssetRegistry/AssetData.h"
+#include "Components/PrimitiveComponent.h"
 #include "DMWorldSubsystem.h"
 #include "DynamicMaterialEditorModule.h"
 #include "DynamicMaterialEditorSettings.h"
@@ -550,19 +551,91 @@ TSharedRef<SWidget> SDMToolBar::CreateSlotsComboBoxWidget()
 
 TSharedRef<SWidget> SDMToolBar::GenerateSelectedMaterialSlotRow(TSharedPtr<FDMObjectMaterialProperty> InSelectedSlot) const
 {
-	if (InSelectedSlot.IsValid())
+	if (!InSelectedSlot.IsValid() || InSelectedSlot->IsProperty())
 	{
-		return SNew(STextBlock)
-			.MinDesiredWidth(100.f)
-			.TextStyle(FDynamicMaterialEditorStyle::Get(), "RegularFont")
-			.Text(this, &SDMToolBar::GetSlotDisplayName, InSelectedSlot);
+		return SNullWidget::NullWidget;
 	}
-	return SNullWidget::NullWidget;
+
+	return SNew(STextBlock)
+		.MinDesiredWidth(100.f)
+		.TextStyle(FDynamicMaterialEditorStyle::Get(), "RegularFont")
+		.Text(this, &SDMToolBar::GetSlotDisplayName, InSelectedSlot.ToWeakPtr())
+		.IsEnabled(this, &SDMToolBar::GetSlotEnabled, InSelectedSlot.ToWeakPtr());
 }
 
-FText SDMToolBar::GetSlotDisplayName(TSharedPtr<FDMObjectMaterialProperty> InSlot) const
+int32 SDMToolBar::GetDuplicateIndex(const FDMObjectMaterialProperty& InObjectProperty) const
 {
-	return InSlot->GetPropertyName(false);
+	UPrimitiveComponent* Outer = Cast<UPrimitiveComponent>(InObjectProperty.GetOuter());
+
+	if (!Outer)
+	{
+		return INDEX_NONE;
+	}
+
+	UDynamicMaterialModelBase* MyModel = InObjectProperty.GetMaterialModelBase();
+
+	if (!MyModel)
+	{
+		return INDEX_NONE;
+	}
+
+	const int32 CurrentIndex = InObjectProperty.GetIndex();
+
+	int32 SameMaterialAsSlot = -1;
+
+	for (int32 Index = 0; Index < CurrentIndex; ++Index)
+	{
+		UDynamicMaterialModelBase* Model = FDMObjectMaterialProperty(Outer, Index).GetMaterialModelBase();
+
+		if (Model && Model == MyModel)
+		{
+			return Index;
+		}
+	}
+
+	return INDEX_NONE;
+}
+
+FText SDMToolBar::GetSlotDisplayName(TWeakPtr<FDMObjectMaterialProperty> InSlotWeak) const
+{
+	TSharedPtr<FDMObjectMaterialProperty> Slot = InSlotWeak.Pin();
+
+	if (!Slot.IsValid())
+	{
+		return FText::GetEmpty();
+	}
+
+	UPrimitiveComponent* Outer = Cast<UPrimitiveComponent>(Slot->GetOuter());
+
+	if (!Outer)
+	{
+		return FText::GetEmpty();
+	}
+
+	const int32 DuplicateIndex = GetDuplicateIndex(*Slot);
+
+	if (DuplicateIndex != INDEX_NONE)
+	{
+		return FText::Format(
+			LOCTEXT("ElementXCopy", "Element {0} [Duplicate of {1}]"),
+			FText::AsNumber(Slot->GetIndex()),
+			FText::AsNumber(DuplicateIndex)
+		);
+	}
+
+	return Slot->GetPropertyName(/* Ignore "create new" status */ false);
+}
+
+bool SDMToolBar::GetSlotEnabled(TWeakPtr<FDMObjectMaterialProperty> InSlotWeak) const
+{
+	TSharedPtr<FDMObjectMaterialProperty> Slot = InSlotWeak.Pin();
+
+	if (!Slot.IsValid())
+	{
+		return false;
+	}
+
+	return GetDuplicateIndex(*Slot) == INDEX_NONE;
 }
 
 FText SDMToolBar::GetSelectedMaterialSlotName() const
@@ -595,8 +668,18 @@ void SDMToolBar::OnMaterialSlotChanged(TSharedPtr<FDMObjectMaterialProperty> InS
 		return;
 	}
 
+	if (GetDuplicateIndex(*InSelectedSlot) != INDEX_NONE)
+	{
+		return;
+	}
+
 	UDynamicMaterialModelBase* SelectedMaterialModelBase = InSelectedSlot->GetMaterialModelBase();
 
+	if (EditorWidget->GetOriginalMaterialModelBase() == SelectedMaterialModelBase)
+	{
+		return;
+	}
+
 	if (IsValid(SelectedMaterialModelBase))
 	{
 		DesignerWidget->OpenObjectMaterialProperty(*InSelectedSlot);
diff --git a/Engine/Plugins/VirtualProduction/DynamicMaterial/Source/DynamicMaterialEditor/Private/UI/Widgets/Editor/SDMToolBar.h b/Engine/Plugins/VirtualProduction/DynamicMaterial/Source/DynamicMaterialEditor/Private/UI/Widgets/Editor/SDMToolBar.h
--- a/Engine/Plugins/VirtualProduction/DynamicMaterial/Source/DynamicMaterialEditor/Private/UI/Widgets/Editor/SDMToolBar.h
+++ b/Engine/Plugins/VirtualProduction/DynamicMaterial/Source/DynamicMaterialEditor/Private/UI/Widgets/Editor/SDMToolBar.h
@@ -74,6 +74,8 @@ protected:
 
 	void SetActorPropertySelected(const FDMObjectMaterialProperty& InObjectProperty);
 
+	int32 GetDuplicateIndex(const FDMObjectMaterialProperty& InObjectProperty) const;
+
 	void UpdateButtonVisibilities();
 
 	TSharedRef<SWidget> CreateToolBarButton(TAttribute<const FSlateBrush*> InImageBrush, const TAttribute<FText>& InTooltipText, FOnClicked InOnClicked);
@@ -81,7 +83,8 @@ protected:
 	TSharedRef<SWidget> CreateSlotsComboBoxWidget();
 
 	TSharedRef<SWidget> GenerateSelectedMaterialSlotRow(TSharedPtr<FDMObjectMaterialProperty> InSelectedSlot) const;
-	FText GetSlotDisplayName(TSharedPtr<FDMObjectMaterialProperty> InSlot) const;
+	FText GetSlotDisplayName(TWeakPtr<FDMObjectMaterialProperty> InSlotWeak) const;
+	bool GetSlotEnabled(TWeakPtr<FDMObjectMaterialProperty> InSlotWeak) const;
 	FText GetSelectedMaterialSlotName() const;
 	void OnMaterialSlotChanged(TSharedPtr<FDMObjectMaterialProperty> InSelectedSlot, ESelectInfo::Type InSelectInfoType);
 
