Material Designer: Reset UI Settings now refreshes the UI.

diff --git a/Engine/Plugins/VirtualProduction/DynamicMaterial/Source/DynamicMaterialEditor/Private/DynamicMaterialEditorSettings.cpp b/Engine/Plugins/VirtualProduction/DynamicMaterial/Source/DynamicMaterialEditor/Private/DynamicMaterialEditorSettings.cpp
--- a/Engine/Plugins/VirtualProduction/DynamicMaterial/Source/DynamicMaterialEditor/Private/DynamicMaterialEditorSettings.cpp
+++ b/Engine/Plugins/VirtualProduction/DynamicMaterial/Source/DynamicMaterialEditor/Private/DynamicMaterialEditorSettings.cpp
@@ -324,6 +324,11 @@ void UDynamicMaterialEditorSettings::ResetAllLayoutSettings()
 	bUVVisualizerVisible = true;
 	StagePreviewSize = 40.f;
 	PropertyPreviewSize = 64.f;
+
+	FPropertyChangedEvent PropertyChangeEvent(UDynamicMaterialEditorSettings::StaticClass()->FindPropertyByName(
+		GET_MEMBER_NAME_CHECKED(UDynamicMaterialEditorSettings, Layout)
+	));
+	OnSettingsChanged.Broadcast(PropertyChangeEvent);
 }
 
 TArray<FDMMaterialEffectList> UDynamicMaterialEditorSettings::GetEffectList() const
