Material Designer: Refactored the 'channel bitmask' function to reduce node count and give better output.

diff --git a/Engine/Plugins/VirtualProduction/DynamicMaterial/Source/DynamicMaterial/Public/Model/IDMMaterialBuildUtilsInterface.h b/Engine/Plugins/VirtualProduction/DynamicMaterial/Source/DynamicMaterial/Public/Model/IDMMaterialBuildUtilsInterface.h
--- a/Engine/Plugins/VirtualProduction/DynamicMaterial/Source/DynamicMaterial/Public/Model/IDMMaterialBuildUtilsInterface.h
+++ b/Engine/Plugins/VirtualProduction/DynamicMaterial/Source/DynamicMaterial/Public/Model/IDMMaterialBuildUtilsInterface.h
@@ -75,6 +75,9 @@ struct IDMMaterialBuildUtilsInterface
 	/** Creates a set of expressions that display this material stage input. */
 	virtual TArray<UMaterialExpression*> CreateExpressionInput(UDMMaterialStageInput* InInput) const = 0;
 
+	/** Searches the outputs of an expression to see if there is an appropriate output to match a requested channel mask. Returns INDEX_NONE on failure. */
+	virtual int32 FindOutputForBitmask(UMaterialExpression* InExpression, int32 InOutputChannels) const = 0;
+
 	/** Creates a set of expressions merging all the inputs for each channel into a single output */
 	virtual UMaterialExpressionComponentMask* CreateExpressionBitMask(UMaterialExpression* InExpression, int32 InOutputIndex, 
 		int32 InOutputChannels) const = 0;
diff --git a/Engine/Plugins/VirtualProduction/DynamicMaterial/Source/DynamicMaterialEditor/Private/Components/DMMaterialEffectFunction.cpp b/Engine/Plugins/VirtualProduction/DynamicMaterial/Source/DynamicMaterialEditor/Private/Components/DMMaterialEffectFunction.cpp
--- a/Engine/Plugins/VirtualProduction/DynamicMaterial/Source/DynamicMaterialEditor/Private/Components/DMMaterialEffectFunction.cpp
+++ b/Engine/Plugins/VirtualProduction/DynamicMaterial/Source/DynamicMaterialEditor/Private/Components/DMMaterialEffectFunction.cpp
@@ -302,21 +302,36 @@ void UDMMaterialEffectFunction::ApplyTo(const TSharedRef<FDMMaterialBuildState>&
 			{
 				if (InOutLastExpressionOutputChannel != FDMMaterialStageConnectorChannel::WHOLE_CHANNEL)
 				{
-					UMaterialExpression* Mask = InBuildState->GetBuildUtils().CreateExpressionBitMask(
-						LastInputExpressions[0], 
-						InOutLastExpressionOutputIndex,
+					const int32 MaskedOutput = InBuildState->GetBuildUtils().FindOutputForBitmask(
+						LastInputExpressions[0],
 						InOutLastExpressionOutputChannel
 					);
 
-					FunctionCall->FunctionInputs[InputIndex].Input.Connect(0, Mask);
+					if (MaskedOutput == INDEX_NONE)
+					{
+						UMaterialExpression* Mask = InBuildState->GetBuildUtils().CreateExpressionBitMask(
+							LastInputExpressions[0],
+							InOutLastExpressionOutputIndex,
+							InOutLastExpressionOutputChannel
+						);
+
+						FunctionCall->FunctionInputs[InputIndex].Input.Connect(0, Mask);
 
-					InOutStageExpressions.Add(Mask);
+						InOutStageExpressions.Add(Mask);
+					}
+					else
+					{
+						FunctionCall->FunctionInputs[InputIndex].Input.Connect(
+							MaskedOutput,
+							LastInputExpressions[0]
+						);
+					}
 				}
 				else
 				{
 					FunctionCall->FunctionInputs[InputIndex].Input.Connect(
 						InOutLastExpressionOutputIndex,
-						LastInputExpressions[InputIndex]
+						LastInputExpressions[0]
 					);
 				}
 			}
diff --git a/Engine/Plugins/VirtualProduction/DynamicMaterial/Source/DynamicMaterialEditor/Private/Components/DMMaterialStageBlend.cpp b/Engine/Plugins/VirtualProduction/DynamicMaterial/Source/DynamicMaterialEditor/Private/Components/DMMaterialStageBlend.cpp
--- a/Engine/Plugins/VirtualProduction/DynamicMaterial/Source/DynamicMaterialEditor/Private/Components/DMMaterialStageBlend.cpp
+++ b/Engine/Plugins/VirtualProduction/DynamicMaterial/Source/DynamicMaterialEditor/Private/Components/DMMaterialStageBlend.cpp
@@ -424,13 +424,25 @@ bool UDMMaterialStageBlend::GenerateStagePreviewMaterial(UDMMaterialStage* InSta
 
 	if (InputConnectionMap[InputB].Channels[0].OutputChannel != FDMMaterialStageConnectorChannel::WHOLE_CHANNEL)
 	{
-		LastExpression = BuildState->GetBuildUtils().CreateExpressionBitMask(
+		const int32 MaskedOutput = BuildState->GetBuildUtils().FindOutputForBitmask(
 			LastExpression,
-			OutputIndex,
 			InputConnectionMap[InputB].Channels[0].OutputChannel
 		);
 
-		OutputIndex = 0;
+		if (MaskedOutput == INDEX_NONE)
+		{
+			LastExpression = BuildState->GetBuildUtils().CreateExpressionBitMask(
+				LastExpression,
+				OutputIndex,
+				InputConnectionMap[InputB].Channels[0].OutputChannel
+			);
+
+			OutputIndex = 0;
+		}
+		else
+		{
+			OutputIndex = MaskedOutput;
+		}
 	}
 
 	OutMaterialExpression = LastExpression;
diff --git a/Engine/Plugins/VirtualProduction/DynamicMaterial/Source/DynamicMaterialEditor/Private/Components/DMMaterialStageThroughput.cpp b/Engine/Plugins/VirtualProduction/DynamicMaterial/Source/DynamicMaterialEditor/Private/Components/DMMaterialStageThroughput.cpp
--- a/Engine/Plugins/VirtualProduction/DynamicMaterial/Source/DynamicMaterialEditor/Private/Components/DMMaterialStageThroughput.cpp
+++ b/Engine/Plugins/VirtualProduction/DynamicMaterial/Source/DynamicMaterialEditor/Private/Components/DMMaterialStageThroughput.cpp
@@ -327,13 +327,25 @@ void UDMMaterialStageThroughput::ConnectOutputToInput_Internal(const TSharedRef<
 	}
 	else
 	{
-		UMaterialExpression* MaskExpression = InBuildState->GetBuildUtils().CreateExpressionBitMask(
-			InSourceExpression, 
-			InSourceOutputIndex, 
+		const int32 MaskedOutput = InBuildState->GetBuildUtils().FindOutputForBitmask(
+			InSourceExpression,
 			InSourceOutputChannel
 		);
 
-		MaskExpression->ConnectExpression(ExpressionInput, 0);
+		if (MaskedOutput == INDEX_NONE)
+		{
+			UMaterialExpression* MaskExpression = InBuildState->GetBuildUtils().CreateExpressionBitMask(
+				InSourceExpression,
+				InSourceOutputIndex,
+				InSourceOutputChannel
+			);
+
+			MaskExpression->ConnectExpression(ExpressionInput, 0);
+		}
+		else
+		{
+			InSourceExpression->ConnectExpression(ExpressionInput, MaskedOutput);
+		}
 	}
 }
 
diff --git a/Engine/Plugins/VirtualProduction/DynamicMaterial/Source/DynamicMaterialEditor/Private/Components/DMMaterialStageThroughputLayerBlend.cpp b/Engine/Plugins/VirtualProduction/DynamicMaterial/Source/DynamicMaterialEditor/Private/Components/DMMaterialStageThroughputLayerBlend.cpp
--- a/Engine/Plugins/VirtualProduction/DynamicMaterial/Source/DynamicMaterialEditor/Private/Components/DMMaterialStageThroughputLayerBlend.cpp
+++ b/Engine/Plugins/VirtualProduction/DynamicMaterial/Source/DynamicMaterialEditor/Private/Components/DMMaterialStageThroughputLayerBlend.cpp
@@ -711,16 +711,28 @@ bool UDMMaterialStageThroughputLayerBlend::GenerateStagePreviewMaterial(UDMMater
 			Channel = FDMMaterialStageConnectorChannel::FIRST_CHANNEL;
 		}
 
-		ComponentMask = BuildState->GetBuildUtils().CreateExpressionBitMask(
+		const int32 MaskedOutput = BuildState->GetBuildUtils().FindOutputForBitmask(
 			ComponentMask,
-			OutputIndex,
 			Channel
 		);
 
-		SourceExpressions.Add(ComponentMask);
+		if (MaskedOutput == INDEX_NONE)
+		{
+			ComponentMask = BuildState->GetBuildUtils().CreateExpressionBitMask(
+				ComponentMask,
+				OutputIndex,
+				Channel
+			);
 
-		OutputIndex = 0;
-		Channel = FDMMaterialStageConnectorChannel::WHOLE_CHANNEL;
+			SourceExpressions.Add(ComponentMask);
+
+			OutputIndex = 0;
+			Channel = FDMMaterialStageConnectorChannel::WHOLE_CHANNEL;
+		}
+		else
+		{
+			OutputIndex = MaskedOutput;
+		}
 	}
 
 	Layer->ApplyEffects(BuildState, InStage, SourceExpressions, Channel, OutputIndex);
diff --git a/Engine/Plugins/VirtualProduction/DynamicMaterial/Source/DynamicMaterialEditor/Private/Model/DMMaterialBuildUtils.cpp b/Engine/Plugins/VirtualProduction/DynamicMaterial/Source/DynamicMaterialEditor/Private/Model/DMMaterialBuildUtils.cpp
--- a/Engine/Plugins/VirtualProduction/DynamicMaterial/Source/DynamicMaterialEditor/Private/Model/DMMaterialBuildUtils.cpp
+++ b/Engine/Plugins/VirtualProduction/DynamicMaterial/Source/DynamicMaterialEditor/Private/Model/DMMaterialBuildUtils.cpp
@@ -297,6 +297,33 @@ TArray<UMaterialExpression*> FDMMaterialBuildUtils::CreateExpressionInput(UDMMat
 	return BuildState.GetStageSourceExpressions(InInput);
 }
 
+int32 FDMMaterialBuildUtils::FindOutputForBitmask(UMaterialExpression* InExpression, int32 InOutputChannels) const
+{
+	const bool bNeedsR = !!(InOutputChannels & FDMMaterialStageConnectorChannel::FIRST_CHANNEL);
+	const bool bNeedsG = !!(InOutputChannels & FDMMaterialStageConnectorChannel::SECOND_CHANNEL);
+	const bool bNeedsB = !!(InOutputChannels & FDMMaterialStageConnectorChannel::THIRD_CHANNEL);
+	const bool bNeedsA = !!(InOutputChannels & FDMMaterialStageConnectorChannel::FOURTH_CHANNEL);
+
+	const TArray<FExpressionOutput>& Outputs = InExpression->GetOutputs();
+
+	for (int32 OutputIndex = 0; OutputIndex < Outputs.Num(); ++OutputIndex)
+	{
+		const FExpressionOutput& Output = Outputs[OutputIndex];
+
+		const bool bMaskMatchExactR = bNeedsR == (Output.MaskR == 1);
+		const bool bMaskMatchExactG = bNeedsG == (Output.MaskG == 1);
+		const bool bMaskMatchExactB = bNeedsB == (Output.MaskB == 1);
+		const bool bMaskMatchExactA = bNeedsA == (Output.MaskA == 1);
+
+		if (bMaskMatchExactR && bMaskMatchExactG && bMaskMatchExactB && bMaskMatchExactA)
+		{
+			return OutputIndex;
+		}
+	}
+
+	return INDEX_NONE;
+}
+
 UMaterialExpressionComponentMask* FDMMaterialBuildUtils::CreateExpressionBitMask(UMaterialExpression* InExpression, 
 	int32 InOutputIndex, int32 InOutputChannels) const
 {
@@ -304,27 +331,214 @@ UMaterialExpressionComponentMask* FDMMaterialBuildUtils::CreateExpressionBitMask
 		CreateExpression<UMaterialExpressionComponentMask>(UE_DM_NodeComment_Default)
 	);
 
+	BuildState.AddOtherExpressions({Mask});
+
 	Mask->Input.Expression = InExpression;
-	Mask->Input.OutputIndex = InOutputIndex;
 
-	if (InOutputChannels == FDMMaterialStageConnectorChannel::WHOLE_CHANNEL)
+	const bool bNeeds[4] = {
+		!!(InOutputChannels & FDMMaterialStageConnectorChannel::FIRST_CHANNEL),
+		!!(InOutputChannels & FDMMaterialStageConnectorChannel::SECOND_CHANNEL),
+		!!(InOutputChannels & FDMMaterialStageConnectorChannel::THIRD_CHANNEL),
+		!!(InOutputChannels & FDMMaterialStageConnectorChannel::FOURTH_CHANNEL)
+	};
+
+	/**
+	 * Find the best matching expression output for our requirements.
+	 */
+
+	const TArray<FExpressionOutput>& Outputs = InExpression->GetOutputs();
+	int32 OutputCountIndices[4] = {INDEX_NONE, INDEX_NONE, INDEX_NONE, INDEX_NONE};
+	int32 RequiredExactOutputIndex = INDEX_NONE;
+	int32 RequiredOutputIndex = INDEX_NONE;
+	int32 UnmaskedOutputIndex = INDEX_NONE;
+
+	for (int32 OutputIndex = 0; OutputIndex < Outputs.Num(); ++OutputIndex)
 	{
-		const EMaterialValueType OutputType = InExpression->GetOutputValueType(InOutputIndex);
+		const FExpressionOutput& Output = Outputs[OutputIndex];
 
-		Mask->R = (OutputType == MCT_Float1);
-		Mask->G = (OutputType == MCT_Float2);
-		Mask->B = (OutputType == MCT_Float3);
-		Mask->A = (OutputType == MCT_Float4);
+		if (!Output.Mask)
+		{
+			if (UnmaskedOutputIndex == INDEX_NONE)
+			{
+				UnmaskedOutputIndex = OutputIndex;
+			}
+
+			continue;
+		}
 
-		return Mask;
+		if (Output.MaskR)
+		{
+			if (OutputCountIndices[0] == INDEX_NONE)
+			{
+				OutputCountIndices[0] = OutputIndex;
+			}
+
+			if (Output.MaskG)
+			{
+				if (OutputCountIndices[1] == INDEX_NONE)
+				{
+					OutputCountIndices[1] = OutputIndex;
+				}
+
+				if (Output.MaskB)
+				{
+					if (OutputCountIndices[2] == INDEX_NONE)
+					{
+						OutputCountIndices[2] = OutputIndex;
+					}
+
+					if (Output.MaskA)
+					{
+						if (OutputCountIndices[3] == INDEX_NONE)
+						{
+							OutputCountIndices[3] = OutputIndex;
+						}
+					}
+				}
+			}
+		}
+
+		if (RequiredExactOutputIndex == INDEX_NONE)
+		{
+			const bool bMaskMatchExactR = bNeeds[0] == (Output.MaskR == 1);
+			const bool bMaskMatchExactG = bNeeds[1] == (Output.MaskG == 1);
+			const bool bMaskMatchExactB = bNeeds[2] == (Output.MaskB == 1);
+			const bool bMaskMatchExactA = bNeeds[3] == (Output.MaskA == 1);
+
+			if (bMaskMatchExactR && bMaskMatchExactG && bMaskMatchExactB && bMaskMatchExactA)
+			{
+				RequiredExactOutputIndex = OutputIndex;
+			}
+		}
+
+		if (RequiredOutputIndex == INDEX_NONE)
+		{
+			const bool bMaskMatchR = !bNeeds[0] || Output.MaskR;
+			const bool bMaskMatchG = !bNeeds[1] || Output.MaskG;
+			const bool bMaskMatchB = !bNeeds[2] || Output.MaskB;
+			const bool bMaskMatchA = !bNeeds[3] || Output.MaskA;
+
+			if (bMaskMatchR && bMaskMatchG && bMaskMatchB && bMaskMatchA)
+			{
+				RequiredOutputIndex = OutputIndex;
+			}
+		}
 	}
 
-	Mask->R = !!(InOutputChannels & FDMMaterialStageConnectorChannel::FIRST_CHANNEL);
-	Mask->G = !!(InOutputChannels & FDMMaterialStageConnectorChannel::SECOND_CHANNEL);
-	Mask->B = !!(InOutputChannels & FDMMaterialStageConnectorChannel::THIRD_CHANNEL);
-	Mask->A = !!(InOutputChannels & FDMMaterialStageConnectorChannel::FOURTH_CHANNEL);
+	/**
+	 * Set the found output on the mask node.
+	 */
 
-	BuildState.AddOtherExpressions({Mask});
+	if (RequiredExactOutputIndex != INDEX_NONE)
+	{
+		Mask->Input.OutputIndex = RequiredExactOutputIndex;
+	}
+	else if (RequiredOutputIndex != INDEX_NONE)
+	{
+		Mask->Input.OutputIndex = RequiredOutputIndex;
+	}
+	else if (bNeeds[3] && OutputCountIndices[3] != INDEX_NONE)
+	{
+		Mask->Input.OutputIndex = OutputCountIndices[3];
+	}
+	else if (bNeeds[2] && OutputCountIndices[2] != INDEX_NONE)
+	{
+		Mask->Input.OutputIndex = OutputCountIndices[2];
+	}
+	else if (bNeeds[1] && OutputCountIndices[1] != INDEX_NONE)
+	{
+		Mask->Input.OutputIndex = OutputCountIndices[1];
+	}
+	else if (bNeeds[0] && OutputCountIndices[0] != INDEX_NONE)
+	{
+		Mask->Input.OutputIndex = OutputCountIndices[0];
+	}
+	else if (UnmaskedOutputIndex != INDEX_NONE)
+	{
+		Mask->Input.OutputIndex = UnmaskedOutputIndex;
+	}
+	else
+	{
+		Mask->Input.OutputIndex = InOutputIndex;
+	}
+
+	Mask->R = 0;
+	Mask->G = 0;
+	Mask->B = 0;
+	Mask->A = 0;
+
+	/**
+	 * Scan the expression for the matching outputs for the mask input.
+	 * E.g.
+	 * Output: GBA (2,3,4) (no red channel)
+	 * Required: GA (2,4)
+	 * Mask: RB (1,3) (GA is shifted 1 to the left as the first output channel is missing)
+	 */
+
+	int32 StartChannel = INDEX_NONE;
+
+	for (int32 MaskChannel = 0; MaskChannel < 4; ++MaskChannel)
+	{
+		bool bHasMaskChannel = false;
+
+		if (bNeeds[MaskChannel])
+		{
+			if (!Outputs[Mask->Input.OutputIndex].Mask)
+			{
+				bHasMaskChannel = true;
+			}
+			else
+			{
+				switch (MaskChannel)
+				{
+					case 0:
+						bHasMaskChannel = !!Outputs[Mask->Input.OutputIndex].MaskR;
+						break;
+
+					case 1:
+						bHasMaskChannel = !!Outputs[Mask->Input.OutputIndex].MaskG;
+						break;
+
+					case 2:
+						bHasMaskChannel = !!Outputs[Mask->Input.OutputIndex].MaskB;
+						break;
+
+					case 3:
+						bHasMaskChannel = !!Outputs[Mask->Input.OutputIndex].MaskA;
+						break;
+				}
+			}
+		}
+
+		if (!bHasMaskChannel)
+		{
+			continue;
+		}
+
+		if (StartChannel == INDEX_NONE)
+		{
+			StartChannel = MaskChannel;
+		}
+
+		switch (MaskChannel - StartChannel)
+		{
+			case 0:
+				Mask->R = 1;
+				break;
+
+			case 1:
+				Mask->G = 1;
+				break;
+
+			case 2:
+				Mask->B = 1;
+				break;
+
+			case 3:
+				Mask->A = 1;
+				break;
+		}
+	}
 
 	return Mask;
 }
diff --git a/Engine/Plugins/VirtualProduction/DynamicMaterial/Source/DynamicMaterialEditor/Public/Model/DMMaterialBuildUtils.h b/Engine/Plugins/VirtualProduction/DynamicMaterial/Source/DynamicMaterialEditor/Public/Model/DMMaterialBuildUtils.h
--- a/Engine/Plugins/VirtualProduction/DynamicMaterial/Source/DynamicMaterialEditor/Public/Model/DMMaterialBuildUtils.h
+++ b/Engine/Plugins/VirtualProduction/DynamicMaterial/Source/DynamicMaterialEditor/Public/Model/DMMaterialBuildUtils.h
@@ -54,6 +54,9 @@ struct FDMMaterialBuildUtils : public IDMMaterialBuildUtilsInterface
 	/** Creates a series of nodes that display a single input. */
 	DYNAMICMATERIALEDITOR_API virtual TArray<UMaterialExpression*> CreateExpressionInput(UDMMaterialStageInput* InInput) const override;
 
+	/** Searches the outputs of an expression to see if there is an appropriate output to match a requested channel mask. Returns INDEX_NONE on failure. */
+	DYNAMICMATERIALEDITOR_API virtual int32 FindOutputForBitmask(UMaterialExpression* InExpression, int32 InOutputChannels) const override;
+
 	/** Creates a mask expression with the given channels exposed. */
 	DYNAMICMATERIALEDITOR_API virtual UMaterialExpressionComponentMask* CreateExpressionBitMask(UMaterialExpression* InExpression, int32 InOutputIndex,
 		int32 InOutputChannels) const override;
