Material Designer: The "Material Settings" and "Material Type" categories are now localized.

diff --git a/Engine/Plugins/VirtualProduction/DynamicMaterial/Source/DynamicMaterialEditor/Private/UI/PropertyGenerators/DMMaterialModelPropertyRowGenerator.cpp b/Engine/Plugins/VirtualProduction/DynamicMaterial/Source/DynamicMaterialEditor/Private/UI/PropertyGenerators/DMMaterialModelPropertyRowGenerator.cpp
--- a/Engine/Plugins/VirtualProduction/DynamicMaterial/Source/DynamicMaterialEditor/Private/UI/PropertyGenerators/DMMaterialModelPropertyRowGenerator.cpp
+++ b/Engine/Plugins/VirtualProduction/DynamicMaterial/Source/DynamicMaterialEditor/Private/UI/PropertyGenerators/DMMaterialModelPropertyRowGenerator.cpp
@@ -94,7 +94,7 @@ void FDMMaterialModelPropertyRowGenerator::AddGlobalValue(FDMComponentPropertyRo
 		GlobalValueParams.CreatePropertyHandleParams(UDMMaterialValue::ValueName)
 	));
 
-	ComponentHandle.CategoryOverrideName = TEXT("Material Settings");
+	ComponentHandle.CategoryOverrideName = *LOCTEXT("MaterialSettings", "Material Settings").ToString();
 	ComponentHandle.NameOverride = InNameOverride;
 
 	if (UDMMaterialValue* MaterialValue = Cast<UDMMaterialValue>(InComponent))
@@ -124,7 +124,7 @@ void FDMMaterialModelPropertyRowGenerator::AddVariable(FDMComponentPropertyRowGe
 		VariableParams.CreatePropertyHandleParams(InPropertyName)
 	));
 
-	ValueHandle.CategoryOverrideName = TEXT("Material Type");
+	ValueHandle.CategoryOverrideName = *LOCTEXT("MaterialType", "Material Type").ToString();
 	ValueHandle.bEnabled = !IsDynamic(Cast<UDynamicMaterialModelBase>(InParams.Object));
 }
 
