Material Designer: Now checks if properties are keyable when displaying keyframe buttons.

diff --git a/Engine/Plugins/VirtualProduction/DynamicMaterial/Source/DynamicMaterialEditor/Private/UI/Widgets/Editor/SDMObjectEditorWidgetBase.cpp b/Engine/Plugins/VirtualProduction/DynamicMaterial/Source/DynamicMaterialEditor/Private/UI/Widgets/Editor/SDMObjectEditorWidgetBase.cpp
--- a/Engine/Plugins/VirtualProduction/DynamicMaterial/Source/DynamicMaterialEditor/Private/UI/Widgets/Editor/SDMObjectEditorWidgetBase.cpp
+++ b/Engine/Plugins/VirtualProduction/DynamicMaterial/Source/DynamicMaterialEditor/Private/UI/Widgets/Editor/SDMObjectEditorWidgetBase.cpp
@@ -457,7 +457,17 @@ bool SDMObjectEditorWidgetBase::CanCreateKeyFrame(TWeakPtr<IPropertyHandle> InPr
 		return false;
 	}
 
-	return PreviewPropertyHandle->IsValidHandle() && OriginalPropertyHandle->IsValidHandle();
+	if (!PreviewPropertyHandle->IsValidHandle() || !OriginalPropertyHandle->IsValidHandle())
+	{
+		return false;
+	}
+
+	if (UObject* Object = GetObject())
+	{
+		return KeyframeHandler->IsPropertyKeyable(Object->GetClass(), *OriginalPropertyHandle);
+	}
+
+	return false;
 }
 
 void SDMObjectEditorWidgetBase::CreateKeyFrame(TWeakPtr<IPropertyHandle> InPreviewPropertyHandleWeak,
