Material Designer: Disabled substrate material support.

diff --git a/Engine/Plugins/VirtualProduction/DynamicMaterial/Source/DynamicMaterialEditor/Private/Utils/DMMaterialUtils.cpp b/Engine/Plugins/VirtualProduction/DynamicMaterial/Source/DynamicMaterialEditor/Private/Utils/DMMaterialUtils.cpp
--- a/Engine/Plugins/VirtualProduction/DynamicMaterial/Source/DynamicMaterialEditor/Private/Utils/DMMaterialUtils.cpp
+++ b/Engine/Plugins/VirtualProduction/DynamicMaterial/Source/DynamicMaterialEditor/Private/Utils/DMMaterialUtils.cpp
@@ -6,8 +6,9 @@
 
 bool FDMMaterialUtils::IsMaterialPropertyActive(const FParams& InParams)
 {
-	const bool bSubstrateEnabled = Substrate::IsSubstrateEnabled();
-	const bool bSubstrateOpacityOverrideAllowed = InParams.BlendMode == BLEND_AlphaComposite; // Should we always have it enabled to be able to be plugged in an fed when blend mode is toggled later on a material instance?
+	// Substrate temporarily disabled.
+	constexpr const bool bSubstrateEnabled = false;
+	constexpr bool bSubstrateOpacityOverrideAllowed = false;
 
 	switch (InParams.Domain)
 	{
