Material Designer: Default textures are now default linear colors on the majority of channels.

diff --git a/Engine/Plugins/VirtualProduction/DynamicMaterial/Source/DynamicMaterialEditor/Private/Components/DMMaterialProperty.cpp b/Engine/Plugins/VirtualProduction/DynamicMaterial/Source/DynamicMaterialEditor/Private/Components/DMMaterialProperty.cpp
--- a/Engine/Plugins/VirtualProduction/DynamicMaterial/Source/DynamicMaterialEditor/Private/Components/DMMaterialProperty.cpp
+++ b/Engine/Plugins/VirtualProduction/DynamicMaterial/Source/DynamicMaterialEditor/Private/Components/DMMaterialProperty.cpp
@@ -231,37 +231,6 @@ void UDMMaterialProperty::AddDefaultMaskStage(UDMMaterialLayerObject* InLayer)
 	check(MaskStage);
 
 	InLayer->SetStage(EDMMaterialLayerStage::Mask, MaskStage);
-
-	if (UTexture* MaskTexture = UDynamicMaterialEditorSettings::Get()->DefaultMask.LoadSynchronous())
-	{
-		const TArray<UDMMaterialStageInput*> MaskStageInputs = MaskStage->GetInputs();
-
-		for (UDMMaterialStageInput* MaskStageInput : MaskStageInputs)
-		{
-			if (UDMMaterialStageInputExpression* MaskInputExpression = Cast<UDMMaterialStageInputExpression>(MaskStageInput))
-			{
-				if (UDMMaterialStageExpressionTextureSample* MaskInputTextureSample = Cast<UDMMaterialStageExpressionTextureSample>(MaskInputExpression->GetMaterialStageExpression()))
-				{
-					if (UDMMaterialStage* MaskTextureInputStage = MaskInputExpression->GetSubStage())
-					{
-						const TArray<UDMMaterialStageInput*> MaskTextureStageInputs = MaskTextureInputStage->GetInputs();
-
-						for (UDMMaterialStageInput* MaskTextureStageInput : MaskTextureStageInputs)
-						{
-							if (UDMMaterialStageInputValue* MaskTextureInputValue = Cast<UDMMaterialStageInputValue>(MaskTextureStageInput))
-							{
-								if (UDMMaterialValueTexture* MaskTextureValue = Cast<UDMMaterialValueTexture>(MaskTextureInputValue->GetValue()))
-								{
-									MaskTextureValue->SetDefaultValue(MaskTexture);
-									MaskTextureValue->ApplyDefaultValue();
-								}
-							}
-						}
-					}
-				}
-			}
-		}
-	}
 }
 
 UDMMaterialComponent* UDMMaterialProperty::AddComponent(FName InName, UDMMaterialComponent* InComponent)
diff --git a/Engine/Plugins/VirtualProduction/DynamicMaterial/Source/DynamicMaterialEditor/Private/Components/DMMaterialStageThroughputLayerBlend.cpp b/Engine/Plugins/VirtualProduction/DynamicMaterial/Source/DynamicMaterialEditor/Private/Components/DMMaterialStageThroughputLayerBlend.cpp
--- a/Engine/Plugins/VirtualProduction/DynamicMaterial/Source/DynamicMaterialEditor/Private/Components/DMMaterialStageThroughputLayerBlend.cpp
+++ b/Engine/Plugins/VirtualProduction/DynamicMaterial/Source/DynamicMaterialEditor/Private/Components/DMMaterialStageThroughputLayerBlend.cpp
@@ -11,6 +11,7 @@
 #include "Components/MaterialStageExpressions/DMMSETextureSample.h"
 #include "Components/MaterialStageInputs/DMMSIExpression.h"
 #include "Components/MaterialStageInputs/DMMSIValue.h"
+#include "Components/MaterialValues/DMMaterialValueFloat3RGB.h"
 #include "Components/MaterialValues/DMMaterialValueTexture.h"
 #include "DMEDefs.h"
 #include "DynamicMaterialEditorSettings.h"
@@ -126,33 +127,20 @@ void UDMMaterialStageThroughputLayerBlend::AddDefaultInput(int32 InInputIndex) c
  
 		case InputMaskSource:
 		{
-			UDMMaterialStageInputExpression* InputExpression = UDMMaterialStageInputExpression::ChangeStageInput_Expression(
+			UDMMaterialStageInputValue* InputValue = UDMMaterialStageInputValue::ChangeStageInput_NewLocalValue(
 				Stage,
-				UDMMaterialStageExpressionTextureSample::StaticClass(),
 				InInputIndex,
-				FDMMaterialStageConnectorChannel::WHOLE_CHANNEL, 
-				0, // RGB pin
-				FDMMaterialStageConnectorChannel::WHOLE_CHANNEL // Alpha channel
-			);
-
-			UDMMaterialSubStage* SubStage = InputExpression->GetSubStage();
-			check(SubStage);
-
-			UDMMaterialStageInputValue* InputValue = UDMMaterialStageInputValue::ChangeStageInput_NewLocalValue(
-				SubStage,
-				0, 
-				FDMMaterialStageConnectorChannel::WHOLE_CHANNEL, 
-				EDMValueType::VT_Texture, 
+				FDMMaterialStageConnectorChannel::WHOLE_CHANNEL,
+				EDMValueType::VT_Float3_RGB,
 				FDMMaterialStageConnectorChannel::WHOLE_CHANNEL
 			);
-
 			check(InputValue);
 
-			UDMMaterialValueTexture* InputTexture = Cast<UDMMaterialValueTexture>(InputValue->GetValue());
-			check(InputTexture);
+			UDMMaterialValueFloat3RGB* InputColor = Cast<UDMMaterialValueFloat3RGB>(InputValue->GetValue());
+			check(InputColor);
 
-			InputTexture->SetDefaultValue(UDynamicMaterialEditorSettings::Get()->DefaultMask.LoadSynchronous());
-			InputTexture->ApplyDefaultValue();
+			InputColor->SetDefaultValue(FLinearColor::White);
+			InputColor->ApplyDefaultValue();
 			break;
 		}
  
diff --git a/Engine/Plugins/VirtualProduction/DynamicMaterial/Source/DynamicMaterialEditor/Private/DynamicMaterialEditorSettings.cpp b/Engine/Plugins/VirtualProduction/DynamicMaterial/Source/DynamicMaterialEditor/Private/DynamicMaterialEditorSettings.cpp
--- a/Engine/Plugins/VirtualProduction/DynamicMaterial/Source/DynamicMaterialEditor/Private/DynamicMaterialEditorSettings.cpp
+++ b/Engine/Plugins/VirtualProduction/DynamicMaterial/Source/DynamicMaterialEditor/Private/DynamicMaterialEditorSettings.cpp
@@ -27,20 +27,16 @@ namespace UE::DynamicMaterialEditor::Private
 	const TMap<EDMMaterialPropertyType, FDMDefaultMaterialPropertySlotValue> DefaultSlotValues = {
 		{EDMMaterialPropertyType::BaseColor,           TSoftObjectPtr<UTexture>(FSoftObjectPath(TEXT("/Script/Engine.Texture2D'/DynamicMaterial/T_Default_Texture.T_Default_Texture'")))},
 		{EDMMaterialPropertyType::EmissiveColor,       TSoftObjectPtr<UTexture>(FSoftObjectPath(TEXT("/Script/Engine.Texture2D'/DynamicMaterial/Textures/SlotDefaults/T_MD_Emissive.T_MD_Emissive'")))},
-		{EDMMaterialPropertyType::Opacity,             TSoftObjectPtr<UTexture>(FSoftObjectPath(TEXT("/Script/Engine.Texture2D'/Engine/EngineResources/WhiteSquareTexture.WhiteSquareTexture'")))},
-		{EDMMaterialPropertyType::OpacityMask,         TSoftObjectPtr<UTexture>(FSoftObjectPath(TEXT("/Script/Engine.Texture2D'/Engine/EngineResources/WhiteSquareTexture.WhiteSquareTexture'")))},
-		{EDMMaterialPropertyType::Metallic,            TSoftObjectPtr<UTexture>(FSoftObjectPath(TEXT("/Script/Engine.Texture2D'/DynamicMaterial/Textures/SlotDefaults/T_MD_Metallic.T_MD_Metallic'")))},
-		{EDMMaterialPropertyType::Specular,            FLinearColor(0.5f, 0.5f, 0.5f, 1.f)},
-		{EDMMaterialPropertyType::Roughness,           TSoftObjectPtr<UTexture>(FSoftObjectPath(TEXT("/Script/Engine.Texture2D'/DynamicMaterial/Textures/SlotDefaults/T_MD_Roughness.T_MD_Roughness'")))},
-		{EDMMaterialPropertyType::Normal,              {
-				EDMDefaultMaterialPropertySlotValueType::Texture,
-				TSoftObjectPtr<UTexture>(FSoftObjectPath(TEXT("/Script/Engine.Texture2D'/DynamicMaterial/Textures/SlotDefaults/T_MD_Normal.T_MD_Normal'"))),
-				FLinearColor(0.f, 0.f, 1.f, 1.f)
-			}},
-		{EDMMaterialPropertyType::AmbientOcclusion,    TSoftObjectPtr<UTexture>(FSoftObjectPath(TEXT("/Script/Engine.Texture2D'/DynamicMaterial/Textures/SlotDefaults/T_MD_AmbientOcclusion.T_MD_AmbientOcclusion'")))},
-		{EDMMaterialPropertyType::Displacement,        TSoftObjectPtr<UTexture>(FSoftObjectPath(TEXT("/Script/Engine.Texture2D'/DynamicMaterial/Textures/SlotDefaults/T_MD_Displacement.T_MD_Displacement'")))},
-		{EDMMaterialPropertyType::SubsurfaceColor,     TSoftObjectPtr<UTexture>(FSoftObjectPath(TEXT("/Script/Engine.Texture2D'/DynamicMaterial/Textures/SlotDefaults/T_MD_SubsurfaceColor.T_MD_SubsurfaceColor'")))},
-		{EDMMaterialPropertyType::SurfaceThickness,    FLinearColor(0.0f, 0.0f, 0.0f, 1.f)},
+		{EDMMaterialPropertyType::Opacity,             FLinearColor::White},
+		{EDMMaterialPropertyType::OpacityMask,         FLinearColor::White},
+		{EDMMaterialPropertyType::Metallic,            FLinearColor::Black},
+		{EDMMaterialPropertyType::Specular,            FLinearColor::Gray},
+		{EDMMaterialPropertyType::Roughness,           FLinearColor::Gray},
+		{EDMMaterialPropertyType::Normal,              FLinearColor::Blue}, // +Z
+		{EDMMaterialPropertyType::AmbientOcclusion,    FLinearColor::White},
+		{EDMMaterialPropertyType::Displacement,        FLinearColor::Gray},
+		{EDMMaterialPropertyType::SubsurfaceColor,     FLinearColor::Black},
+		{EDMMaterialPropertyType::SurfaceThickness,    FLinearColor::Black},
 	};
 }
 
