Material Designer: Changed the way the generated materials are compiled.

diff --git a/Engine/Plugins/VirtualProduction/DynamicMaterial/Source/DynamicMaterialEditor/Private/Model/DMMaterialBuildState.cpp b/Engine/Plugins/VirtualProduction/DynamicMaterial/Source/DynamicMaterialEditor/Private/Model/DMMaterialBuildState.cpp
--- a/Engine/Plugins/VirtualProduction/DynamicMaterial/Source/DynamicMaterialEditor/Private/Model/DMMaterialBuildState.cpp
+++ b/Engine/Plugins/VirtualProduction/DynamicMaterial/Source/DynamicMaterialEditor/Private/Model/DMMaterialBuildState.cpp
@@ -44,9 +44,7 @@ FDMMaterialBuildState::~FDMMaterialBuildState()
 {
 	FDMMaterialNodeArranger(DynamicMaterial).ArrangeNodes();
 
-	// let the material update itself if necessary
-	DynamicMaterial->PreEditChange(nullptr);
-	DynamicMaterial->PostEditChange();
+	BuildMaterial();
 }
 
 UMaterial* FDMMaterialBuildState::GetDynamicMaterial() const
@@ -518,3 +516,9 @@ void FDMMaterialBuildState::SetGlobalExpression(FName InName, UMaterialExpressio
 {
 	GlobalExpressions.FindOrAdd(InName) = InExpression;
 }
+
+void FDMMaterialBuildState::BuildMaterial()
+{
+	DynamicMaterial->UpdateCachedExpressionData();
+	DynamicMaterial->CacheShaders(EMaterialShaderPrecompileMode::Background);
+}
diff --git a/Engine/Plugins/VirtualProduction/DynamicMaterial/Source/DynamicMaterialEditor/Public/Model/DMMaterialBuildState.h b/Engine/Plugins/VirtualProduction/DynamicMaterial/Source/DynamicMaterialEditor/Public/Model/DMMaterialBuildState.h
--- a/Engine/Plugins/VirtualProduction/DynamicMaterial/Source/DynamicMaterialEditor/Public/Model/DMMaterialBuildState.h
+++ b/Engine/Plugins/VirtualProduction/DynamicMaterial/Source/DynamicMaterialEditor/Public/Model/DMMaterialBuildState.h
@@ -169,4 +169,6 @@ private:
 	TMap<const UDMMaterialStageSource*, TArray<UMaterialExpression*>> StageSources;
 	TSet<UMaterialExpression*> OtherExpressions;
 	TMap<FName, UMaterialExpression*> GlobalExpressions;
+
+	void BuildMaterial();
 };
