Material Designer: Removed extra, unnecessary padding on details panel widget.

diff --git a/Engine/Plugins/VirtualProduction/DynamicMaterial/Source/DynamicMaterialEditor/Private/DetailsPanel/Widgets/SDMDetailsPanelMaterialInterfaceWidget.cpp b/Engine/Plugins/VirtualProduction/DynamicMaterial/Source/DynamicMaterialEditor/Private/DetailsPanel/Widgets/SDMDetailsPanelMaterialInterfaceWidget.cpp
--- a/Engine/Plugins/VirtualProduction/DynamicMaterial/Source/DynamicMaterialEditor/Private/DetailsPanel/Widgets/SDMDetailsPanelMaterialInterfaceWidget.cpp
+++ b/Engine/Plugins/VirtualProduction/DynamicMaterial/Source/DynamicMaterialEditor/Private/DetailsPanel/Widgets/SDMDetailsPanelMaterialInterfaceWidget.cpp
@@ -44,7 +44,6 @@ void SDMDetailsPanelMaterialInterfaceWidget::Construct(const FArguments& InArgs,
 		.AutoHeight()
 		.HAlign(EHorizontalAlignment::HAlign_Left)
 		.VAlign(EVerticalAlignment::VAlign_Center)
-		.Padding(10.f, 5.f, 10.f, 5.f)
 		[
 			SNew(SObjectPropertyEntryBox)
 			.AllowClear(true)
