Material Designer: Removed in-plugin material property testing.

#Synced-CL 48081101

diff --git a/Engine/Plugins/VirtualProduction/DynamicMaterial/Source/DynamicMaterialEditor/Private/Components/DMMaterialProperty.cpp b/Engine/Plugins/VirtualProduction/DynamicMaterial/Source/DynamicMaterialEditor/Private/Components/DMMaterialProperty.cpp
--- a/Engine/Plugins/VirtualProduction/DynamicMaterial/Source/DynamicMaterialEditor/Private/Components/DMMaterialProperty.cpp
+++ b/Engine/Plugins/VirtualProduction/DynamicMaterial/Source/DynamicMaterialEditor/Private/Components/DMMaterialProperty.cpp
@@ -34,7 +34,6 @@
 #include "Model/DMMaterialBuildUtils.h"
 #include "Model/DynamicMaterialModelEditorOnlyData.h"
 #include "Utils/DMMaterialFunctionLibrary.h"
-#include "Utils/DMMaterialUtils.h"
 #include "Utils/DMPrivate.h"
 #include "Utils/DMUtils.h"
 
@@ -784,22 +783,14 @@ void UDMMaterialProperty::SetOutputProcessor(UMaterialFunctionInterface* InFunct
 	OnOutputProcessorUpdated();
 }
 
-bool UDMMaterialProperty::IsValidForModel(const UDynamicMaterialModelEditorOnlyData& InMaterialModel) const
-{
-	return FDMMaterialUtils::IsMaterialPropertyActive({
-		FDMUtils::MaterialPropertyTypeToMaterialProperty(MaterialProperty),
-		InMaterialModel.GetDomain(),
-		InMaterialModel.GetBlendMode(),
-		InMaterialModel.GetShadingModel() == EDMMaterialShadingModel::DefaultLit ? EMaterialShadingModel::MSM_DefaultLit : EMaterialShadingModel::MSM_Unlit,
-		TLM_Surface,
-		/* Tesselation enabled */ InMaterialModel.IsNaniteTessellationEnabled(),
-		/* BlendableOutputAlpha (Post Process Alpha) */ false,
-		/* Uses Distortion */ false,
-		/* Shading model from main material */ false,
-		/* Outputting translucency velocity */ InMaterialModel.GetBlendMode() != BLEND_Opaque,
-		/* Thin surface */ false,
-		/* Is supported (substrate check) */ true
-	});
+bool UDMMaterialProperty::IsValidForModel(const UDynamicMaterialModelEditorOnlyData& InEditorOnlyData) const
+{
+	if (UMaterial* GeneratedMaterial = InEditorOnlyData.GetGeneratedMaterial())
+	{
+		return GeneratedMaterial->IsPropertyActiveInEditor(FDMUtils::MaterialPropertyTypeToMaterialProperty(MaterialProperty));
+	}
+
+	return false;
 }
 
 void UDMMaterialProperty::OnOutputProcessorUpdated()
diff --git a/Engine/Plugins/VirtualProduction/DynamicMaterial/Source/DynamicMaterialEditor/Private/Utils/DMMaterialUtils.cpp b/Engine/Plugins/VirtualProduction/DynamicMaterial/Source/DynamicMaterialEditor/Private/Utils/DMMaterialUtils.cpp
deleted file mode 100644
diff --git a/Engine/Plugins/VirtualProduction/DynamicMaterial/Source/DynamicMaterialEditor/Private/Utils/DMMaterialUtils.h b/Engine/Plugins/VirtualProduction/DynamicMaterial/Source/DynamicMaterialEditor/Private/Utils/DMMaterialUtils.h
deleted file mode 100644
diff --git a/Engine/Plugins/VirtualProduction/DynamicMaterial/Source/DynamicMaterialEditor/Public/Components/DMMaterialProperty.h b/Engine/Plugins/VirtualProduction/DynamicMaterial/Source/DynamicMaterialEditor/Public/Components/DMMaterialProperty.h
--- a/Engine/Plugins/VirtualProduction/DynamicMaterial/Source/DynamicMaterialEditor/Public/Components/DMMaterialProperty.h
+++ b/Engine/Plugins/VirtualProduction/DynamicMaterial/Source/DynamicMaterialEditor/Public/Components/DMMaterialProperty.h
@@ -105,7 +105,7 @@ public:
 	 * @param InMaterialModel The model with the settings.
 	 * @return True if it is valid.
 	 */
-	DYNAMICMATERIALEDITOR_API virtual bool IsValidForModel(const UDynamicMaterialModelEditorOnlyData& InMaterialModel) const;
+	DYNAMICMATERIALEDITOR_API virtual bool IsValidForModel(const UDynamicMaterialModelEditorOnlyData& InEditorOnlyData) const;
 
 	/**
 	 * Empty the connection map and recalculate the input types.
