Material Designer: Fixed bug preventing the general settings panel from propertly updating the preview material (and the ui)

#Synced-CL 48081101

diff --git a/Engine/Plugins/VirtualProduction/DynamicMaterial/Source/DynamicMaterialEditor/Private/UI/Widgets/SDMMaterialEditor.cpp b/Engine/Plugins/VirtualProduction/DynamicMaterial/Source/DynamicMaterialEditor/Private/UI/Widgets/SDMMaterialEditor.cpp
--- a/Engine/Plugins/VirtualProduction/DynamicMaterial/Source/DynamicMaterialEditor/Private/UI/Widgets/SDMMaterialEditor.cpp
+++ b/Engine/Plugins/VirtualProduction/DynamicMaterial/Source/DynamicMaterialEditor/Private/UI/Widgets/SDMMaterialEditor.cpp
@@ -1283,7 +1283,7 @@ TSharedRef<SDMToolBar> SDMMaterialEditor::CreateSlot_ToolBar()
 
 TSharedRef<SDMMaterialGlobalSettingsEditor> SDMMaterialEditor::CreateSlot_GlobalSettingsEditor()
 {
-	return SNew(SDMMaterialGlobalSettingsEditor, SharedThis(this), GetOriginalMaterialModelBase());
+	return SNew(SDMMaterialGlobalSettingsEditor, SharedThis(this), GetPreviewMaterialModelBase());
 }
 
 TSharedRef<SDMMaterialProperties> SDMMaterialEditor::CreateSlot_MaterialProperties()
