Material Designer: Added a lot more options to the global material setings.

#Synced-CL 49024827

diff --git a/Engine/Plugins/VirtualProduction/DynamicMaterial/Source/DynamicMaterialEditor/Private/DynamicMaterialEditorSettings.cpp b/Engine/Plugins/VirtualProduction/DynamicMaterial/Source/DynamicMaterialEditor/Private/DynamicMaterialEditorSettings.cpp
--- a/Engine/Plugins/VirtualProduction/DynamicMaterial/Source/DynamicMaterialEditor/Private/DynamicMaterialEditorSettings.cpp
+++ b/Engine/Plugins/VirtualProduction/DynamicMaterial/Source/DynamicMaterialEditor/Private/DynamicMaterialEditorSettings.cpp
@@ -15,6 +15,7 @@
 #include "Misc/Paths.h"
 #include "Model/DynamicMaterialModel.h"
 #include "Model/DynamicMaterialModelBase.h"
+#include "Model/DynamicMaterialModelEditorOnlyData.h"
 #include "Modules/ModuleManager.h"
 #include "UObject/UObjectIterator.h"
 
@@ -154,16 +155,16 @@ UDynamicMaterialEditorSettings::UDynamicMaterialEditorSettings()
 	Opaque.bEmissive = true;
 	Opaque.DefaultBlendMode = BLEND_Opaque;
 	Opaque.DefaultShadingModel = EDMMaterialShadingModel::DefaultLit;
-	Opaque.bDefaultAnimated = false;
-	Opaque.bDefaultTwoSided = true;
+	Opaque.UsageFlag |= EDMMaterialFlagsUsage::DMMFU_NiagaraMeshParticles;
+	Opaque.GeneralFlags |= static_cast<uint8>(EDMMaterialFlagsGeneral::TwoSided);
 
 	FDMMaterialChannelListPreset Emissive;
 	Emissive.Name = TEXT("Emissive");
 	Emissive.bEmissive = true;
 	Emissive.DefaultBlendMode = BLEND_Opaque;
 	Emissive.DefaultShadingModel = EDMMaterialShadingModel::Unlit;
-	Emissive.bDefaultAnimated = false;
-	Emissive.bDefaultTwoSided = true;
+	Emissive.UsageFlag |= EDMMaterialFlagsUsage::DMMFU_NiagaraMeshParticles;
+	Emissive.GeneralFlags |= static_cast<uint8>(EDMMaterialFlagsGeneral::TwoSided);
 
 	FDMMaterialChannelListPreset Translucent;
 	Translucent.Name = TEXT("Translucent");
@@ -171,8 +172,8 @@ UDynamicMaterialEditorSettings::UDynamicMaterialEditorSettings()
 	Translucent.bOpacity = true;
 	Translucent.DefaultBlendMode = BLEND_Translucent;
 	Translucent.DefaultShadingModel = EDMMaterialShadingModel::Unlit;
-	Translucent.bDefaultAnimated = false;
-	Translucent.bDefaultTwoSided = true;
+	Translucent.UsageFlag |= EDMMaterialFlagsUsage::DMMFU_NiagaraMeshParticles;
+	Translucent.GeneralFlags |= static_cast<uint8>(EDMMaterialFlagsGeneral::TwoSided);
 
 	FDMMaterialChannelListPreset PBR;
 	PBR.Name = TEXT("PBR");
@@ -187,8 +188,8 @@ UDynamicMaterialEditorSettings::UDynamicMaterialEditorSettings()
 	PBR.bDisplacement = true;
 	PBR.DefaultBlendMode = BLEND_Opaque;
 	PBR.DefaultShadingModel = EDMMaterialShadingModel::DefaultLit;
-	PBR.bDefaultAnimated = false;
-	PBR.bDefaultTwoSided = true;
+	PBR.UsageFlag |= EDMMaterialFlagsUsage::DMMFU_NiagaraMeshParticles;
+	PBR.GeneralFlags |= static_cast<uint8>(EDMMaterialFlagsGeneral::TwoSided);
 
 	FDMMaterialChannelListPreset All;
 	All.Name = TEXT("All");
@@ -210,8 +211,8 @@ UDynamicMaterialEditorSettings::UDynamicMaterialEditorSettings()
 	All.bSurfaceThickness = true;
 	All.DefaultBlendMode = BLEND_Opaque;
 	All.DefaultShadingModel = EDMMaterialShadingModel::DefaultLit;
-	All.bDefaultAnimated = false;
-	All.bDefaultTwoSided = true;
+	All.UsageFlag |= EDMMaterialFlagsUsage::DMMFU_NiagaraMeshParticles;
+	All.GeneralFlags |= static_cast<uint8>(EDMMaterialFlagsGeneral::TwoSided);
 
 	MaterialChannelPresets.Add(PBR);
 	MaterialChannelPresets.Add(Opaque);
@@ -246,18 +247,22 @@ void UDynamicMaterialEditorSettings::PostInitProperties()
 {
 	Super::PostInitProperties();
 
-	if (bValidatedPresets)
+	if (!bValidatedPresets)
 	{
-		return;
+		for (FDMMaterialChannelListPreset& Preset : MaterialChannelPresets)
+		{
+			Preset.bBaseColor = Preset.Name != TEXT("Emissive");
+			Preset.bEmissive = true;
+		}
+
+		bValidatedPresets = true;
 	}
 
-	for (FDMMaterialChannelListPreset& Preset : MaterialChannelPresets)
+	if (!bConvertedVarsToFlags)
 	{
-		Preset.bBaseColor = Preset.Name != TEXT("Emissive");
-		Preset.bEmissive = true;
+		ConvertDeprecatedVarsToFlags();
+		bConvertedVarsToFlags = true;
 	}
-
-	bValidatedPresets = true;
 }
 
 void UDynamicMaterialEditorSettings::PreEditChange(FEditPropertyChain& InPropertyAboutToChange)
@@ -603,4 +608,24 @@ void UDynamicMaterialEditorSettings::EnsureUniqueChannelPresetNames()
 	PreEditPresetNames.Empty();
 }
 
+void UDynamicMaterialEditorSettings::ConvertDeprecatedVarsToFlags()
+{
+	PRAGMA_DISABLE_DEPRECATION_WARNINGS
+
+	for (FDMMaterialChannelListPreset& Preset : MaterialChannelPresets)
+	{
+		if (Preset.bDefaultTwoSided_DEPRECATED)
+		{
+			Preset.GeneralFlags |= static_cast<uint8>(EDMMaterialFlagsGeneral::TwoSided);
+		}
+
+		if (Preset.bDefaultAnimated_DEPRECATED)
+		{
+			Preset.MotionFlags |= static_cast<uint8>(EDMMaterialFlagsMotion::PixelAnimation);
+		}
+	}
+
+	PRAGMA_ENABLE_DEPRECATION_WARNINGS
+}
+
 #undef LOCTEXT_NAMESPACE
diff --git a/Engine/Plugins/VirtualProduction/DynamicMaterial/Source/DynamicMaterialEditor/Private/DynamicMaterialEditorSettings.h b/Engine/Plugins/VirtualProduction/DynamicMaterial/Source/DynamicMaterialEditor/Private/DynamicMaterialEditorSettings.h
--- a/Engine/Plugins/VirtualProduction/DynamicMaterial/Source/DynamicMaterialEditor/Private/DynamicMaterialEditorSettings.h
+++ b/Engine/Plugins/VirtualProduction/DynamicMaterial/Source/DynamicMaterialEditor/Private/DynamicMaterialEditorSettings.h
@@ -119,13 +119,50 @@ struct FDMMaterialChannelListPreset
 	UPROPERTY(Config, EditAnywhere, BlueprintReadOnly, Category = "Channel")
 	EDMMaterialShadingModel DefaultShadingModel = EDMMaterialShadingModel::Unlit;
 
-	UPROPERTY(Config, EditAnywhere, BlueprintReadOnly, Category = "Channel")
-	bool bDefaultAnimated = false;
+	UPROPERTY(Config, EditAnywhere, BlueprintReadOnly, Category = "Usage",
+		meta = (Bitmask, BitmaskEnum = "/Script/DynamicMaterialEditor.EDMMaterialFlagsUsage"))
+	int32 UsageFlag = 0;
 
-	UPROPERTY(Config, EditAnywhere, BlueprintReadOnly, Category = "Channel")
-	bool bDefaultTwoSided = true;
+	UPROPERTY(Config, EditAnywhere, BlueprintReadOnly, Category = "General",
+		meta = (Bitmask, BitmaskEnum = "/Script/DynamicMaterialEditor.EDMMaterialFlagsGeneral"))
+	uint8 GeneralFlags = 0;
+
+	UPROPERTY(Config, EditAnywhere, BlueprintReadOnly, Category = "Lighting",
+		meta = (Bitmask, BitmaskEnum = "/Script/DynamicMaterialEditor.EDMMaterialFlagsLighting"))
+	uint8 LightingFlags = 0;
+
+	UPROPERTY(Config, EditAnywhere, BlueprintReadOnly, Category = "Translucency",
+		meta = (Bitmask, BitmaskEnum = "/Script/DynamicMaterialEditor.EDMMaterialFlagsTranslucency"))
+	uint8 TranslucencyFlags = 0;
+
+	UPROPERTY(Config, EditAnywhere, BlueprintReadOnly, Category = "Motion",
+		meta = (Bitmask, BitmaskEnum = "/Script/DynamicMaterialEditor.EDMMaterialFlagsMotion"))
+	uint8 MotionFlags = 0;
+
+	UPROPERTY(Config, EditAnywhere, BlueprintReadOnly, Category = "Forward Renderer",
+		meta = (Bitmask, BitmaskEnum = "/Script/DynamicMaterialEditor.EDMMaterialFlagsForwardRenderer"))
+	uint8 ForwardRendererFlags = 0;
 
 	bool IsPropertyEnabled(EDMMaterialPropertyType InProperty) const;
+
+	/*
+	 * Deprecated
+	 */
+	UE_DEPRECATED(5.8, "Moved to MotionFlags")
+	UPROPERTY(Config, meta = (DeprecatedProperty, DeprecationMessage = "Moved to MotionFlags"))
+	bool bDefaultAnimated_DEPRECATED = false;
+
+	UE_DEPRECATED(5.8, "Moved to GeneralFlags")
+	UPROPERTY(Config, meta = (DeprecatedProperty, DeprecationMessage = "Moved to GeneralFlags"))
+	bool bDefaultTwoSided_DEPRECATED = true;
+
+	PRAGMA_DISABLE_DEPRECATION_WARNINGS
+	FDMMaterialChannelListPreset() = default;
+	FDMMaterialChannelListPreset(const FDMMaterialChannelListPreset&) = default;
+	FDMMaterialChannelListPreset(FDMMaterialChannelListPreset&&) = default;
+	FDMMaterialChannelListPreset& operator=(const FDMMaterialChannelListPreset&) = default;
+	FDMMaterialChannelListPreset& operator=(FDMMaterialChannelListPreset&&) = default;
+	PRAGMA_ENABLE_DEPRECATION_WARNINGS
 };
 
 UENUM(BlueprintType)
@@ -351,6 +388,9 @@ public:
 	UPROPERTY(Config)
 	bool bValidatedPresets = false;
 
+	UPROPERTY(Config)
+	bool bConvertedVarsToFlags = false;
+
 	/** This variable is accessed in multiple places, so this is a quick accessor. */
 	static bool IsUseLinearColorForVectorsEnabled();
 
@@ -377,4 +417,6 @@ private:
 	FOnFinishedChangingProperties OnSettingsChanged;
 
 	void EnsureUniqueChannelPresetNames();
+
+	void ConvertDeprecatedVarsToFlags();
 };
diff --git a/Engine/Plugins/VirtualProduction/DynamicMaterial/Source/DynamicMaterialEditor/Private/Model/DynamicMaterialModelEditorOnlyData.cpp b/Engine/Plugins/VirtualProduction/DynamicMaterial/Source/DynamicMaterialEditor/Private/Model/DynamicMaterialModelEditorOnlyData.cpp
--- a/Engine/Plugins/VirtualProduction/DynamicMaterial/Source/DynamicMaterialEditor/Private/Model/DynamicMaterialModelEditorOnlyData.cpp
+++ b/Engine/Plugins/VirtualProduction/DynamicMaterial/Source/DynamicMaterialEditor/Private/Model/DynamicMaterialModelEditorOnlyData.cpp
@@ -42,6 +42,7 @@
 #include "Misc/Guid.h"
 #include "Model/DMMaterialBuildState.h"
 #include "Model/DynamicMaterialModel.h"
+#include "Model/DynamicMaterialModelEditorOnlyDataVersion.h"
 #include "UObject/Package.h"
 #include "Utils/DMBuildRequestSubsystem.h"
 
@@ -172,13 +173,15 @@ UDynamicMaterialModelEditorOnlyData::UDynamicMaterialModelEditorOnlyData()
 	, Domain(EMaterialDomain::MD_Surface)
 	, BlendMode(EBlendMode::BLEND_Opaque)
 	, ShadingModel(EDMMaterialShadingModel::DefaultLit)
-	, bHasPixelAnimation(false)
-	, bTwoSided(true)
-	, bOutputTranslucentVelocityEnabled(true)
-	, bNaniteTessellationEnabled(true)
+	, UsageFlags(0)
+	, GeneralFlags(static_cast<uint8>(EDMMaterialFlagsGeneral::TwoSided))
+	, LightingFlags(0)
+	, TranslucencyFlags(0)
+	, MotionFlags(static_cast<uint8>(EDMMaterialFlagsMotion::OutputTranslucentVelocity | EDMMaterialFlagsMotion::ResponsiveAA))
+	, ForwardRendererFlags(0)
+	, OpacityMaskClipValue(0.5f)
 	, DisplacementCenter(0.5f)
 	, DisplacementMagnitude(1.f)
-	, bResponsiveAAEnabled(true)
 	, bBuildRequested(false)
 {
 	BaseColor =           CreateDefaultSubobject<UDMMaterialPropertyBaseColor>(          "MaterialProperty_BaseColor");
@@ -282,13 +285,11 @@ void UDynamicMaterialModelEditorOnlyData::OnBlendModeChanged()
 	switch (BlendMode)
 	{
 		case EBlendMode::BLEND_Opaque:
-			SetHasPixelAnimation(false);
 			RemoveSlotForMaterialProperty(EDMMaterialPropertyType::Opacity);
 			RemoveSlotForMaterialProperty(EDMMaterialPropertyType::OpacityMask);
 			break;
 
 		case EBlendMode::BLEND_Masked:
-			SetHasPixelAnimation(false);
 			EnsureSwapSlotMaterialProperty(EDMMaterialPropertyType::Opacity, EDMMaterialPropertyType::OpacityMask);
 			break;
 
@@ -312,7 +313,7 @@ void UDynamicMaterialModelEditorOnlyData::OnMaterialFlagChanged()
 	RequestMaterialBuild();
 }
 
-void UDynamicMaterialModelEditorOnlyData::OnDisplacementSettingsChanged()
+void UDynamicMaterialModelEditorOnlyData::OnMaterialSettingsChanged()
 {
 	RequestMaterialBuild();
 }
@@ -403,20 +404,12 @@ void UDynamicMaterialModelEditorOnlyData::BuildMaterial(bool bInDirtyAssets)
 	Expressions.Empty();
 	MaterialModel->DynamicMaterial->MaterialDomain = Domain;
 	MaterialModel->DynamicMaterial->BlendMode = BlendMode;
-	MaterialModel->DynamicMaterial->bHasPixelAnimation = bHasPixelAnimation;
-	MaterialModel->DynamicMaterial->TwoSided = bTwoSided;
-	MaterialModel->DynamicMaterial->bOutputTranslucentVelocity = bOutputTranslucentVelocityEnabled;
-	MaterialModel->DynamicMaterial->bEnableTessellation = bNaniteTessellationEnabled;
-	MaterialModel->DynamicMaterial->bEnableResponsiveAA = bResponsiveAAEnabled;
-	MaterialModel->DynamicMaterial->DisplacementScaling.Magnitude = DisplacementMagnitude;
-	MaterialModel->DynamicMaterial->DisplacementScaling.Center = DisplacementCenter;
 
-	// Not setting this to true can cause the level associated with this material to dirty itself when it
-	// is used with Niagara. It doesn't negatively affect the material in any meaningful way.
-	MaterialModel->DynamicMaterial->SetUsageByFlag(MATUSAGE_NiagaraMeshParticles, true);
+	SetMaterialFlags(MaterialModel->DynamicMaterial);
 
-	// Need Nanite usage flag to be set to true if there is displacement/tessellation.
-	MaterialModel->DynamicMaterial->SetUsageByFlag(MATUSAGE_Nanite, bNaniteTessellationEnabled);
+	MaterialModel->DynamicMaterial->OpacityMaskClipValue = OpacityMaskClipValue;
+	MaterialModel->DynamicMaterial->DisplacementScaling.Magnitude = DisplacementMagnitude;
+	MaterialModel->DynamicMaterial->DisplacementScaling.Center = DisplacementCenter;
 
 	switch (ShadingModel)
 	{
@@ -598,6 +591,53 @@ void UDynamicMaterialModelEditorOnlyData::BuildMaterial(bool bInDirtyAssets)
 	OnMaterialBuiltDelegate.Broadcast(MaterialModel);
 }
 
+void UDynamicMaterialModelEditorOnlyData::SetMaterialFlags(UMaterial* InMaterial)
+{
+	// Match data type of EMaterialUsage (int)
+	for (int Flag = 0; Flag < EMaterialUsage::MATUSAGE_MAX; ++Flag)
+	{
+		const EMaterialUsage UsageFlag = static_cast<EMaterialUsage>(Flag);
+		const EDMMaterialFlagsUsage DMUsageFlag = static_cast<EDMMaterialFlagsUsage>(1 << Flag);
+
+		InMaterial->SetUsageByFlag(UsageFlag, HasUsageFlag(DMUsageFlag));
+	}
+
+	InMaterial->TwoSided = HasGeneralFlag(EDMMaterialFlagsGeneral::TwoSided);
+	InMaterial->bTangentSpaceNormal = HasGeneralFlag(EDMMaterialFlagsGeneral::TangentSpaceNormal);
+	InMaterial->bUsesDistortion = HasGeneralFlag(EDMMaterialFlagsGeneral::UsesDistortion);
+	InMaterial->bNormalCurvatureToRoughness = HasGeneralFlag(EDMMaterialFlagsGeneral::NormalCurvatureToRoughness);
+	InMaterial->bFullyRough = HasGeneralFlag(EDMMaterialFlagsGeneral::FullyRough);
+	InMaterial->DitheredLODTransition = HasGeneralFlag(EDMMaterialFlagsGeneral::DitheredLODTransition);
+	InMaterial->bEnableDisplacementFade = HasGeneralFlag(EDMMaterialFlagsGeneral::EnableDisplacementFade);
+	InMaterial->bGenerateSphericalParticleNormals = HasGeneralFlag(EDMMaterialFlagsGeneral::GenerateSphericalParticleNormals);
+
+	InMaterial->bCastDynamicShadowAsMasked = HasLightingFlag(EDMMaterialFlagsLighting::CastDynamicShadowAsMasked);
+	InMaterial->bContactShadows = HasLightingFlag(EDMMaterialFlagsLighting::ContactShadows);
+	InMaterial->bUseEmissiveForDynamicAreaLighting = HasLightingFlag(EDMMaterialFlagsLighting::UseEmissiveForDynamicAreaLighting);
+	InMaterial->bCompatibleWithLumenCardSharing = HasLightingFlag(EDMMaterialFlagsLighting::CompatibleWithLumenCardSharing);
+	InMaterial->bAllowTranslucentLocalLightShadow = HasLightingFlag(EDMMaterialFlagsLighting::AllowTranslucentLocalLightShadow);
+	InMaterial->bUseLightmapDirectionality = HasLightingFlag(EDMMaterialFlagsLighting::UseLightmapDirectionality);
+	InMaterial->bAllowNegativeEmissiveColor = HasLightingFlag(EDMMaterialFlagsLighting::AllowNegativeEmissiveColor);
+
+	InMaterial->bDisableDepthTest = HasTranslucencyFlag(EDMMaterialFlagsTranslucency::DisableDepthTest);
+	InMaterial->bScreenSpaceReflections = HasTranslucencyFlag(EDMMaterialFlagsTranslucency::ScreenSpaceReflections);
+	InMaterial->AllowTranslucentCustomDepthWrites = HasTranslucencyFlag(EDMMaterialFlagsTranslucency::AllowTranslucentCustomDepthWrites);
+	InMaterial->bAllowFrontLayerTranslucency = HasTranslucencyFlag(EDMMaterialFlagsTranslucency::AllowFrontLayerTranslucency);
+	InMaterial->bWriteOnlyAlpha = HasTranslucencyFlag(EDMMaterialFlagsTranslucency::WriteOnlyAlpha);
+	InMaterial->DitherOpacityMask = HasTranslucencyFlag(EDMMaterialFlagsTranslucency::DitherOpacityMask);
+
+	InMaterial->bHasPixelAnimation = HasMotionFlag(EDMMaterialFlagsMotion::PixelAnimation);
+	InMaterial->bEnableResponsiveAA = HasMotionFlag(EDMMaterialFlagsMotion::ResponsiveAA);
+	InMaterial->bOutputTranslucentVelocity = HasMotionFlag(EDMMaterialFlagsMotion::OutputTranslucentVelocity);
+
+	InMaterial->bEnableMobileSeparateTranslucency = HasForwardRendererFlag(EDMMaterialFlagsForwardRenderer::MobileSeparateTranslucency);
+	InMaterial->bUseAlphaToCoverage = HasForwardRendererFlag(EDMMaterialFlagsForwardRenderer::AlphaToCoverage);
+	InMaterial->bForwardRenderUsePreintegratedGFForSimpleIBL = HasForwardRendererFlag(EDMMaterialFlagsForwardRenderer::ForwardRenderUsePreintegratedGFForSimpleIBL);
+	InMaterial->bUseHQForwardReflections = HasForwardRendererFlag(EDMMaterialFlagsForwardRenderer::HighQualityForwardReflections);
+	InMaterial->bForwardBlendsSkyLightCubemaps = HasForwardRendererFlag(EDMMaterialFlagsForwardRenderer::ForwardBlendsSkyLightCubemaps);
+	InMaterial->bUsePlanarForwardReflections = HasForwardRendererFlag(EDMMaterialFlagsForwardRenderer::PlanarForwardReflections);
+}
+
 void UDynamicMaterialModelEditorOnlyData::RequestMaterialBuild(EDMBuildRequestType InRequestType)
 {
 	UDynamicMaterialEditorSettings* Settings = UDynamicMaterialEditorSettings::Get();
@@ -735,8 +775,12 @@ void UDynamicMaterialModelEditorOnlyData::SetChannelListPreset(FName InPresetNam
 
 	SetBlendMode(Preset->DefaultBlendMode);
 	SetShadingModel(Preset->DefaultShadingModel);
-	SetHasPixelAnimation(Preset->bDefaultAnimated);
-	SetIsTwoSided(Preset->bDefaultTwoSided);
+	UsageFlags = Preset->UsageFlag;
+	GeneralFlags = Preset->GeneralFlags;
+	LightingFlags = Preset->LightingFlags;
+	TranslucencyFlags = Preset->TranslucencyFlags;
+	MotionFlags = Preset->MotionFlags;
+	ForwardRendererFlags = Preset->ForwardRendererFlags;
 }
 
 const FMaterialStatistics& UDynamicMaterialModelEditorOnlyData::GetMaterialStats() const
@@ -914,89 +958,257 @@ void UDynamicMaterialModelEditorOnlyData::SetShadingModel(EDMMaterialShadingMode
 	OnShadingModelChanged();
 }
 
-bool UDynamicMaterialModelEditorOnlyData::GetHasPixelAnimation() const
+EDMMaterialFlagsUsage UDynamicMaterialModelEditorOnlyData::GetUsageFlags() const
 {
-	return bHasPixelAnimation;
+	return static_cast<EDMMaterialFlagsUsage>(UsageFlags);
 }
 
-void UDynamicMaterialModelEditorOnlyData::SetHasPixelAnimation(bool bInHasAnimation)
+bool UDynamicMaterialModelEditorOnlyData::HasUsageFlag(EDMMaterialFlagsUsage InFlag) const
 {
-	if (bHasPixelAnimation == bInHasAnimation)
+	return (UsageFlags & static_cast<int>(InFlag)) > 0;
+}
+
+bool UDynamicMaterialModelEditorOnlyData::SetUsageFlag(EDMMaterialFlagsUsage InFlag, bool bInSet)
+{
+	if (bInSet == HasUsageFlag(InFlag))
 	{
-		return;
+		return false;
 	}
 
-	bHasPixelAnimation = bInHasAnimation;
+	if (bInSet)
+	{
+		UsageFlags |= static_cast<int>(InFlag);
+	}
+	else
+	{
+		UsageFlags &= ~static_cast<int>(InFlag);
+	}
 
 	OnMaterialFlagChanged();
+
+	return true;
 }
 
-bool UDynamicMaterialModelEditorOnlyData::GetIsTwoSided() const
+EDMMaterialFlagsGeneral UDynamicMaterialModelEditorOnlyData::GetGeneralFlags() const
 {
-	return bTwoSided;
+	return static_cast<EDMMaterialFlagsGeneral>(GeneralFlags);
 }
 
-void UDynamicMaterialModelEditorOnlyData::SetIsTwoSided(bool bInEnabled)
+bool UDynamicMaterialModelEditorOnlyData::HasGeneralFlag(EDMMaterialFlagsGeneral InFlag) const
 {
-	if (bTwoSided == bInEnabled)
+	return (GeneralFlags & static_cast<uint8>(InFlag)) > 0;
+}
+
+bool UDynamicMaterialModelEditorOnlyData::SetGeneralFlag(EDMMaterialFlagsGeneral InFlag, bool bInSet)
+{
+	if (bInSet == HasGeneralFlag(InFlag))
 	{
-		return;
+		return false;
 	}
 
-	bTwoSided = bInEnabled;
+	if (bInSet)
+	{
+		GeneralFlags |= static_cast<uint8>(InFlag);
+	}
+	else
+	{
+		GeneralFlags &= ~static_cast<uint8>(InFlag);
+	}
 
 	OnMaterialFlagChanged();
+
+	return true;
 }
 
-bool UDynamicMaterialModelEditorOnlyData::IsOutputTranslucentVelocityEnabled() const
+EDMMaterialFlagsLighting UDynamicMaterialModelEditorOnlyData::GetLightingFlags() const
 {
-	return bOutputTranslucentVelocityEnabled;
+	return static_cast<EDMMaterialFlagsLighting>(LightingFlags);
 }
 
-void UDynamicMaterialModelEditorOnlyData::SetOutputTranslucentVelocityEnabled(bool bInEnabled)
+bool UDynamicMaterialModelEditorOnlyData::HasLightingFlag(EDMMaterialFlagsLighting InFlag) const
 {
-	if (bOutputTranslucentVelocityEnabled == bInEnabled)
+	return (LightingFlags & static_cast<uint8>(InFlag)) > 0;
+}
+
+bool UDynamicMaterialModelEditorOnlyData::SetLightingFlag(EDMMaterialFlagsLighting InFlag, bool bInSet)
+{
+	if (bInSet == HasLightingFlag(InFlag))
 	{
-		return;
+		return false;
 	}
 
-	bOutputTranslucentVelocityEnabled = bInEnabled;
+	if (bInSet)
+	{
+		LightingFlags |= static_cast<uint8>(InFlag);
+	}
+	else
+	{
+		LightingFlags &= ~static_cast<uint8>(InFlag);
+	}
 
 	OnMaterialFlagChanged();
+
+	return true;
 }
 
-bool UDynamicMaterialModelEditorOnlyData::IsNaniteTessellationEnabled() const
+EDMMaterialFlagsTranslucency UDynamicMaterialModelEditorOnlyData::GetTranslucencyFlags() const
 {
-	return bNaniteTessellationEnabled;
+	return static_cast<EDMMaterialFlagsTranslucency>(TranslucencyFlags);
 }
 
-void UDynamicMaterialModelEditorOnlyData::SetNaniteTessellationEnabled(bool bInEnabled)
+bool UDynamicMaterialModelEditorOnlyData::HasTranslucencyFlag(EDMMaterialFlagsTranslucency InFlag) const
+{
+	return (TranslucencyFlags & static_cast<uint8>(InFlag)) > 0;
+}
+
+bool UDynamicMaterialModelEditorOnlyData::SetTranslucencyFlag(EDMMaterialFlagsTranslucency InFlag, bool bInSet)
 {
-	if (bNaniteTessellationEnabled == bInEnabled)
+	if (bInSet == HasTranslucencyFlag(InFlag))
 	{
-		return;
+		return false;
 	}
 
-	bNaniteTessellationEnabled = bInEnabled;
+	if (bInSet)
+	{
+		TranslucencyFlags |= static_cast<uint8>(InFlag);
+	}
+	else
+	{
+		TranslucencyFlags &= ~static_cast<uint8>(InFlag);
+	}
 
 	OnMaterialFlagChanged();
+
+	return true;
 }
 
-bool UDynamicMaterialModelEditorOnlyData::IsResponsiveAAEnabled() const
+EDMMaterialFlagsMotion UDynamicMaterialModelEditorOnlyData::GetMotionFlags() const
 {
-	return bResponsiveAAEnabled;
+	return static_cast<EDMMaterialFlagsMotion>(MotionFlags);
 }
 
-void UDynamicMaterialModelEditorOnlyData::SetResponsiveAAEnabled(bool bInEnabled)
+bool UDynamicMaterialModelEditorOnlyData::HasMotionFlag(EDMMaterialFlagsMotion InFlag) const
+{
+	return (MotionFlags & static_cast<uint8>(InFlag)) > 0;
+}
+
+bool UDynamicMaterialModelEditorOnlyData::SetMotionFlag(EDMMaterialFlagsMotion InFlag, bool bInSet)
 {
-	if (bResponsiveAAEnabled == bInEnabled)
+	if (bInSet == HasMotionFlag(InFlag))
 	{
-		return;
+		return false;
 	}
 
-	bResponsiveAAEnabled = bInEnabled;
+	if (bInSet)
+	{
+		MotionFlags |= static_cast<uint8>(InFlag);
+	}
+	else
+	{
+		MotionFlags &= ~static_cast<uint8>(InFlag);
+	}
 
 	OnMaterialFlagChanged();
+
+	return true;
+}
+
+EDMMaterialFlagsForwardRenderer UDynamicMaterialModelEditorOnlyData::GetForwardRendererFlags() const
+{
+	return static_cast<EDMMaterialFlagsForwardRenderer>(ForwardRendererFlags);
+}
+
+bool UDynamicMaterialModelEditorOnlyData::HasForwardRendererFlag(EDMMaterialFlagsForwardRenderer InFlag) const
+{
+	return (ForwardRendererFlags & static_cast<uint8>(InFlag)) > 0;
+}
+
+bool UDynamicMaterialModelEditorOnlyData::SetForwardRendererFlag(EDMMaterialFlagsForwardRenderer InFlag, bool bInSet)
+{
+	if (bInSet == HasForwardRendererFlag(InFlag))
+	{
+		return false;
+	}
+
+	if (bInSet)
+	{
+		ForwardRendererFlags |= static_cast<uint8>(InFlag);
+	}
+	else
+	{
+		ForwardRendererFlags &= ~static_cast<uint8>(InFlag);
+	}
+
+	OnMaterialFlagChanged();
+
+	return true;
+}
+
+float UDynamicMaterialModelEditorOnlyData::GetOpacityMaskClipValue() const
+{
+	return OpacityMaskClipValue;
+}
+
+void UDynamicMaterialModelEditorOnlyData::SetOpacityMaskClipValue(float InClipValue)
+{
+	if (FMath::IsNearlyEqual(OpacityMaskClipValue, InClipValue))
+	{
+		return;
+	}
+
+	OpacityMaskClipValue = InClipValue;
+
+	OnMaterialSettingsChanged();
+}
+
+bool UDynamicMaterialModelEditorOnlyData::GetHasPixelAnimation() const
+{
+	return HasMotionFlag(EDMMaterialFlagsMotion::PixelAnimation);
+}
+
+void UDynamicMaterialModelEditorOnlyData::SetHasPixelAnimation(bool bInHasAnimation)
+{
+	SetMotionFlag(EDMMaterialFlagsMotion::PixelAnimation, bInHasAnimation);
+}
+
+bool UDynamicMaterialModelEditorOnlyData::GetIsTwoSided() const
+{
+	return HasGeneralFlag(EDMMaterialFlagsGeneral::TwoSided);
+}
+
+void UDynamicMaterialModelEditorOnlyData::SetIsTwoSided(bool bInEnabled)
+{
+	SetGeneralFlag(EDMMaterialFlagsGeneral::TwoSided, bInEnabled);
+}
+
+bool UDynamicMaterialModelEditorOnlyData::IsOutputTranslucentVelocityEnabled() const
+{
+	return HasMotionFlag(EDMMaterialFlagsMotion::OutputTranslucentVelocity);
+}
+
+void UDynamicMaterialModelEditorOnlyData::SetOutputTranslucentVelocityEnabled(bool bInEnabled)
+{
+	SetMotionFlag(EDMMaterialFlagsMotion::OutputTranslucentVelocity, bInEnabled);
+}
+
+bool UDynamicMaterialModelEditorOnlyData::IsNaniteTessellationEnabled() const
+{
+	return HasUsageFlag(EDMMaterialFlagsUsage::DMMFU_Nanite);
+}
+
+void UDynamicMaterialModelEditorOnlyData::SetNaniteTessellationEnabled(bool bInEnabled)
+{
+	SetUsageFlag(EDMMaterialFlagsUsage::DMMFU_Nanite, bInEnabled);
+}
+
+bool UDynamicMaterialModelEditorOnlyData::IsResponsiveAAEnabled() const
+{
+	return HasMotionFlag(EDMMaterialFlagsMotion::ResponsiveAA);
+}
+
+void UDynamicMaterialModelEditorOnlyData::SetResponsiveAAEnabled(bool bInEnabled)
+{
+	SetMotionFlag(EDMMaterialFlagsMotion::ResponsiveAA, bInEnabled);
 }
 
 float UDynamicMaterialModelEditorOnlyData::GetDisplacementCenter() const
@@ -1013,7 +1225,7 @@ void UDynamicMaterialModelEditorOnlyData::SetDisplacementCenter(float InCenter)
 
 	DisplacementCenter = InCenter;
 
-	OnDisplacementSettingsChanged();
+	OnMaterialSettingsChanged();
 }
 
 float UDynamicMaterialModelEditorOnlyData::GetDisplacementMagnitude() const
@@ -1030,7 +1242,7 @@ void UDynamicMaterialModelEditorOnlyData::SetDisplacementMagnitude(float InMagni
 
 	DisplacementMagnitude = InMagnitude;
 
-	OnDisplacementSettingsChanged();
+	OnMaterialSettingsChanged();
 }
 
 void UDynamicMaterialModelEditorOnlyData::OpenMaterialEditor() const
@@ -1462,18 +1674,39 @@ void UDynamicMaterialModelEditorOnlyData::PostEditChangeChainProperty(FPropertyC
 	{
 		OnShadingModelChanged();
 	}
-	else if (Property == GET_MEMBER_NAME_CHECKED(ThisClass, bHasPixelAnimation)
-		|| Property == GET_MEMBER_NAME_CHECKED(ThisClass, bTwoSided)
-		|| Property == GET_MEMBER_NAME_CHECKED(ThisClass, bOutputTranslucentVelocityEnabled)
-		|| Property == GET_MEMBER_NAME_CHECKED(ThisClass, bNaniteTessellationEnabled)
-		|| Property == GET_MEMBER_NAME_CHECKED(ThisClass, bResponsiveAAEnabled))
+	else if (Property == GET_MEMBER_NAME_CHECKED(ThisClass, UsageFlags)
+		|| Property == GET_MEMBER_NAME_CHECKED(ThisClass, GeneralFlags)
+		|| Property == GET_MEMBER_NAME_CHECKED(ThisClass, LightingFlags)
+		|| Property == GET_MEMBER_NAME_CHECKED(ThisClass, TranslucencyFlags)
+		|| Property == GET_MEMBER_NAME_CHECKED(ThisClass, MotionFlags)
+		|| Property == GET_MEMBER_NAME_CHECKED(ThisClass, ForwardRendererFlags))
 	{
 		OnMaterialFlagChanged();
 	}
-	else if (Property == GET_MEMBER_NAME_CHECKED(ThisClass, DisplacementCenter)
+	else if (Property == GET_MEMBER_NAME_CHECKED(ThisClass, OpacityMaskClipValue)
+		|| Property == GET_MEMBER_NAME_CHECKED(ThisClass, DisplacementCenter)
 		|| Property == GET_MEMBER_NAME_CHECKED(ThisClass, DisplacementMagnitude))
 	{
-		OnDisplacementSettingsChanged();
+		 OnMaterialSettingsChanged();
+	}
+}
+
+void UDynamicMaterialModelEditorOnlyData::Serialize(FArchive& InAr)
+{
+	InAr.UsingCustomVersion(FDynamicMaterialModelEditorOnlyDataVersion::GUID);
+
+	Super::Serialize(InAr);
+
+	if (!InAr.IsLoading())
+	{
+		return;
+	}
+
+	int32 EODVersion = InAr.CustomVer(FDynamicMaterialModelEditorOnlyDataVersion::GUID);
+
+	if (EODVersion < FDynamicMaterialModelEditorOnlyDataVersion::FlagsAdded)
+	{
+		ConvertDeprecatedVarsToFlags();
 	}
 }
 
@@ -1567,6 +1800,38 @@ void UDynamicMaterialModelEditorOnlyData::ReinitComponents()
 	}
 }
 
+void UDynamicMaterialModelEditorOnlyData::ConvertDeprecatedVarsToFlags()
+{
+	PRAGMA_DISABLE_DEPRECATION_WARNINGS
+
+	if (bNaniteTessellationEnabled_DEPRECATED)
+	{
+		UsageFlags |= static_cast<int>(EDMMaterialFlagsUsage::DMMFU_Nanite);
+	}
+
+	if (bTwoSided_DEPRECATED)
+	{
+		GeneralFlags |= static_cast<uint8>(EDMMaterialFlagsGeneral::TwoSided);
+	}
+
+	if (bHasPixelAnimation_DEPRECATED)
+	{
+		MotionFlags |= static_cast<uint8>(EDMMaterialFlagsMotion::PixelAnimation);
+	}
+
+	if (bOutputTranslucentVelocityEnabled_DEPRECATED)
+	{
+		MotionFlags |= static_cast<uint8>(EDMMaterialFlagsMotion::OutputTranslucentVelocity);
+	}
+
+	if (bResponsiveAAEnabled_DEPRECATED)
+	{
+		MotionFlags |= static_cast<uint8>(EDMMaterialFlagsMotion::ResponsiveAA);
+	}
+
+	PRAGMA_ENABLE_DEPRECATION_WARNINGS
+}
+
 void UDynamicMaterialModelEditorOnlyData::PostEditorDuplicate()
 {
 	if (GUndo)
diff --git a/Engine/Plugins/VirtualProduction/DynamicMaterial/Source/DynamicMaterialEditor/Private/Model/DynamicMaterialModelEditorOnlyDataVersion.h b/Engine/Plugins/VirtualProduction/DynamicMaterial/Source/DynamicMaterialEditor/Private/Model/DynamicMaterialModelEditorOnlyDataVersion.h
--- a/Engine/Plugins/VirtualProduction/DynamicMaterial/Source/DynamicMaterialEditor/Private/Model/DynamicMaterialModelEditorOnlyDataVersion.h
+++ b/Engine/Plugins/VirtualProduction/DynamicMaterial/Source/DynamicMaterialEditor/Private/Model/DynamicMaterialModelEditorOnlyDataVersion.h
@@ -18,6 +18,9 @@ public:
 		/** When the global variables were added and a couple were renamed. */
 		GlobalValueRename,
 
+		/** When usage and material flags were added. */
+		FlagsAdded,
+
 		VersionPlusOne,
 		LatestVersion = VersionPlusOne - 1
 	};
diff --git a/Engine/Plugins/VirtualProduction/DynamicMaterial/Source/DynamicMaterialEditor/Private/UI/PropertyGenerators/DMMaterialModelPropertyRowGenerator.cpp b/Engine/Plugins/VirtualProduction/DynamicMaterial/Source/DynamicMaterialEditor/Private/UI/PropertyGenerators/DMMaterialModelPropertyRowGenerator.cpp
--- a/Engine/Plugins/VirtualProduction/DynamicMaterial/Source/DynamicMaterialEditor/Private/UI/PropertyGenerators/DMMaterialModelPropertyRowGenerator.cpp
+++ b/Engine/Plugins/VirtualProduction/DynamicMaterial/Source/DynamicMaterialEditor/Private/UI/PropertyGenerators/DMMaterialModelPropertyRowGenerator.cpp
@@ -52,19 +52,25 @@ void FDMMaterialModelPropertyRowGenerator::AddMaterialModelProperties(FDMCompone
 			GET_MEMBER_NAME_CHECKED(UDynamicMaterialModelEditorOnlyData, ShadingModel));
 
 		AddVariable(InParams, EditorOnlyData,
-			GET_MEMBER_NAME_CHECKED(UDynamicMaterialModelEditorOnlyData, bHasPixelAnimation));
+			GET_MEMBER_NAME_CHECKED(UDynamicMaterialModelEditorOnlyData, UsageFlags));
 
 		AddVariable(InParams, EditorOnlyData,
-			GET_MEMBER_NAME_CHECKED(UDynamicMaterialModelEditorOnlyData, bTwoSided));
+			GET_MEMBER_NAME_CHECKED(UDynamicMaterialModelEditorOnlyData, GeneralFlags));
 
 		AddVariable(InParams, EditorOnlyData,
-			GET_MEMBER_NAME_CHECKED(UDynamicMaterialModelEditorOnlyData, bResponsiveAAEnabled));
+			GET_MEMBER_NAME_CHECKED(UDynamicMaterialModelEditorOnlyData, LightingFlags));
 
 		AddVariable(InParams, EditorOnlyData,
-			GET_MEMBER_NAME_CHECKED(UDynamicMaterialModelEditorOnlyData, bOutputTranslucentVelocityEnabled));
+			GET_MEMBER_NAME_CHECKED(UDynamicMaterialModelEditorOnlyData, TranslucencyFlags));
 
 		AddVariable(InParams, EditorOnlyData,
-			GET_MEMBER_NAME_CHECKED(UDynamicMaterialModelEditorOnlyData, bNaniteTessellationEnabled));
+			GET_MEMBER_NAME_CHECKED(UDynamicMaterialModelEditorOnlyData, OpacityMaskClipValue));
+
+		AddVariable(InParams, EditorOnlyData,
+			GET_MEMBER_NAME_CHECKED(UDynamicMaterialModelEditorOnlyData, MotionFlags));
+
+		AddVariable(InParams, EditorOnlyData,
+			GET_MEMBER_NAME_CHECKED(UDynamicMaterialModelEditorOnlyData, ForwardRendererFlags));
 
 		AddVariable(InParams, EditorOnlyData,
 			GET_MEMBER_NAME_CHECKED(UDynamicMaterialModelEditorOnlyData, DisplacementMagnitude));
diff --git a/Engine/Plugins/VirtualProduction/DynamicMaterial/Source/DynamicMaterialEditor/Public/Model/DynamicMaterialModelEditorOnlyData.h b/Engine/Plugins/VirtualProduction/DynamicMaterial/Source/DynamicMaterialEditor/Public/Model/DynamicMaterialModelEditorOnlyData.h
--- a/Engine/Plugins/VirtualProduction/DynamicMaterial/Source/DynamicMaterialEditor/Public/Model/DynamicMaterialModelEditorOnlyData.h
+++ b/Engine/Plugins/VirtualProduction/DynamicMaterial/Source/DynamicMaterialEditor/Public/Model/DynamicMaterialModelEditorOnlyData.h
@@ -43,6 +43,128 @@ enum class EDMState : uint8
 	Building
 };
 
+/** Convert engine linear enum to bitflag enum. */
+UENUM(BlueprintType, meta = (BitFlags, UseEnumValuesAsMaskValuesInEditor = "true"))
+enum EDMMaterialFlagsUsage : int32
+{
+	DMMFU_None                   = 0                                                    UMETA(Hidden),
+	/** Indicates that the material and its instances can be used on static meshes. */
+	DMMFU_StaticMesh             = 1 << EMaterialUsage::MATUSAGE_StaticMesh             UMETA(DisplayName = "Static Mesh"),
+	/** Indicates that the material and its instances can be used with skeletal meshes. */
+	DMMFU_SkeletalMesh           = 1 << EMaterialUsage::MATUSAGE_SkeletalMesh           UMETA(DisplayName = "Skeletal Mesh"),
+	/** Indicates that the material and its instances can be used with spline meshes. */
+	DMMFU_SplineMesh             = 1 << EMaterialUsage::MATUSAGE_SplineMesh             UMETA(DisplayName = "Spline Mesh"),
+	/** Indicates that the material and its instances can be used with instanced static meshes. */
+	DMMFU_InstancedStaticMeshes  = 1 << EMaterialUsage::MATUSAGE_InstancedStaticMeshes  UMETA(DisplayName = "Instanced Static Mesh"),
+	/** Indicates that the material and its instances can be used with particle sprites. */
+	DMMFU_ParticleSprites        = 1 << EMaterialUsage::MATUSAGE_ParticleSprites        UMETA(DisplayName = "Particle Sprites"),
+	/** Indicates that the material and its instances can be used with beam trails. */
+	DMMFU_BeamTrails             = 1 << EMaterialUsage::MATUSAGE_BeamTrails             UMETA(DisplayName = "Beam Trails"),
+	/** Indicates that the material and its instances can be used with mesh particles. */
+	DMMFU_MeshParticles          = 1 << EMaterialUsage::MATUSAGE_MeshParticles          UMETA(DisplayName = "Mesh Particles"),
+	/** Indicates that the material and its instances can be used with static lighting. */
+	DMMFU_StaticLighting         = 1 << EMaterialUsage::MATUSAGE_StaticLighting         UMETA(DisplayName = "Static Lighting"),
+	/** Indicates that the material and its instances can be used with morph targets. */
+	DMMFU_MorphTargets           = 1 << EMaterialUsage::MATUSAGE_MorphTargets           UMETA(DisplayName = "Morph Targets"),
+	/** Indicates that the material and its instances can be use with geometry collections. */
+	DMMFU_GeometryCollections    = 1 << EMaterialUsage::MATUSAGE_GeometryCollections    UMETA(DisplayName = "Geometry Collections"),
+	/** Indicates that the material and its instances can be use with geometry caches. */
+	DMMFU_GeometryCache          = 1 << EMaterialUsage::MATUSAGE_GeometryCache          UMETA(DisplayName = "Geometry Cache"),
+	/** Indicates that the material and its instances can be used with clothing. */
+	DMMFU_Clothing               = 1 << EMaterialUsage::MATUSAGE_Clothing               UMETA(DisplayName = "Clothing"),
+	/** Indicates that the material and its instances can be used with Niagara sprites. */
+	DMMFU_NiagaraSprites         = 1 << EMaterialUsage::MATUSAGE_NiagaraSprites         UMETA(DisplayName = "Niagara Sprites"),
+	/** Indicates that the material and its instances can be used with Niagara ribbons. */
+	DMMFU_NiagaraRibbons         = 1 << EMaterialUsage::MATUSAGE_NiagaraRibbons         UMETA(DisplayName = "Niagara Ribbons"),
+	/** Indicates that the material and its instances can be used with Niagara meshes. */
+	DMMFU_NiagaraMeshParticles   = 1 << EMaterialUsage::MATUSAGE_NiagaraMeshParticles   UMETA(DisplayName = "Niagara Mesh Particles"),
+	/** Indicates that the material and its instances can be use with water. */
+	DMMFU_Water                  = 1 << EMaterialUsage::MATUSAGE_Water                  UMETA(DisplayName = "Water"),
+	/** Indicates that the material and its instances can be use with hair strands. */
+	DMMFU_HairStrands            = 1 << EMaterialUsage::MATUSAGE_HairStrands            UMETA(DisplayName = "HairStrands"),
+	/** Indicates that the material and its instances can be use with LiDAR Point Clouds. */
+	DMMFU_LidarPointCloud        = 1 << EMaterialUsage::MATUSAGE_LidarPointCloud        UMETA(DisplayName = "Lidar Point Cloud"),
+	/** Indicates that the material and its instances can be used with Virtual Heightfield Mesh. (Deprecated). */
+	DMMFU_VirtualHeightfieldMesh = 1 << EMaterialUsage::MATUSAGE_VirtualHeightfieldMesh UMETA(DisplayName = "Virtual Heightfield Mesh"),
+	/** Indicates that the material and its instances can be used with Nanite meshes. */
+	DMMFU_Nanite                 = 1 << EMaterialUsage::MATUSAGE_Nanite                 UMETA(DisplayName = "Nanite"),
+	/** Indicates that the material and its instances can be used with Nanite voxel meshes. */
+	DMMFU_Voxels                 = 1 << EMaterialUsage::MATUSAGE_Voxels                 UMETA(DisplayName = "Voxels"),
+	/** Indicates that the material and its instances with volumetric cloud. Without that flag, it can only be used on volumetric fog. */
+	DMMFU_VolumetricCloud        = 1 << EMaterialUsage::MATUSAGE_VolumetricCloud        UMETA(DisplayName = "Volumetric Clouds"),
+	/** Indicates that the material and its instances with heterogeneous volumes. Without that flag, it can only be used on volumetric fog. */
+	DMMFU_HeterogeneousVolumes   = 1 << EMaterialUsage::MATUSAGE_HeterogeneousVolumes   UMETA(DisplayName = "Hetergeneous Volumes"),
+	/** Indicates that the material and its instances can be used with editor compositing. */
+	DMMFU_EditorCompositing      = 1 << EMaterialUsage::MATUSAGE_EditorCompositing      UMETA(DisplayName = "Editor Compositing"),
+	/** Indicates that the material and its instances can be used with neural network engine. */
+	DMMFU_NeuralNetworks         = 1 << EMaterialUsage::MATUSAGE_NeuralNetworks         UMETA(DisplayName = "Neural Networks")
+};
+
+UENUM(BlueprintType, meta = (BitFlags, UseEnumValuesAsMaskValuesInEditor = "true"))
+enum class EDMMaterialFlagsGeneral : uint8
+{
+	None                             = 0      UMETA(HIdden),
+	TwoSided                         = 1 << 0 UMETA(ToolTip = "Indicates that the material should be rendered without backface culling and the normal should be flipped for backfaces."),
+	TangentSpaceNormal               = 1 << 1 UMETA(ToolTip = "Whether the material takes a tangent space normal or a world space normal as input. (TangentSpace requires extra instructions but is often more convenient)."),
+	UsesDistortion                   = 1 << 2 UMETA(ToolTip = "Indicates that the material and its instances can be used with distortion. This will result in the shaders required to support distortion being compiled which will increase shader compile time and memory usage."),
+	NormalCurvatureToRoughness       = 1 << 3 UMETA(ToolTip = "Reduce roughness based on screen space normal changes."),
+	FullyRough                       = 1 << 4 UMETA(ToolTip = "Forces the material to be completely rough. Saves a number of instructions and one sampler."),
+	DitheredLODTransition            = 1 << 5 UMETA(ToolTip = "Whether meshes rendered with the material should support dithered LOD transitions."),
+	EnableDisplacementFade           = 1 << 6 UMETA(ToolTip = "Enables fading out and disabling of dynamic displacement in the distance, as displacement becomes unnoticeable."),
+	GenerateSphericalParticleNormals = 1 << 7 UMETA(ToolTip = "Whether to generate spherical normals for particles that use this material.")
+};
+ENUM_CLASS_FLAGS(EDMMaterialFlagsGeneral)
+
+UENUM(BlueprintType, meta = (BitFlags, UseEnumValuesAsMaskValuesInEditor = "true"))
+enum class EDMMaterialFlagsLighting : uint8
+{
+	None                              = 0      UMETA(HIdden),
+	CastDynamicShadowAsMasked         = 1 << 0 UMETA(ToolTip = "Whether the material should cast shadows as masked even though it has a translucent blend mode."),
+	ContactShadows                    = 1 << 1 UMETA(ToolTip = "Contact shadows on translucency."),
+	UseEmissiveForDynamicAreaLighting = 1 << 2 UMETA(ToolTip = "Whether the material's emissive colour is injected into the LightPropagationVolume."),
+	CompatibleWithLumenCardSharing    = 1 << 3 UMETA(ToolTip = "Whether to allow to share Lumen Cards between different instances even when material uses world position or per instance data, which may change material look per instance. All materials on a component needs this flag set for sharing to work."),
+	AllowTranslucentLocalLightShadow  = 1 << 4 UMETA(ToolTip = "Allows a translucent material to receive local light shadows."),
+	UseLightmapDirectionality         = 1 << 5 UMETA(ToolTip = "Use lightmap directionality and per pixel normals. If disabled, lighting from lightmaps will be flat but cheaper."),
+	AllowNegativeEmissiveColor        = 1 << 6 UMETA(ToolTip = "Whether the material should allow outputting negative emissive color values.  Only allowed on unlit materials.")
+};
+ENUM_CLASS_FLAGS(EDMMaterialFlagsLighting)
+
+UENUM(BlueprintType, meta = (BitFlags, UseEnumValuesAsMaskValuesInEditor = "true"))
+enum class EDMMaterialFlagsTranslucency : uint8
+{
+	None                              = 0      UMETA(HIdden),
+	DisableDepthTest                  = 1 << 0 UMETA(ToolTip = "Whether to draw on top of opaque pixels even if behind them. This only has meaning for translucency."),
+	ScreenSpaceReflections            = 1 << 1 UMETA(ToolTip = "SSR on translucency."),
+	AllowTranslucentCustomDepthWrites = 1 << 2 UMETA(ToolTip = "Allows a translucent material to be used with custom depth writing by compiling additional shaders."),
+	AllowFrontLayerTranslucency       = 1 << 3 UMETA(ToolTip = "Allows a translucent material to be used with Front Layer Translucency by compiling additional shaders. Useful for controlling what should be included in Front Layer Translucency."),
+	WriteOnlyAlpha                    = 1 << 4 UMETA(ToolTip = "Whether the transluency pass should write its alpha, and only the alpha, into the framebuffer."),
+	DitherOpacityMask                 = 1 << 5 UMETA(ToolTip = "Dither opacity mask. When combined with Temporal AA this can be used as a form of limited translucency which supports all lighting features.")
+};
+ENUM_CLASS_FLAGS(EDMMaterialFlagsTranslucency)
+
+UENUM(BlueprintType, meta = (BitFlags, UseEnumValuesAsMaskValuesInEditor = "true"))
+enum class EDMMaterialFlagsMotion : uint8
+{
+	None                      = 0      UMETA(HIdden),
+	PixelAnimation            = 1 << 0 UMETA(ToolTip = "Whether the opaque material has any pixel animations happening, that isn't included in the geometric velocities. This allows to disable renderer's heuristics that assumes animation is fully described with motion vector, such as TSR's anti-flickering heuristic."),
+	ResponsiveAA              = 1 << 1 UMETA(ToolTip = "Indicates that the material should be rendered using responsive anti-aliasing. Improves sharpness of small moving particles such as sparks. Only use for small moving features because it will cause aliasing of the background."),
+	OutputTranslucentVelocity = 1 << 2 UMETA(ToolTip = "Whether ranslucent materials will output motion vectors and write to depth buffer in velocity pass.")
+};
+ENUM_CLASS_FLAGS(EDMMaterialFlagsMotion)
+
+UENUM(BlueprintType, meta = (BitFlags, UseEnumValuesAsMaskValuesInEditor = "true"))
+enum class EDMMaterialFlagsForwardRenderer : uint8
+{
+	None                                        = 0      UMETA(HIdden),
+	MobileSeparateTranslucency                  = 1 << 0 UMETA(ToolTip = "Indicates that the translucent material should not be affected by bloom or DOF. (Note: Depth testing is not available)."),
+	AlphaToCoverage                             = 1 << 1 UMETA(ToolTip = "Use alpha to coverage for masked material on mobile, make sure MSAA is enabled as well."),
+	ForwardRenderUsePreintegratedGFForSimpleIBL = 1 << 2 UMETA(ToolTip = "Use preintegrated GF lut for simple IBL, but will use one more sampler."),
+	HighQualityForwardReflections               = 1 << 3 UMETA(ToolTip = "Enables multiple parallax-corrected reflection captures that blend together."),
+	ForwardBlendsSkyLightCubemaps               = 1 << 4 UMETA(ToolTip = "Enables blending of sky light cubemap textures. When disabled, the secondary cubemap is only visible when the blend factor is 1."),
+	PlanarForwardReflections                    = 1 << 5 UMETA(ToolTip = "Enabling this setting reduces the number of samplers available to the material as one more sampler will be used for the planar reflection.")
+};
+ENUM_CLASS_FLAGS(EDMMaterialFlagsForwardRenderer)
+
 UCLASS(MinimalAPI, BlueprintType, EditInlineNew, DefaultToInstanced, ClassGroup = "Material Designer")
 class UDynamicMaterialModelEditorOnlyData : public UObject, public IDynamicMaterialModelEditorOnlyDataInterface, public FNotifyHook, public IDMBuildable
 {
@@ -120,34 +242,64 @@ public:
 	DYNAMICMATERIALEDITOR_API void SetShadingModel(EDMMaterialShadingModel InShadingModel);
 
 	UFUNCTION(BlueprintPure, Category = "Material Designer")
-	DYNAMICMATERIALEDITOR_API bool GetHasPixelAnimation() const;
+	DYNAMICMATERIALEDITOR_API EDMMaterialFlagsUsage GetUsageFlags() const;
+
+	UFUNCTION(BlueprintPure, Category = "Material Designer")
+	DYNAMICMATERIALEDITOR_API bool HasUsageFlag(EDMMaterialFlagsUsage InFlag) const;
 
 	UFUNCTION(BlueprintCallable, Category = "Material Designer")
-	DYNAMICMATERIALEDITOR_API void SetHasPixelAnimation(bool bInHasAnimation);
+	DYNAMICMATERIALEDITOR_API bool SetUsageFlag(EDMMaterialFlagsUsage InFlag, bool bInSet);
 
 	UFUNCTION(BlueprintPure, Category = "Material Designer")
-	DYNAMICMATERIALEDITOR_API bool GetIsTwoSided() const;
+	DYNAMICMATERIALEDITOR_API EDMMaterialFlagsGeneral GetGeneralFlags() const;
+
+	UFUNCTION(BlueprintPure, Category = "Material Designer")
+	DYNAMICMATERIALEDITOR_API bool HasGeneralFlag(EDMMaterialFlagsGeneral InFlag) const;
 
 	UFUNCTION(BlueprintCallable, Category = "Material Designer")
-	DYNAMICMATERIALEDITOR_API void SetIsTwoSided(bool bInEnabled);
+	DYNAMICMATERIALEDITOR_API bool SetGeneralFlag(EDMMaterialFlagsGeneral InFlag, bool bInSet);
 
 	UFUNCTION(BlueprintPure, Category = "Material Designer")
-	DYNAMICMATERIALEDITOR_API bool IsOutputTranslucentVelocityEnabled() const;
+	DYNAMICMATERIALEDITOR_API EDMMaterialFlagsLighting GetLightingFlags() const;
+
+	UFUNCTION(BlueprintPure, Category = "Material Designer")
+	DYNAMICMATERIALEDITOR_API bool HasLightingFlag(EDMMaterialFlagsLighting InFlag) const;
 
 	UFUNCTION(BlueprintCallable, Category = "Material Designer")
-	DYNAMICMATERIALEDITOR_API void SetOutputTranslucentVelocityEnabled(bool bInEnabled);
+	DYNAMICMATERIALEDITOR_API bool SetLightingFlag(EDMMaterialFlagsLighting InFlag, bool bInSet);
 
 	UFUNCTION(BlueprintPure, Category = "Material Designer")
-	DYNAMICMATERIALEDITOR_API bool IsNaniteTessellationEnabled() const;
+	DYNAMICMATERIALEDITOR_API EDMMaterialFlagsTranslucency GetTranslucencyFlags() const;
+
+	UFUNCTION(BlueprintPure, Category = "Material Designer")
+	DYNAMICMATERIALEDITOR_API bool HasTranslucencyFlag(EDMMaterialFlagsTranslucency InFlag) const;
 
 	UFUNCTION(BlueprintCallable, Category = "Material Designer")
-	DYNAMICMATERIALEDITOR_API void SetNaniteTessellationEnabled(bool bInEnabled);
+	DYNAMICMATERIALEDITOR_API bool SetTranslucencyFlag(EDMMaterialFlagsTranslucency InFlag, bool bInSet);
 
 	UFUNCTION(BlueprintPure, Category = "Material Designer")
-	DYNAMICMATERIALEDITOR_API bool IsResponsiveAAEnabled() const;
+	DYNAMICMATERIALEDITOR_API EDMMaterialFlagsMotion GetMotionFlags() const;
+
+	UFUNCTION(BlueprintPure, Category = "Material Designer")
+	DYNAMICMATERIALEDITOR_API bool HasMotionFlag(EDMMaterialFlagsMotion InFlag) const;
 
 	UFUNCTION(BlueprintCallable, Category = "Material Designer")
-	DYNAMICMATERIALEDITOR_API void SetResponsiveAAEnabled(bool bInEnabled);
+	DYNAMICMATERIALEDITOR_API bool SetMotionFlag(EDMMaterialFlagsMotion InFlag, bool bInSet);
+
+	UFUNCTION(BlueprintPure, Category = "Material Designer")
+	DYNAMICMATERIALEDITOR_API EDMMaterialFlagsForwardRenderer GetForwardRendererFlags() const;
+
+	UFUNCTION(BlueprintPure, Category = "Material Designer")
+	DYNAMICMATERIALEDITOR_API bool HasForwardRendererFlag(EDMMaterialFlagsForwardRenderer InFlag) const;
+
+	UFUNCTION(BlueprintCallable, Category = "Material Designer")
+	DYNAMICMATERIALEDITOR_API bool SetForwardRendererFlag(EDMMaterialFlagsForwardRenderer InFlag, bool bInSet);
+
+	UFUNCTION(BlueprintPure, Category = "Material Designer")
+	DYNAMICMATERIALEDITOR_API float GetOpacityMaskClipValue() const;
+
+	UFUNCTION(BlueprintCallable, Category = "Material Designer")
+	DYNAMICMATERIALEDITOR_API void SetOpacityMaskClipValue(float InClipValue);
 
 	UFUNCTION(BlueprintPure, Category = "Material Designer")
 	DYNAMICMATERIALEDITOR_API float GetDisplacementCenter() const;
@@ -241,6 +393,7 @@ public:
 	DYNAMICMATERIALEDITOR_API virtual void PostEditImport() override;
 	DYNAMICMATERIALEDITOR_API virtual void PostDuplicate(bool bInDuplicateForPIE) override;
 	DYNAMICMATERIALEDITOR_API virtual void PostEditChangeChainProperty(FPropertyChangedChainEvent& InPropertyChangedEvent) override;
+	DYNAMICMATERIALEDITOR_API virtual void Serialize(FArchive& InAr) override;
 	//~ End UObject
 
 	//~ Begin IDMBuildable
@@ -259,6 +412,50 @@ public:
 	DYNAMICMATERIALEDITOR_API virtual UDMMaterialComponent* GetSubComponentByPath(FDMComponentPath& InPath, const FDMComponentPathSegment& InPathSegment) const override;
 	//~ End IDynamicMaterialModelEditorOnlyDataInterface
 
+	/*
+	 * Deprecated
+	 */
+
+	UE_DEPRECATED(5.8, "Use HasMotionFlag().")
+	UFUNCTION(meta = (DeprecatedFunction = "Note", DeprecationMessage = "Use HasMotionFlag()."))
+	DYNAMICMATERIALEDITOR_API bool GetHasPixelAnimation() const;
+
+	UE_DEPRECATED(5.8, "Use SetMotionFlag().")
+	UFUNCTION(meta = (DeprecatedFunction = "Note", DeprecationMessage = "Use SetMotionFlag()."))
+	DYNAMICMATERIALEDITOR_API void SetHasPixelAnimation(bool bInHasAnimation);
+
+	UE_DEPRECATED(5.8, "Use HasGeneralFlag().")
+	UFUNCTION(meta = (DeprecatedFunction = "Note", DeprecationMessage = "Use HasGeneralFlag()."))
+	DYNAMICMATERIALEDITOR_API bool GetIsTwoSided() const;
+
+	UE_DEPRECATED(5.8, "Use SetGeneralFlag().")
+	UFUNCTION(meta = (DeprecatedFunction = "Note", DeprecationMessage = "Use SetGeneralFlag()."))
+	DYNAMICMATERIALEDITOR_API void SetIsTwoSided(bool bInEnabled);
+
+	UE_DEPRECATED(5.8, "Use HasMotionFlag().")
+	UFUNCTION(meta = (DeprecatedFunction = "Note", DeprecationMessage = "Use HasMotionFlag()."))
+	DYNAMICMATERIALEDITOR_API bool IsOutputTranslucentVelocityEnabled() const;
+
+	UE_DEPRECATED(5.8, "Use SetMotionFlag().")
+	UFUNCTION(meta = (DeprecatedFunction = "Note", DeprecationMessage = "Use SetMotionFlag()."))
+	DYNAMICMATERIALEDITOR_API void SetOutputTranslucentVelocityEnabled(bool bInEnabled);
+
+	UE_DEPRECATED(5.8, "Use HasUsageFlag().")
+	UFUNCTION(meta = (DeprecatedFunction = "Note", DeprecationMessage = "Use HasUsageFlag()."))
+	DYNAMICMATERIALEDITOR_API bool IsNaniteTessellationEnabled() const;
+
+	UE_DEPRECATED(5.8, "Use SetUsageFlag().")
+	UFUNCTION(meta = (DeprecatedFunction = "Note", DeprecationMessage = "Use SetUsageFlag()."))
+	DYNAMICMATERIALEDITOR_API void SetNaniteTessellationEnabled(bool bInEnabled);
+
+	UE_DEPRECATED(5.8, "Use HasMotionFlag().")
+	UFUNCTION(meta = (DeprecatedFunction = "Note", DeprecationMessage = "Use HasMotionFlag()."))
+	DYNAMICMATERIALEDITOR_API bool IsResponsiveAAEnabled() const;
+
+	UE_DEPRECATED(5.8, "Use SetMotionFlag().")
+	UFUNCTION(meta = (DeprecatedFunction = "Note", DeprecationMessage = "Use SetMotionFlag()."))
+	DYNAMICMATERIALEDITOR_API void SetResponsiveAAEnabled(bool bInEnabled);
+
 protected:
 	UPROPERTY(VisibleInstanceOnly, BlueprintReadOnly, Category = "Material Designer")
 	TObjectPtr<UDynamicMaterialModel> MaterialModel;
@@ -278,44 +475,48 @@ protected:
 		meta = (NotKeyframeable, ValidEnumValues = "Unlit,DefaultLit"))
 	EDMMaterialShadingModel ShadingModel;
 
-	/**
-	 * Whether the opaque material has any pixel animations happening, that isn't included in the geometric velocities.
-	 * This allows to disable renderer's heuristics that assumes animation is fully described with motion vector, such as TSR's anti-flickering heuristic.
-	 */
-	UPROPERTY(EditInstanceOnly, BlueprintReadOnly, Category = "Material Designer", meta = (NotKeyframeable))
-	bool bHasPixelAnimation;
+	UPROPERTY(EditInstanceOnly, BlueprintReadOnly, Category = "Material Designer|Usage",
+		meta = (Bitmask, BitmaskEnum = "/Script/DynamicMaterialEditor.EDMMaterialFlagsUsage"))
+	int32 UsageFlags;
 
-	/** Indicates that the material should be rendered without backface culling and the normal should be flipped for backfaces. */
-	UPROPERTY(EditInstanceOnly, BlueprintReadOnly, Category = "Material Designer", meta = (NotKeyframeable))
-	bool bTwoSided;
+	UPROPERTY(EditInstanceOnly, BlueprintReadOnly, Category = "Material Designer|General",
+		meta = (Bitmask, BitmaskEnum = "/Script/DynamicMaterialEditor.EDMMaterialFlagsGeneral"))
+	uint8 GeneralFlags;
 
-	/** When true, translucent materials will output motion vectors and write to depth buffer in velocity pass. */
-	UPROPERTY(EditInstanceOnly, BlueprintReadOnly, Category = "Material Designer", 
-		meta = (NotKeyframeable, DisplayName = "Output Translucent Velocity"))
-	bool bOutputTranslucentVelocityEnabled;
+	UPROPERTY(EditInstanceOnly, BlueprintReadOnly, Category = "Material Designer|Lighting",
+		meta = (Bitmask, BitmaskEnum = "/Script/DynamicMaterialEditor.EDMMaterialFlagsLighting",
+			EditCondition = "ShadingModel == EDMMaterialShadingModel::DefaultLit"))
+	uint8 LightingFlags;
 
-	/** Whether tessellation is enabled on the material. NOTE: Required for displacement to work. */
-	UPROPERTY(EditInstanceOnly, BlueprintReadOnly, Category = "Material Designer", 
-		meta = (NotKeyframeable, DisplayName = "Nanite Tessellation"))
-	bool bNaniteTessellationEnabled;
+	UPROPERTY(EditInstanceOnly, BlueprintReadOnly, Category = "Material Designer|Translucency",
+		meta = (Bitmask, BitmaskEnum = "/Script/DynamicMaterialEditor.EDMMaterialFlagsTranslucency",
+			EditCondition = "BlendMode == EBlendMode::BLEND_Masked || BlendMode == EBlendMode::BLEND_Translucent"))
+	uint8 TranslucencyFlags;
+
+	UPROPERTY(EditInstanceOnly, BlueprintReadOnly, Category = "Material Designer|Motion",
+		meta = (Bitmask, BitmaskEnum = "/Script/DynamicMaterialEditor.EDMMaterialFlagsMotion"))
+	uint8 MotionFlags;
+
+	UPROPERTY(EditInstanceOnly, BlueprintReadOnly, Category = "Material Designer|Forward Renderer",
+		meta = (Bitmask, BitmaskEnum = "/Script/DynamicMaterialEditor.EDMMaterialFlagsForwardRenderer"))
+	uint8 ForwardRendererFlags;
+
+	UPROPERTY(EditInstanceOnly, BlueprintReadOnly, Category = "Material Designer|Translucency",
+		meta = (NotKeyframeable, UIMin = 0, UIMax = 1, ClampMin = 0, ClampMax = 1,
+			EditCondition = "BlendMode == EBlendMode::BLEND_Masked"))
+	float OpacityMaskClipValue;
 
 	/** Mid point for displacement in the range 0-1. */
-	UPROPERTY(EditInstanceOnly, BlueprintReadOnly, Category = "Material Designer", 
-		meta = (NotKeyframeable, UIMin = 0, UIMax = 1, ClampMin = 0, ClampMax = 1))
+	UPROPERTY(EditInstanceOnly, BlueprintReadOnly, Category = "Material Designer|General", 
+		meta = (NotKeyframeable, UIMin = 0, UIMax = 1, ClampMin = 0, ClampMax = 1,
+			EditCondition = "UsageFlags & EDMMaterialFlagsUsage::DMMFU_Nanite"))
 	float DisplacementCenter;
 
 	/** Multipler for displacement values. */
-	UPROPERTY(EditInstanceOnly, BlueprintReadOnly, Category = "Material Designer", meta = (NotKeyframeable))
+	UPROPERTY(EditInstanceOnly, BlueprintReadOnly, Category = "Material Designer|General", meta = (NotKeyframeable,
+		EditCondition = "UsageFlags & EDMMaterialFlagsUsage::DMMFU_Nanite"))
 	float DisplacementMagnitude;
 
-	/**
-	 * Indicates that the material should be rendered using responsive anti-aliasing. Improves sharpness of small moving particles such as sparks.
-	 * Only use for small moving features because it will cause aliasing of the background.
-	 */
-	UPROPERTY(EditInstanceOnly, BlueprintReadOnly, Category = "Material Designer", 
-		meta = (NotKeyframeable, DisplayName = "Responsive AA"))
-	bool bResponsiveAAEnabled;
-
 	UPROPERTY(VisibleInstanceOnly, BlueprintReadOnly, TextExportTransient, Category = "Material Designer")
 	TMap<EDMMaterialPropertyType, TObjectPtr<UDMMaterialSlot>> PropertySlotMap;
 
@@ -406,6 +607,8 @@ protected:
 	UFUNCTION(BlueprintCallable, Category = "Material Designer")
 	void BuildMaterial(bool bInDirtyAssets);
 
+	void SetMaterialFlags(UMaterial* InMaterial);
+
 	FString GetMaterialAssetPath() const;
 	FString GetMaterialAssetName() const;
 	FString GetMaterialPackageName(const FString& InMaterialBaseName) const;
@@ -427,9 +630,35 @@ protected:
 	void OnBlendModeChanged();
 	void OnShadingModelChanged();
 	void OnMaterialFlagChanged();
-	void OnDisplacementSettingsChanged();
+	void OnMaterialSettingsChanged();
 
 	//~ Begin IDynamicMaterialModelEditorOnlyDataInterface
 	virtual void ReinitComponents() override;
 	//~ End IDynamicMaterialModelEditorOnlyDataInterface
+
+	/*
+	 * Deprecated
+	 */
+
+	void ConvertDeprecatedVarsToFlags();
+
+	UE_DEPRECATED(5.8, "Moved to GeneralFlags")
+	UPROPERTY(meta = (DeprecatedProperty, DeprecationMessage = "Moved to GeneralFlags"))
+	bool bTwoSided_DEPRECATED;
+
+	UE_DEPRECATED(5.8, "Moved to MotionFlags")
+	UPROPERTY(meta = (DeprecatedProperty, DeprecationMessage = "Moved to MotionFlags"))
+	bool bHasPixelAnimation_DEPRECATED;
+
+	UE_DEPRECATED(5.8, "Moved to MotionFlags")
+	UPROPERTY(meta = (DeprecatedProperty, DeprecationMessage = "Moved to MotionFlags"))
+	bool bOutputTranslucentVelocityEnabled_DEPRECATED;
+
+	UE_DEPRECATED(5.8, "Moved to UsageFlags")
+	UPROPERTY(meta = (DeprecatedProperty, DeprecationMessage = "Moved to UsageFlags"))
+	bool bNaniteTessellationEnabled_DEPRECATED;
+
+	UE_DEPRECATED(5.8, "Moved to MotionFlags")
+	UPROPERTY(meta = (DeprecatedProperty, DeprecationMessage = "Moved to MotionFlags"))
+	bool bResponsiveAAEnabled_DEPRECATED;
 };
