Material Designer: Fix log spam about edit condition

#Synced-CL 49293031

diff --git a/Engine/Plugins/VirtualProduction/DynamicMaterial/Source/DynamicMaterialEditor/Public/Model/DynamicMaterialModelEditorOnlyData.h b/Engine/Plugins/VirtualProduction/DynamicMaterial/Source/DynamicMaterialEditor/Public/Model/DynamicMaterialModelEditorOnlyData.h
--- a/Engine/Plugins/VirtualProduction/DynamicMaterial/Source/DynamicMaterialEditor/Public/Model/DynamicMaterialModelEditorOnlyData.h
+++ b/Engine/Plugins/VirtualProduction/DynamicMaterial/Source/DynamicMaterialEditor/Public/Model/DynamicMaterialModelEditorOnlyData.h
@@ -476,29 +476,29 @@ protected:
 	EDMMaterialShadingModel ShadingModel;
 
 	UPROPERTY(EditInstanceOnly, BlueprintReadOnly, Category = "Material Designer|Usage",
-		meta = (Bitmask, BitmaskEnum = "/Script/DynamicMaterialEditor.EDMMaterialFlagsUsage"))
+		meta = (NotKeyframeable, Bitmask, BitmaskEnum = "/Script/DynamicMaterialEditor.EDMMaterialFlagsUsage"))
 	int32 UsageFlags;
 
 	UPROPERTY(EditInstanceOnly, BlueprintReadOnly, Category = "Material Designer|General",
-		meta = (Bitmask, BitmaskEnum = "/Script/DynamicMaterialEditor.EDMMaterialFlagsGeneral"))
+		meta = (NotKeyframeable, Bitmask, BitmaskEnum = "/Script/DynamicMaterialEditor.EDMMaterialFlagsGeneral"))
 	uint8 GeneralFlags;
 
 	UPROPERTY(EditInstanceOnly, BlueprintReadOnly, Category = "Material Designer|Lighting",
-		meta = (Bitmask, BitmaskEnum = "/Script/DynamicMaterialEditor.EDMMaterialFlagsLighting",
+		meta = (NotKeyframeable, Bitmask, BitmaskEnum = "/Script/DynamicMaterialEditor.EDMMaterialFlagsLighting",
 			EditCondition = "ShadingModel == EDMMaterialShadingModel::DefaultLit"))
 	uint8 LightingFlags;
 
 	UPROPERTY(EditInstanceOnly, BlueprintReadOnly, Category = "Material Designer|Translucency",
-		meta = (Bitmask, BitmaskEnum = "/Script/DynamicMaterialEditor.EDMMaterialFlagsTranslucency",
+		meta = (NotKeyframeable, Bitmask, BitmaskEnum = "/Script/DynamicMaterialEditor.EDMMaterialFlagsTranslucency",
 			EditCondition = "BlendMode == EBlendMode::BLEND_Masked || BlendMode == EBlendMode::BLEND_Translucent"))
 	uint8 TranslucencyFlags;
 
 	UPROPERTY(EditInstanceOnly, BlueprintReadOnly, Category = "Material Designer|Motion",
-		meta = (Bitmask, BitmaskEnum = "/Script/DynamicMaterialEditor.EDMMaterialFlagsMotion"))
+		meta = (NotKeyframeable, Bitmask, BitmaskEnum = "/Script/DynamicMaterialEditor.EDMMaterialFlagsMotion"))
 	uint8 MotionFlags;
 
 	UPROPERTY(EditInstanceOnly, BlueprintReadOnly, Category = "Material Designer|Forward Renderer",
-		meta = (Bitmask, BitmaskEnum = "/Script/DynamicMaterialEditor.EDMMaterialFlagsForwardRenderer"))
+		meta = (NotKeyframeable, Bitmask, BitmaskEnum = "/Script/DynamicMaterialEditor.EDMMaterialFlagsForwardRenderer"))
 	uint8 ForwardRendererFlags;
 
 	UPROPERTY(EditInstanceOnly, BlueprintReadOnly, Category = "Material Designer|Translucency",
@@ -508,13 +508,12 @@ protected:
 
 	/** Mid point for displacement in the range 0-1. */
 	UPROPERTY(EditInstanceOnly, BlueprintReadOnly, Category = "Material Designer|General", 
-		meta = (NotKeyframeable, UIMin = 0, UIMax = 1, ClampMin = 0, ClampMax = 1,
-			EditCondition = "UsageFlags & EDMMaterialFlagsUsage::DMMFU_Nanite"))
+		meta = (NotKeyframeable, UIMin = 0, UIMax = 1, ClampMin = 0, ClampMax = 1))
 	float DisplacementCenter;
 
 	/** Multipler for displacement values. */
-	UPROPERTY(EditInstanceOnly, BlueprintReadOnly, Category = "Material Designer|General", meta = (NotKeyframeable,
-		EditCondition = "UsageFlags & EDMMaterialFlagsUsage::DMMFU_Nanite"))
+	UPROPERTY(EditInstanceOnly, BlueprintReadOnly, Category = "Material Designer|General", 
+		meta = (NotKeyframeable))
 	float DisplacementMagnitude;
 
 	UPROPERTY(VisibleInstanceOnly, BlueprintReadOnly, TextExportTransient, Category = "Material Designer")
