Material Designer: Fixed edit conditions on nanite properties

#Synced-CL 49293031

diff --git a/Engine/Plugins/VirtualProduction/DynamicMaterial/Source/DynamicMaterialEditor/Public/Model/DynamicMaterialModelEditorOnlyData.h b/Engine/Plugins/VirtualProduction/DynamicMaterial/Source/DynamicMaterialEditor/Public/Model/DynamicMaterialModelEditorOnlyData.h
--- a/Engine/Plugins/VirtualProduction/DynamicMaterial/Source/DynamicMaterialEditor/Public/Model/DynamicMaterialModelEditorOnlyData.h
+++ b/Engine/Plugins/VirtualProduction/DynamicMaterial/Source/DynamicMaterialEditor/Public/Model/DynamicMaterialModelEditorOnlyData.h
@@ -508,12 +508,13 @@ protected:
 
 	/** Mid point for displacement in the range 0-1. */
 	UPROPERTY(EditInstanceOnly, BlueprintReadOnly, Category = "Material Designer|General", 
-		meta = (NotKeyframeable, UIMin = 0, UIMax = 1, ClampMin = 0, ClampMax = 1))
+		meta = (NotKeyframeable, UIMin = 0, UIMax = 1, ClampMin = 0, ClampMax = 1,
+			EditCondition = "IsNaniteEnabled()"))
 	float DisplacementCenter;
 
 	/** Multipler for displacement values. */
 	UPROPERTY(EditInstanceOnly, BlueprintReadOnly, Category = "Material Designer|General", 
-		meta = (NotKeyframeable))
+		meta = (NotKeyframeable, EditCondition = "IsNaniteEnabled()"))
 	float DisplacementMagnitude;
 
 	UPROPERTY(VisibleInstanceOnly, BlueprintReadOnly, TextExportTransient, Category = "Material Designer")
@@ -660,4 +661,11 @@ protected:
 	UE_DEPRECATED(5.8, "Moved to MotionFlags")
 	UPROPERTY(meta = (DeprecatedProperty, DeprecationMessage = "Moved to MotionFlags"))
 	bool bResponsiveAAEnabled_DEPRECATED;
+
+private:
+	UFUNCTION()
+	bool IsNaniteEnabled() const
+	{
+		return UsageFlags & EDMMaterialFlagsUsage::DMMFU_Nanite;
+	}
 };
