Material Designer: Materials will now set at least 1 usage flag.

#Synced-CL 49215135

diff --git a/Engine/Plugins/VirtualProduction/DynamicMaterial/Source/DynamicMaterialEditor/Private/DynamicMaterialEditorSettings.cpp b/Engine/Plugins/VirtualProduction/DynamicMaterial/Source/DynamicMaterialEditor/Private/DynamicMaterialEditorSettings.cpp
--- a/Engine/Plugins/VirtualProduction/DynamicMaterial/Source/DynamicMaterialEditor/Private/DynamicMaterialEditorSettings.cpp
+++ b/Engine/Plugins/VirtualProduction/DynamicMaterial/Source/DynamicMaterialEditor/Private/DynamicMaterialEditorSettings.cpp
@@ -155,7 +155,7 @@ UDynamicMaterialEditorSettings::UDynamicMaterialEditorSettings()
 	Opaque.bEmissive = true;
 	Opaque.DefaultBlendMode = BLEND_Opaque;
 	Opaque.DefaultShadingModel = EDMMaterialShadingModel::DefaultLit;
-	Opaque.UsageFlag |= EDMMaterialFlagsUsage::DMMFU_NiagaraMeshParticles;
+	Opaque.UsageFlag |= EDMMaterialFlagsUsage::DMMFU_StaticMesh | EDMMaterialFlagsUsage::DMMFU_NiagaraMeshParticles;
 	Opaque.GeneralFlags |= static_cast<uint8>(EDMMaterialFlagsGeneral::TwoSided);
 
 	FDMMaterialChannelListPreset Emissive;
@@ -163,7 +163,7 @@ UDynamicMaterialEditorSettings::UDynamicMaterialEditorSettings()
 	Emissive.bEmissive = true;
 	Emissive.DefaultBlendMode = BLEND_Opaque;
 	Emissive.DefaultShadingModel = EDMMaterialShadingModel::Unlit;
-	Emissive.UsageFlag |= EDMMaterialFlagsUsage::DMMFU_NiagaraMeshParticles;
+	Emissive.UsageFlag |= EDMMaterialFlagsUsage::DMMFU_StaticMesh | EDMMaterialFlagsUsage::DMMFU_NiagaraMeshParticles;
 	Emissive.GeneralFlags |= static_cast<uint8>(EDMMaterialFlagsGeneral::TwoSided);
 
 	FDMMaterialChannelListPreset Translucent;
@@ -172,7 +172,7 @@ UDynamicMaterialEditorSettings::UDynamicMaterialEditorSettings()
 	Translucent.bOpacity = true;
 	Translucent.DefaultBlendMode = BLEND_Translucent;
 	Translucent.DefaultShadingModel = EDMMaterialShadingModel::Unlit;
-	Translucent.UsageFlag |= EDMMaterialFlagsUsage::DMMFU_NiagaraMeshParticles;
+	Translucent.UsageFlag |= EDMMaterialFlagsUsage::DMMFU_StaticMesh | EDMMaterialFlagsUsage::DMMFU_NiagaraMeshParticles;
 	Translucent.GeneralFlags |= static_cast<uint8>(EDMMaterialFlagsGeneral::TwoSided);
 
 	FDMMaterialChannelListPreset PBR;
@@ -188,7 +188,7 @@ UDynamicMaterialEditorSettings::UDynamicMaterialEditorSettings()
 	PBR.bDisplacement = true;
 	PBR.DefaultBlendMode = BLEND_Opaque;
 	PBR.DefaultShadingModel = EDMMaterialShadingModel::DefaultLit;
-	PBR.UsageFlag |= EDMMaterialFlagsUsage::DMMFU_NiagaraMeshParticles;
+	PBR.UsageFlag |= EDMMaterialFlagsUsage::DMMFU_StaticMesh | EDMMaterialFlagsUsage::DMMFU_NiagaraMeshParticles;
 	PBR.GeneralFlags |= static_cast<uint8>(EDMMaterialFlagsGeneral::TwoSided);
 
 	FDMMaterialChannelListPreset All;
@@ -211,7 +211,7 @@ UDynamicMaterialEditorSettings::UDynamicMaterialEditorSettings()
 	All.bSurfaceThickness = true;
 	All.DefaultBlendMode = BLEND_Opaque;
 	All.DefaultShadingModel = EDMMaterialShadingModel::DefaultLit;
-	All.UsageFlag |= EDMMaterialFlagsUsage::DMMFU_NiagaraMeshParticles;
+	All.UsageFlag |= EDMMaterialFlagsUsage::DMMFU_StaticMesh | EDMMaterialFlagsUsage::DMMFU_NiagaraMeshParticles;
 	All.GeneralFlags |= static_cast<uint8>(EDMMaterialFlagsGeneral::TwoSided);
 
 	MaterialChannelPresets.Add(PBR);
diff --git a/Engine/Plugins/VirtualProduction/DynamicMaterial/Source/DynamicMaterialEditor/Private/Model/DynamicMaterialModelEditorOnlyData.cpp b/Engine/Plugins/VirtualProduction/DynamicMaterial/Source/DynamicMaterialEditor/Private/Model/DynamicMaterialModelEditorOnlyData.cpp
--- a/Engine/Plugins/VirtualProduction/DynamicMaterial/Source/DynamicMaterialEditor/Private/Model/DynamicMaterialModelEditorOnlyData.cpp
+++ b/Engine/Plugins/VirtualProduction/DynamicMaterial/Source/DynamicMaterialEditor/Private/Model/DynamicMaterialModelEditorOnlyData.cpp
@@ -602,6 +602,12 @@ void UDynamicMaterialModelEditorOnlyData::SetMaterialFlags(UMaterial* InMaterial
 		InMaterial->SetUsageByFlag(UsageFlag, HasUsageFlag(DMUsageFlag));
 	}
 
+	// Must have at least 1 usage flag
+	if (UsageFlags == 0)
+	{
+		InMaterial->SetUsageByFlag(EMaterialUsage::MATUSAGE_StaticMesh, true);
+	}
+
 	InMaterial->TwoSided = HasGeneralFlag(EDMMaterialFlagsGeneral::TwoSided);
 	InMaterial->bTangentSpaceNormal = HasGeneralFlag(EDMMaterialFlagsGeneral::TangentSpaceNormal);
 	InMaterial->bUsesDistortion = HasGeneralFlag(EDMMaterialFlagsGeneral::UsesDistortion);
