Material Designer: Axis-locking has been removed from Float2 values, except for texture tiling.

#Synced-CL 50113862

diff --git a/Engine/Plugins/VirtualProduction/DynamicMaterial/Source/DynamicMaterial/Public/Components/MaterialValues/DMMaterialValueFloat2.h b/Engine/Plugins/VirtualProduction/DynamicMaterial/Source/DynamicMaterial/Public/Components/MaterialValues/DMMaterialValueFloat2.h
--- a/Engine/Plugins/VirtualProduction/DynamicMaterial/Source/DynamicMaterial/Public/Components/MaterialValues/DMMaterialValueFloat2.h
+++ b/Engine/Plugins/VirtualProduction/DynamicMaterial/Source/DynamicMaterial/Public/Components/MaterialValues/DMMaterialValueFloat2.h
@@ -59,7 +59,7 @@ public:
  
 protected: 
 	UPROPERTY(EditInstanceOnly, BlueprintReadWrite, Getter, Setter, BlueprintSetter = SetValue, Category = "Material Designer",
-		meta = (AllowPrivateAccess = "true", DisplayName = "Vector2D", AllowPreserveRatio, Delta = 0.001))
+		meta = (AllowPrivateAccess = "true", DisplayName = "Vector2D", Delta = 0.001))
 	FVector2D Value;
  
 #if WITH_EDITORONLY_DATA
