Material Designer: Axis lock also removed from float3 xyz values. And both float2 and float3 xyz dynamic values.

#Synced-CL 50214844

diff --git a/Engine/Plugins/VirtualProduction/DynamicMaterial/Source/DynamicMaterial/Public/Components/MaterialValues/DMMaterialValueFloat3XYZ.h b/Engine/Plugins/VirtualProduction/DynamicMaterial/Source/DynamicMaterial/Public/Components/MaterialValues/DMMaterialValueFloat3XYZ.h
--- a/Engine/Plugins/VirtualProduction/DynamicMaterial/Source/DynamicMaterial/Public/Components/MaterialValues/DMMaterialValueFloat3XYZ.h
+++ b/Engine/Plugins/VirtualProduction/DynamicMaterial/Source/DynamicMaterial/Public/Components/MaterialValues/DMMaterialValueFloat3XYZ.h
@@ -59,7 +59,7 @@ public:
  
 protected: 
 	UPROPERTY(EditInstanceOnly, BlueprintReadWrite, Getter, Setter, BlueprintSetter = SetValue, Category = "Material Designer",
-		meta = (AllowPrivateAccess = "true", DisplayName = "Vector3D", AllowPreserveRatio))
+		meta = (AllowPrivateAccess = "true", DisplayName = "Vector3D"))
 	FVector Value;
  
 #if WITH_EDITORONLY_DATA
diff --git a/Engine/Plugins/VirtualProduction/DynamicMaterial/Source/DynamicMaterial/Public/Components/MaterialValuesDynamic/DMMaterialValueFloat2Dynamic.h b/Engine/Plugins/VirtualProduction/DynamicMaterial/Source/DynamicMaterial/Public/Components/MaterialValuesDynamic/DMMaterialValueFloat2Dynamic.h
--- a/Engine/Plugins/VirtualProduction/DynamicMaterial/Source/DynamicMaterial/Public/Components/MaterialValuesDynamic/DMMaterialValueFloat2Dynamic.h
+++ b/Engine/Plugins/VirtualProduction/DynamicMaterial/Source/DynamicMaterial/Public/Components/MaterialValuesDynamic/DMMaterialValueFloat2Dynamic.h
@@ -45,7 +45,7 @@ public:
 
 protected:
 	UPROPERTY(EditInstanceOnly, BlueprintReadWrite, Getter, Setter, BlueprintSetter = SetValue, Category = "Material Designer",
-		meta = (AllowPrivateAccess = "true", DisplayName = "Vector2D", AllowPreserveRatio, Delta = 0.001))
+		meta = (AllowPrivateAccess = "true", DisplayName = "Vector2D", Delta = 0.001))
 	FVector2D Value;
 
 	//~ Begin IDMParameterContainer
diff --git a/Engine/Plugins/VirtualProduction/DynamicMaterial/Source/DynamicMaterial/Public/Components/MaterialValuesDynamic/DMMaterialValueFloat3XYZDynamic.h b/Engine/Plugins/VirtualProduction/DynamicMaterial/Source/DynamicMaterial/Public/Components/MaterialValuesDynamic/DMMaterialValueFloat3XYZDynamic.h
--- a/Engine/Plugins/VirtualProduction/DynamicMaterial/Source/DynamicMaterial/Public/Components/MaterialValuesDynamic/DMMaterialValueFloat3XYZDynamic.h
+++ b/Engine/Plugins/VirtualProduction/DynamicMaterial/Source/DynamicMaterial/Public/Components/MaterialValuesDynamic/DMMaterialValueFloat3XYZDynamic.h
@@ -45,7 +45,7 @@ public:
 
 protected:
 	UPROPERTY(EditInstanceOnly, BlueprintReadWrite, Getter, Setter, BlueprintSetter = SetValue, Category = "Material Designer",
-		meta = (AllowPrivateAccess = "true", DisplayName = "Vector3D", AllowPreserveRatio, Delta = 0.001))
+		meta = (AllowPrivateAccess = "true", DisplayName = "Vector3D", Delta = 0.001))
 	FVector Value;
 
 	//~ Begin IDMParameterContainer
