[AvalancheTests] Automate 'Material Designer - Add/Delete an FX layer to an existing Material Designer Layer'

#rnx

#Synced-CL 45465245

diff --git a/Engine/Plugins/VirtualProduction/DynamicMaterial/Source/DynamicMaterialEditor/Private/DynamicMaterialEditorSettings.h b/Engine/Plugins/VirtualProduction/DynamicMaterial/Source/DynamicMaterialEditor/Private/DynamicMaterialEditorSettings.h
--- a/Engine/Plugins/VirtualProduction/DynamicMaterial/Source/DynamicMaterialEditor/Private/DynamicMaterialEditorSettings.h
+++ b/Engine/Plugins/VirtualProduction/DynamicMaterial/Source/DynamicMaterialEditor/Private/DynamicMaterialEditorSettings.h
@@ -226,7 +226,7 @@ public:
 	UDynamicMaterialEditorSettings();
 	virtual ~UDynamicMaterialEditorSettings() override = default;
 
-	static UDynamicMaterialEditorSettings* Get();
+	DLLEXPORT static UDynamicMaterialEditorSettings* Get();
 
 	/** Changes the currently active material in the designer following actor/object selection. */
 	UPROPERTY(Config, EditAnywhere, BlueprintReadOnly, Category = "Editor")
@@ -260,13 +260,13 @@ protected:
 	UPROPERTY(Config, EditAnywhere, Category = "Editor", meta = (EditCondition = "LiveEditMode == EDMLiveEditMode::Disabled"))
 	bool bAutomaticallyApplyToSourceOnPreviewCompile;
 
+public:
 	/**
 	 * When enabled, reduces compile options to a single mode called Live Edit.
 	 */
 	UPROPERTY(Config, EditAnywhere, BlueprintReadOnly, Category = "Editor")
 	EDMLiveEditMode LiveEditMode;
 
-public:
 	UFUNCTION(BlueprintPure, Category = "Material Designer")
 	bool ShouldAutomaticallyCopyParametersToSourceMaterial() const
 	{
