[AvalancheTests] Automate 'Material Designer - Save/Load FX Layer preset from Material Designer'

#rnx

#Synced-CL 45465245

diff --git a/Engine/Plugins/VirtualProduction/DynamicMaterial/Source/DynamicMaterialEditor/Private/UI/Widgets/Editor/SDMMaterialSlotEditor.h b/Engine/Plugins/VirtualProduction/DynamicMaterial/Source/DynamicMaterialEditor/Private/UI/Widgets/Editor/SDMMaterialSlotEditor.h
--- a/Engine/Plugins/VirtualProduction/DynamicMaterial/Source/DynamicMaterialEditor/Private/UI/Widgets/Editor/SDMMaterialSlotEditor.h
+++ b/Engine/Plugins/VirtualProduction/DynamicMaterial/Source/DynamicMaterialEditor/Private/UI/Widgets/Editor/SDMMaterialSlotEditor.h
@@ -5,7 +5,7 @@
 #include "Misc/NotifyHook.h"
 #include "Widgets/SCompoundWidget.h"
 
-#include "UI/Utils/DMWidgetSlot.h"
+#include "DynamicMaterialEditor/Private/UI/Utils/DMWidgetSlot.h"
 
 class FDetailColumnSizeData;
 class ICustomDetailsViewItem;
@@ -52,7 +52,7 @@ public:
 
 	TSharedPtr<SDMMaterialEditor> GetEditorWidget() const;
 
-	UDMMaterialSlot* GetSlot() const;
+	DLLEXPORT UDMMaterialSlot* GetSlot() const;
 
 	/** Actions */
 	void ClearSelection();
