Static Analysis fixes

#rnx

#Synced-CL 51101942

diff --git a/Engine/Plugins/Animation/ControlRig/Source/ControlRigDeveloper/Private/ControlRigEditorAsset.cpp b/Engine/Plugins/Animation/ControlRig/Source/ControlRigDeveloper/Private/ControlRigEditorAsset.cpp
--- a/Engine/Plugins/Animation/ControlRig/Source/ControlRigDeveloper/Private/ControlRigEditorAsset.cpp
+++ b/Engine/Plugins/Animation/ControlRig/Source/ControlRigDeveloper/Private/ControlRigEditorAsset.cpp
@@ -1128,9 +1128,8 @@ EControlRigType IControlRigEditorAssetInterface::GetRigType(const FAssetData& In
 		return Result;
 	}
 
-	EControlRigType RigType;
-	ControlRigTypeProperty->ImportText_Direct(*ControlRigTypeString, &RigType, nullptr, EPropertyPortFlags::PPF_None);
-	return RigType;
+	ControlRigTypeProperty->ImportText_Direct(*ControlRigTypeString, &Result, nullptr, EPropertyPortFlags::PPF_None);
+	return Result;
 }
 
 TArray<FSoftObjectPath> IControlRigEditorAssetInterface::GetReferencesToRigModule(const FAssetData& InModuleAsset)
diff --git a/Engine/Plugins/VirtualProduction/DynamicMaterial/Source/DynamicMaterialEditor/Private/Components/MaterialStageExpressions/DMMSEMathBase.cpp b/Engine/Plugins/VirtualProduction/DynamicMaterial/Source/DynamicMaterialEditor/Private/Components/MaterialStageExpressions/DMMSEMathBase.cpp
--- a/Engine/Plugins/VirtualProduction/DynamicMaterial/Source/DynamicMaterialEditor/Private/Components/MaterialStageExpressions/DMMSEMathBase.cpp
+++ b/Engine/Plugins/VirtualProduction/DynamicMaterial/Source/DynamicMaterialEditor/Private/Components/MaterialStageExpressions/DMMSEMathBase.cpp
@@ -138,40 +138,37 @@ bool UDMMaterialStageExpressionMathBase::CanInputAcceptType(int32 InInputIndex,
 
 						if (StageInputOutputs.IsValidIndex(Channel.OutputIndex))
 						{
-							if (StageInputOutputs.IsValidIndex(Channel.OutputIndex))
+							if (Channel.OutputChannel == FDMMaterialStageConnectorChannel::WHOLE_CHANNEL)
 							{
-								if (Channel.OutputChannel == FDMMaterialStageConnectorChannel::WHOLE_CHANNEL)
-								{
-									InputType = StageInputOutputs[Channel.OutputIndex].Type;
-									goto FoundInputType;
-								}
-								else
+								InputType = StageInputOutputs[Channel.OutputIndex].Type;
+								goto FoundInputType;
+							}
+							else
+							{
+								switch (InputType)
 								{
-									switch (InputType)
-									{
-										case EDMValueType::VT_None:
-											InputType = EDMValueType::VT_Float1;
-											break;
-
-										case EDMValueType::VT_Float1:
-											InputType = EDMValueType::VT_Float2;
-											break;
-
-										case EDMValueType::VT_Float2:
-											InputType = EDMValueType::VT_Float3_XYZ;
-											break;
-
-										case EDMValueType::VT_Float3_RGB:
-										case EDMValueType::VT_Float3_RPY:
-										case EDMValueType::VT_Float3_XYZ:
-											InputType = EDMValueType::VT_Float4_RGBA;
-											break;
-
-										// Not valid
-										default:
-										case EDMValueType::VT_Float4_RGBA:
-											return false;
-									}
+									case EDMValueType::VT_None:
+										InputType = EDMValueType::VT_Float1;
+										break;
+
+									case EDMValueType::VT_Float1:
+										InputType = EDMValueType::VT_Float2;
+										break;
+
+									case EDMValueType::VT_Float2:
+										InputType = EDMValueType::VT_Float3_XYZ;
+										break;
+
+									case EDMValueType::VT_Float3_RGB:
+									case EDMValueType::VT_Float3_RPY:
+									case EDMValueType::VT_Float3_XYZ:
+										InputType = EDMValueType::VT_Float4_RGBA;
+										break;
+
+									// Not valid
+									default:
+									case EDMValueType::VT_Float4_RGBA:
+										return false;
 								}
 							}
 						}
diff --git a/Engine/Source/Runtime/Engine/Private/Elements/Framework/TypedElementCommonActions.cpp b/Engine/Source/Runtime/Engine/Private/Elements/Framework/TypedElementCommonActions.cpp
--- a/Engine/Source/Runtime/Engine/Private/Elements/Framework/TypedElementCommonActions.cpp
+++ b/Engine/Source/Runtime/Engine/Private/Elements/Framework/TypedElementCommonActions.cpp
@@ -177,7 +177,7 @@ FTypedElementCommonActions_CopyPasteCapability UTypedElementCommonActions::CanCo
 		ITypedElementWorldInterface* WorldInterface = Registry->GetElementInterface<ITypedElementWorldInterface>(ElementsByTypePair.Key);
 		if (CommonActionsCustomization && WorldInterface)
 		{
-			bCapable = bCapable & CommonActionsCustomization->IsCopyCapable();
+			bCapable = bCapable && CommonActionsCustomization->IsCopyCapable();
 			bAllowed = bAllowed && CommonActionsCustomization->CanCopyElements(WorldInterface, ElementsByTypePair.Value);
 		}
 		else
