Build fix

#rnx

#Synced-CL 51344488

diff --git a/Engine/Plugins/VirtualProduction/DynamicMaterial/Source/DynamicMaterial/Public/Utils/DMJsonUtils.h b/Engine/Plugins/VirtualProduction/DynamicMaterial/Source/DynamicMaterial/Public/Utils/DMJsonUtils.h
--- a/Engine/Plugins/VirtualProduction/DynamicMaterial/Source/DynamicMaterial/Public/Utils/DMJsonUtils.h
+++ b/Engine/Plugins/VirtualProduction/DynamicMaterial/Source/DynamicMaterial/Public/Utils/DMJsonUtils.h
@@ -3,6 +3,7 @@
 #pragma once
 
 #include "Concepts/BaseStructureProvider.h"
+#include "Containers/StringFwd.h"
 #include "Math/MathFwd.h"
 #include "Templates/SharedPointerFwd.h"
 #include "UObject/Interface.h"
