Build fix

#rnx
#Synced-CL 51354297

diff --git a/Engine/Plugins/VirtualProduction/DynamicMaterial/Source/DynamicMaterial/Private/Components/RenderTargetRenderers/DMRenderTargetUMGWidgetRenderer.cpp b/Engine/Plugins/VirtualProduction/DynamicMaterial/Source/DynamicMaterial/Private/Components/RenderTargetRenderers/DMRenderTargetUMGWidgetRenderer.cpp
--- a/Engine/Plugins/VirtualProduction/DynamicMaterial/Source/DynamicMaterial/Private/Components/RenderTargetRenderers/DMRenderTargetUMGWidgetRenderer.cpp
+++ b/Engine/Plugins/VirtualProduction/DynamicMaterial/Source/DynamicMaterial/Private/Components/RenderTargetRenderers/DMRenderTargetUMGWidgetRenderer.cpp
@@ -3,6 +3,7 @@
 #include "Components/RenderTargetRenderers/DMRenderTargetUMGWidgetRenderer.h"
 #include "Components/MaterialValues/DMMaterialValueRenderTarget.h"
 #include "Components/Widget.h"
+#include "Containers/SharedString.h"
 #include "Dom/JsonValue.h"
 #include "Engine/Engine.h"
 
